USD: Reformat the UnrealUSDWrapper module.

diff --git a/Engine/Plugins/Importers/USDImporter/Source/UnrealUSDWrapper/Private/UsdWrappers/UsdStage.cpp b/Engine/Plugins/Importers/USDImporter/Source/UnrealUSDWrapper/Private/UsdWrappers/UsdStage.cpp
--- a/Engine/Plugins/Importers/USDImporter/Source/UnrealUSDWrapper/Private/UsdWrappers/UsdStage.cpp
+++ b/Engine/Plugins/Importers/USDImporter/Source/UnrealUSDWrapper/Private/UsdWrappers/UsdStage.cpp
@@ -10,14 +10,12 @@
 #include "USDMemory.h"
 
 #if USE_USD_SDK
-
 #include "USDIncludesStart.h"
-	#include "pxr/usd/usd/editContext.h"
-	#include "pxr/usd/usd/prim.h"
-	#include "pxr/usd/usd/stage.h"
+#include "pxr/usd/usd/editContext.h"
+#include "pxr/usd/usd/prim.h"
+#include "pxr/usd/usd/stage.h"
 #include "USDIncludesEnd.h"
-
-#endif // #if USE_USD_SDK
+#endif	  // #if USE_USD_SDK
 
 namespace UE
 {
@@ -29,13 +27,13 @@ namespace UE
 		public:
 			FUsdStageImpl() = default;
 
-			explicit FUsdStageImpl( const PtrType& InUsdPtr )
-				: UsdPtr( InUsdPtr )
+			explicit FUsdStageImpl(const PtrType& InUsdPtr)
+				: UsdPtr(InUsdPtr)
 			{
 			}
 
-			explicit FUsdStageImpl( PtrType&& InUsdPtr )
-				: UsdPtr( MoveTemp( InUsdPtr ) )
+			explicit FUsdStageImpl(PtrType&& InUsdPtr)
+				: UsdPtr(MoveTemp(InUsdPtr))
 			{
 			}
 
@@ -49,7 +47,7 @@ namespace UE
 				return UsdPtr.Get();
 			}
 
-			TUsdStore< PtrType > UsdPtr;
+			TUsdStore<PtrType> UsdPtr;
 		};
 	}
 
@@ -57,81 +55,81 @@ namespace UE
 	FUsdStageBase<PtrType>::FUsdStageBase()
 	{
 		FScopedUnrealAllocs UnrealAllocs;
-		Impl = MakeUnique< Internal::FUsdStageImpl<PtrType> >();
+		Impl = MakeUnique<Internal::FUsdStageImpl<PtrType>>();
 	}
 
 	template<typename PtrType>
-	FUsdStageBase<PtrType>::FUsdStageBase( const FUsdStage& Other )
-		: Impl( MakeUnique< Internal::FUsdStageImpl<PtrType> >( Other.Impl->GetInner() ) )
+	FUsdStageBase<PtrType>::FUsdStageBase(const FUsdStage& Other)
+		: Impl(MakeUnique<Internal::FUsdStageImpl<PtrType>>(Other.Impl->GetInner()))
 	{
 	}
 
 	template<typename PtrType>
-	FUsdStageBase<PtrType>::FUsdStageBase( FUsdStage&& Other )
-		: Impl( MakeUnique< Internal::FUsdStageImpl<PtrType> >( MoveTemp( Other.Impl->GetInner() ) ) )
+	FUsdStageBase<PtrType>::FUsdStageBase(FUsdStage&& Other)
+		: Impl(MakeUnique<Internal::FUsdStageImpl<PtrType>>(MoveTemp(Other.Impl->GetInner())))
 	{
 	}
 
 	template<typename PtrType>
-	FUsdStageBase<PtrType>::FUsdStageBase( const FUsdStageWeak& Other )
-		: Impl( MakeUnique< Internal::FUsdStageImpl<PtrType> >( Other.Impl->GetInner() ) )
+	FUsdStageBase<PtrType>::FUsdStageBase(const FUsdStageWeak& Other)
+		: Impl(MakeUnique<Internal::FUsdStageImpl<PtrType>>(Other.Impl->GetInner()))
 	{
 	}
 
 	template<typename PtrType>
-	FUsdStageBase<PtrType>::FUsdStageBase( FUsdStageWeak&& Other )
-		: Impl( MakeUnique< Internal::FUsdStageImpl<PtrType> >( MoveTemp( Other.Impl->GetInner() ) ) )
+	FUsdStageBase<PtrType>::FUsdStageBase(FUsdStageWeak&& Other)
+		: Impl(MakeUnique<Internal::FUsdStageImpl<PtrType>>(MoveTemp(Other.Impl->GetInner())))
 	{
 	}
 
 	template<typename PtrType>
-	FUsdStageBase<PtrType>& FUsdStageBase<PtrType>::operator=( const FUsdStage& Other )
+	FUsdStageBase<PtrType>& FUsdStageBase<PtrType>::operator=(const FUsdStage& Other)
 	{
 		FScopedUnrealAllocs UnrealAllocs;
-		Impl = MakeUnique< Internal::FUsdStageImpl<PtrType> >( Other.Impl->GetInner() );
+		Impl = MakeUnique<Internal::FUsdStageImpl<PtrType>>(Other.Impl->GetInner());
 		return *this;
 	}
 
 	template<typename PtrType>
-	FUsdStageBase<PtrType>& FUsdStageBase<PtrType>::operator=( FUsdStage&& Other )
+	FUsdStageBase<PtrType>& FUsdStageBase<PtrType>::operator=(FUsdStage&& Other)
 	{
 		FScopedUnrealAllocs UnrealAllocs;
-		Impl = MakeUnique< Internal::FUsdStageImpl<PtrType> >( MoveTemp( Other.Impl->GetInner() ) );
+		Impl = MakeUnique<Internal::FUsdStageImpl<PtrType>>(MoveTemp(Other.Impl->GetInner()));
 		return *this;
 	}
 
 	template<typename PtrType>
-	FUsdStageBase<PtrType>& FUsdStageBase<PtrType>::operator=( const FUsdStageWeak& Other )
+	FUsdStageBase<PtrType>& FUsdStageBase<PtrType>::operator=(const FUsdStageWeak& Other)
 	{
 		FScopedUnrealAllocs UnrealAllocs;
-		Impl = MakeUnique< Internal::FUsdStageImpl<PtrType> >( Other.Impl->GetInner() );
+		Impl = MakeUnique<Internal::FUsdStageImpl<PtrType>>(Other.Impl->GetInner());
 		return *this;
 	}
 
 	template<typename PtrType>
-	FUsdStageBase<PtrType>& FUsdStageBase<PtrType>::operator=( FUsdStageWeak&& Other )
+	FUsdStageBase<PtrType>& FUsdStageBase<PtrType>::operator=(FUsdStageWeak&& Other)
 	{
 		FScopedUnrealAllocs UnrealAllocs;
-		Impl = MakeUnique< Internal::FUsdStageImpl<PtrType> >( MoveTemp( Other.Impl->GetInner() ) );
+		Impl = MakeUnique<Internal::FUsdStageImpl<PtrType>>(MoveTemp(Other.Impl->GetInner()));
 		return *this;
 	}
 
 	template<typename PtrType>
 	FUsdStageBase<PtrType>::operator bool() const
 	{
-		return ( bool ) Impl->GetInner();
+		return (bool)Impl->GetInner();
 	}
 
 	template<typename PtrType>
-	bool FUsdStageBase<PtrType>::operator==( const FUsdStageBase<PtrType>& Other ) const
+	bool FUsdStageBase<PtrType>::operator==(const FUsdStageBase<PtrType>& Other) const
 	{
 		return Impl->GetInner() == Other.Impl->GetInner();
 	}
 
 	template<typename PtrType>
-	bool FUsdStageBase<PtrType>::operator!=( const FUsdStageBase<PtrType>& Other ) const
+	bool FUsdStageBase<PtrType>::operator!=(const FUsdStageBase<PtrType>& Other) const
 	{
-		return !( *this == Other );
+		return !(*this == Other);
 	}
 
 	template<typename PtrType>
@@ -144,88 +142,87 @@ namespace UE
 	// Auto conversion from/to PtrType
 #if USE_USD_SDK
 	template<typename PtrType>
-	FUsdStageBase<PtrType>::FUsdStageBase( const pxr::UsdStageRefPtr& InUsdPtr )
+	FUsdStageBase<PtrType>::FUsdStageBase(const pxr::UsdStageRefPtr& InUsdPtr)
 	{
 		FScopedUnrealAllocs UnrealAllocs;
 
-		Impl = MakeUnique< Internal::FUsdStageImpl<PtrType> >( InUsdPtr );
+		Impl = MakeUnique<Internal::FUsdStageImpl<PtrType>>(InUsdPtr);
 	}
 
 	template<typename PtrType>
-	FUsdStageBase<PtrType>::FUsdStageBase( pxr::UsdStageRefPtr&& InUsdPtr )
+	FUsdStageBase<PtrType>::FUsdStageBase(pxr::UsdStageRefPtr&& InUsdPtr)
 	{
 		FScopedUnrealAllocs UnrealAllocs;
 
-		Impl = MakeUnique< Internal::FUsdStageImpl<PtrType> >( InUsdPtr );
+		Impl = MakeUnique<Internal::FUsdStageImpl<PtrType>>(InUsdPtr);
 	}
 
 	template<typename PtrType>
-	FUsdStageBase<PtrType>::FUsdStageBase( const pxr::UsdStageWeakPtr& InUsdPtr )
+	FUsdStageBase<PtrType>::FUsdStageBase(const pxr::UsdStageWeakPtr& InUsdPtr)
 	{
 		FScopedUnrealAllocs UnrealAllocs;
 
-		Impl = MakeUnique< Internal::FUsdStageImpl<PtrType> >( InUsdPtr );
+		Impl = MakeUnique<Internal::FUsdStageImpl<PtrType>>(InUsdPtr);
 	}
 
 	template<typename PtrType>
-	FUsdStageBase<PtrType>::FUsdStageBase( pxr::UsdStageWeakPtr&& InUsdPtr )
+	FUsdStageBase<PtrType>::FUsdStageBase(pxr::UsdStageWeakPtr&& InUsdPtr)
 	{
 		FScopedUnrealAllocs UnrealAllocs;
 
-		Impl = MakeUnique< Internal::FUsdStageImpl<PtrType> >( InUsdPtr );
+		Impl = MakeUnique<Internal::FUsdStageImpl<PtrType>>(InUsdPtr);
 	}
 
 	template<typename PtrType>
-	FUsdStageBase<PtrType>::operator PtrType& ( )
+	FUsdStageBase<PtrType>::operator PtrType&()
 	{
 		return Impl->GetInner();
 	}
 
 	template<typename PtrType>
-	FUsdStageBase<PtrType>::operator const PtrType& ( ) const
+	FUsdStageBase<PtrType>::operator const PtrType&() const
 	{
 		return Impl->GetInner();
 	}
 
 	template<typename PtrType>
-	FUsdStageBase<PtrType>::operator pxr::UsdStageRefPtr ( ) const
+	FUsdStageBase<PtrType>::operator pxr::UsdStageRefPtr() const
 	{
 		return Impl->GetInner();
 	}
 
 	template<typename PtrType>
-	FUsdStageBase<PtrType>::operator pxr::UsdStageWeakPtr ( ) const
+	FUsdStageBase<PtrType>::operator pxr::UsdStageWeakPtr() const
 	{
 		return Impl->GetInner();
 	}
-#endif // USE_USD_SDK
+#endif	  // USE_USD_SDK
 
 	template<typename PtrType>
-	void FUsdStageBase<PtrType>::LoadAndUnload( const TSet<FSdfPath>& LoadSet, const TSet<FSdfPath>& UnloadSet, EUsdLoadPolicy Policy )
+	void FUsdStageBase<PtrType>::LoadAndUnload(const TSet<FSdfPath>& LoadSet, const TSet<FSdfPath>& UnloadSet, EUsdLoadPolicy Policy)
 	{
 #if USE_USD_SDK
-		if ( PtrType& Ptr = Impl->GetInner() )
+		if (PtrType& Ptr = Impl->GetInner())
 		{
 			FScopedUsdAllocs Allocs;
 
 			pxr::SdfPathSet UsdLoadSet;
-			for ( const FSdfPath& Path : LoadSet )
+			for (const FSdfPath& Path : LoadSet)
 			{
-				UsdLoadSet.insert( static_cast< pxr::SdfPath >( Path ) );
+				UsdLoadSet.insert(static_cast<pxr::SdfPath>(Path));
 			}
 
 			pxr::SdfPathSet UsdUnloadSet;
-			for ( const FSdfPath& Path : UnloadSet )
+			for (const FSdfPath& Path : UnloadSet)
 			{
-				UsdUnloadSet.insert( static_cast< pxr::SdfPath >( Path ) );
+				UsdUnloadSet.insert(static_cast<pxr::SdfPath>(Path));
 			}
 
 			Ptr->LoadAndUnload(
 				UsdLoadSet,
 				UsdUnloadSet,
-				Policy == EUsdLoadPolicy::UsdLoadWithDescendants
-					? pxr::UsdLoadPolicy::UsdLoadWithDescendants
-					: pxr::UsdLoadPolicy::UsdLoadWithoutDescendants
+				Policy == EUsdLoadPolicy::UsdLoadWithDescendants ? pxr::UsdLoadPolicy::UsdLoadWithDescendants
+																 : pxr::UsdLoadPolicy::UsdLoadWithoutDescendants
 			);
 		}
 #endif
@@ -235,31 +232,31 @@ namespace UE
 	FSdfLayer FUsdStageBase<PtrType>::GetRootLayer() const
 	{
 #if USE_USD_SDK
-		if ( PtrType& Ptr = Impl->GetInner() )
+		if (PtrType& Ptr = Impl->GetInner())
 		{
-			return FSdfLayer( Ptr->GetRootLayer() );
+			return FSdfLayer(Ptr->GetRootLayer());
 		}
-#endif // #if USE_USD_SDK
+#endif	  // #if USE_USD_SDK
 		return FSdfLayer();
 	}
 
 	template<typename PtrType>
-	bool FUsdStageBase<PtrType>::Export( const TCHAR* FileName, bool bAddSourceFileComment, const TMap<FString, FString>& FileFormatArguments ) const
+	bool FUsdStageBase<PtrType>::Export(const TCHAR* FileName, bool bAddSourceFileComment, const TMap<FString, FString>& FileFormatArguments) const
 	{
 #if USE_USD_SDK
-		if ( PtrType& Ptr = Impl->GetInner() )
+		if (PtrType& Ptr = Impl->GetInner())
 		{
 			FScopedUsdAllocs Allocs;
 
 			std::map<std::string, std::string> UsdFileFormatArguments;
-			for ( const TPair<FString, FString>& Pair : FileFormatArguments )
+			for (const TPair<FString, FString>& Pair : FileFormatArguments)
 			{
-				UsdFileFormatArguments.insert( { TCHAR_TO_ANSI( *Pair.Key ), TCHAR_TO_ANSI( *Pair.Value ) } );
+				UsdFileFormatArguments.insert({TCHAR_TO_ANSI(*Pair.Key), TCHAR_TO_ANSI(*Pair.Value)});
 			}
 
-			return Ptr->Export( TCHAR_TO_ANSI( FileName ), bAddSourceFileComment, UsdFileFormatArguments );
+			return Ptr->Export(TCHAR_TO_ANSI(FileName), bAddSourceFileComment, UsdFileFormatArguments);
 		}
-#endif // #if USE_USD_SDK
+#endif	  // #if USE_USD_SDK
 
 		return false;
 	}
@@ -268,24 +265,24 @@ namespace UE
 	FSdfLayer FUsdStageBase<PtrType>::GetSessionLayer() const
 	{
 #if USE_USD_SDK
-		if ( PtrType& Ptr = Impl->GetInner() )
+		if (PtrType& Ptr = Impl->GetInner())
 		{
-			return FSdfLayer( Ptr->GetSessionLayer() );
+			return FSdfLayer(Ptr->GetSessionLayer());
 		}
-#endif // #if USE_USD_SDK
+#endif	  // #if USE_USD_SDK
 
 		return FSdfLayer();
 	}
 
 	template<typename PtrType>
-	bool FUsdStageBase<PtrType>::HasLocalLayer( const FSdfLayer & Layer ) const
+	bool FUsdStageBase<PtrType>::HasLocalLayer(const FSdfLayer& Layer) const
 	{
 #if USE_USD_SDK
-		if ( PtrType& Ptr = Impl->GetInner() )
+		if (PtrType& Ptr = Impl->GetInner())
 		{
-			return Ptr->HasLocalLayer( pxr::SdfLayerRefPtr( Layer ) );
+			return Ptr->HasLocalLayer(pxr::SdfLayerRefPtr(Layer));
 		}
-#endif // #if USE_USD_SDK
+#endif	  // #if USE_USD_SDK
 
 		return false;
 	}
@@ -294,11 +291,11 @@ namespace UE
 	FUsdPrim FUsdStageBase<PtrType>::GetPseudoRoot() const
 	{
 #if USE_USD_SDK
-		if ( PtrType& Ptr = Impl->GetInner() )
+		if (PtrType& Ptr = Impl->GetInner())
 		{
-			return FUsdPrim( Ptr->GetPseudoRoot() );
+			return FUsdPrim(Ptr->GetPseudoRoot());
 		}
-#endif // #if USE_USD_SDK
+#endif	  // #if USE_USD_SDK
 
 		return FUsdPrim();
 	}
@@ -307,24 +304,24 @@ namespace UE
 	FUsdPrim FUsdStageBase<PtrType>::GetDefaultPrim() const
 	{
 #if USE_USD_SDK
-		if ( PtrType& Ptr = Impl->GetInner() )
+		if (PtrType& Ptr = Impl->GetInner())
 		{
-			return FUsdPrim( Ptr->GetDefaultPrim() );
+			return FUsdPrim(Ptr->GetDefaultPrim());
 		}
-#endif // #if USE_USD_SDK
+#endif	  // #if USE_USD_SDK
 
 		return FUsdPrim();
 	}
 
 	template<typename PtrType>
-	FUsdPrim FUsdStageBase<PtrType>::GetPrimAtPath( const FSdfPath & Path ) const
+	FUsdPrim FUsdStageBase<PtrType>::GetPrimAtPath(const FSdfPath& Path) const
 	{
 #if USE_USD_SDK
-		if ( PtrType& Ptr = Impl->GetInner() )
+		if (PtrType& Ptr = Impl->GetInner())
 		{
-			return FUsdPrim( Ptr->GetPrimAtPath( Path ) );
+			return FUsdPrim(Ptr->GetPrimAtPath(Path));
 		}
-#endif // #if USE_USD_SDK
+#endif	  // #if USE_USD_SDK
 
 		return FUsdPrim();
 	}
@@ -333,113 +330,113 @@ namespace UE
 	bool FUsdStageBase<PtrType>::IsEditTargetValid() const
 	{
 #if USE_USD_SDK
-		if ( PtrType& Ptr = Impl->GetInner() )
+		if (PtrType& Ptr = Impl->GetInner())
 		{
 			return Ptr->GetEditTarget().IsValid();
 		}
-#endif // #if USE_USD_SDK
+#endif	  // #if USE_USD_SDK
 
 		return false;
 	}
 
 	template<typename PtrType>
-	void FUsdStageBase<PtrType>::SetEditTarget( const FSdfLayer & Layer )
+	void FUsdStageBase<PtrType>::SetEditTarget(const FSdfLayer& Layer)
 	{
 #if USE_USD_SDK
-		if ( PtrType& Ptr = Impl->GetInner() )
+		if (PtrType& Ptr = Impl->GetInner())
 		{
 			FScopedUsdAllocs UsdAllocs;
 
-			pxr::SdfLayerRefPtr LayerRef( Layer );
-			const pxr::UsdEditTarget EditTarget = Ptr->GetEditTargetForLocalLayer( LayerRef );
+			pxr::SdfLayerRefPtr LayerRef(Layer);
+			const pxr::UsdEditTarget EditTarget = Ptr->GetEditTargetForLocalLayer(LayerRef);
 
-			Ptr->SetEditTarget( EditTarget );
+			Ptr->SetEditTarget(EditTarget);
 		}
-#endif // #if USE_USD_SDK
+#endif	  // #if USE_USD_SDK
 	}
 
 	template<typename PtrType>
-	TArray<UE::FSdfLayer> FUsdStageBase<PtrType>::GetLayerStack( bool bIncludeSessionLayers /*= true */ ) const
+	TArray<UE::FSdfLayer> FUsdStageBase<PtrType>::GetLayerStack(bool bIncludeSessionLayers /*= true */) const
 	{
 		TArray<UE::FSdfLayer> Result;
 #if USE_USD_SDK
-		if ( PtrType& Ptr = Impl->GetInner() )
+		if (PtrType& Ptr = Impl->GetInner())
 		{
 			FScopedUsdAllocs UsdAllocs;
 
-			pxr::SdfLayerHandleVector LayerHandleVector = Ptr->GetLayerStack( bIncludeSessionLayers );
+			pxr::SdfLayerHandleVector LayerHandleVector = Ptr->GetLayerStack(bIncludeSessionLayers);
 
-			Result.Reserve( LayerHandleVector.size() );
+			Result.Reserve(LayerHandleVector.size());
 
-			for ( pxr::SdfLayerHandle Layer : LayerHandleVector )
+			for (pxr::SdfLayerHandle Layer : LayerHandleVector)
 			{
-				Result.Add( UE::FSdfLayer{ Layer } );
+				Result.Add(UE::FSdfLayer{Layer});
 			}
 		}
-#endif // #if USE_USD_SDK
+#endif	  // #if USE_USD_SDK
 		return Result;
 	}
 
 	template<typename PtrType>
-	TArray<UE::FSdfLayer> FUsdStageBase<PtrType>::GetUsedLayers( bool bIncludeClipLayers /*= true */ ) const
+	TArray<UE::FSdfLayer> FUsdStageBase<PtrType>::GetUsedLayers(bool bIncludeClipLayers /*= true */) const
 	{
 		TArray<UE::FSdfLayer> Result;
 #if USE_USD_SDK
-		if ( PtrType& Ptr = Impl->GetInner() )
+		if (PtrType& Ptr = Impl->GetInner())
 		{
 			FScopedUsdAllocs UsdAllocs;
 
-			pxr::SdfLayerHandleVector LayerHandleVector = Ptr->GetUsedLayers( bIncludeClipLayers );
+			pxr::SdfLayerHandleVector LayerHandleVector = Ptr->GetUsedLayers(bIncludeClipLayers);
 
-			Result.Reserve( LayerHandleVector.size() );
+			Result.Reserve(LayerHandleVector.size());
 
-			for ( pxr::SdfLayerHandle Layer : LayerHandleVector )
+			for (pxr::SdfLayerHandle Layer : LayerHandleVector)
 			{
-				Result.Add( UE::FSdfLayer{ Layer } );
+				Result.Add(UE::FSdfLayer{Layer});
 			}
 		}
-#endif // #if USE_USD_SDK
+#endif	  // #if USE_USD_SDK
 		return Result;
 	}
 
 	template<typename PtrType>
-	void FUsdStageBase<PtrType>::MuteAndUnmuteLayers( const TArray<FString>& MuteLayers, const TArray<FString>& UnmuteLayers )
+	void FUsdStageBase<PtrType>::MuteAndUnmuteLayers(const TArray<FString>& MuteLayers, const TArray<FString>& UnmuteLayers)
 	{
 #if USE_USD_SDK
-		if ( PtrType& Ptr = Impl->GetInner() )
+		if (PtrType& Ptr = Impl->GetInner())
 		{
 			FScopedUsdAllocs UsdAllocs;
 
 			std::vector<std::string> UsdMuteLayers;
-			UsdMuteLayers.reserve( MuteLayers.Num() );
-			for ( const FString& MuteLayer : MuteLayers )
+			UsdMuteLayers.reserve(MuteLayers.Num());
+			for (const FString& MuteLayer : MuteLayers)
 			{
-				UsdMuteLayers.push_back( TCHAR_TO_ANSI( *MuteLayer ) );
+				UsdMuteLayers.push_back(TCHAR_TO_ANSI(*MuteLayer));
 			}
 
 			std::vector<std::string> UsdUnmuteLayers;
-			UsdUnmuteLayers.reserve( UnmuteLayers.Num() );
-			for ( const FString& UnmuteLayer : UnmuteLayers )
+			UsdUnmuteLayers.reserve(UnmuteLayers.Num());
+			for (const FString& UnmuteLayer : UnmuteLayers)
 			{
-				UsdUnmuteLayers.push_back( TCHAR_TO_ANSI( *UnmuteLayer ) );
+				UsdUnmuteLayers.push_back(TCHAR_TO_ANSI(*UnmuteLayer));
 			}
 
-			Ptr->MuteAndUnmuteLayers( UsdMuteLayers, UsdUnmuteLayers );
+			Ptr->MuteAndUnmuteLayers(UsdMuteLayers, UsdUnmuteLayers);
 		}
-#endif // #if USE_USD_SDK
+#endif	  // #if USE_USD_SDK
 	}
 
 	template<typename PtrType>
-	bool FUsdStageBase<PtrType>::IsLayerMuted( const FString& LayerIdentifier ) const
+	bool FUsdStageBase<PtrType>::IsLayerMuted(const FString& LayerIdentifier) const
 	{
 #if USE_USD_SDK
-		if ( PtrType& Ptr = Impl->GetInner() )
+		if (PtrType& Ptr = Impl->GetInner())
 		{
 			FScopedUsdAllocs UsdAllocs;
 
-			return Ptr->IsLayerMuted( TCHAR_TO_ANSI( *LayerIdentifier ) );
+			return Ptr->IsLayerMuted(TCHAR_TO_ANSI(*LayerIdentifier));
 		}
-#endif // #if USE_USD_SDK
+#endif	  // #if USE_USD_SDK
 
 		return false;
 	}
@@ -448,66 +445,66 @@ namespace UE
 	FSdfLayer FUsdStageBase<PtrType>::GetEditTarget() const
 	{
 #if USE_USD_SDK
-		if ( PtrType& Ptr = Impl->GetInner() )
+		if (PtrType& Ptr = Impl->GetInner())
 		{
-			if ( IsEditTargetValid() )
+			if (IsEditTargetValid())
 			{
-				return FSdfLayer( Ptr->GetEditTarget().GetLayer() );
+				return FSdfLayer(Ptr->GetEditTarget().GetLayer());
 			}
 		}
-#endif // #if USE_USD_SDK
+#endif	  // #if USE_USD_SDK
 
 		return FSdfLayer();
 	}
 
 	template<typename PtrType>
-	bool FUsdStageBase<PtrType>::GetMetadata( const TCHAR * Key, UE::FVtValue & Value ) const
+	bool FUsdStageBase<PtrType>::GetMetadata(const TCHAR* Key, UE::FVtValue& Value) const
 	{
 #if USE_USD_SDK
-		if ( PtrType& Ptr = Impl->GetInner() )
+		if (PtrType& Ptr = Impl->GetInner())
 		{
-			return Ptr->GetMetadata( pxr::TfToken{ TCHAR_TO_ANSI( Key ) }, &Value.GetUsdValue() );
+			return Ptr->GetMetadata(pxr::TfToken{TCHAR_TO_ANSI(Key)}, &Value.GetUsdValue());
 		}
-#endif // #if USE_USD_SDK
+#endif	  // #if USE_USD_SDK
 
 		return false;
 	}
 
 	template<typename PtrType>
-	bool FUsdStageBase<PtrType>::HasMetadata( const TCHAR * Key ) const
+	bool FUsdStageBase<PtrType>::HasMetadata(const TCHAR* Key) const
 	{
 #if USE_USD_SDK
-		if ( PtrType& Ptr = Impl->GetInner() )
+		if (PtrType& Ptr = Impl->GetInner())
 		{
-			return Ptr->HasMetadata( pxr::TfToken{ TCHAR_TO_ANSI( Key ) } );
+			return Ptr->HasMetadata(pxr::TfToken{TCHAR_TO_ANSI(Key)});
 		}
-#endif // #if USE_USD_SDK
+#endif	  // #if USE_USD_SDK
 
 		return false;
 	}
 
 	template<typename PtrType>
-	bool FUsdStageBase<PtrType>::SetMetadata( const TCHAR * Key, const UE::FVtValue & Value ) const
+	bool FUsdStageBase<PtrType>::SetMetadata(const TCHAR* Key, const UE::FVtValue& Value) const
 	{
 #if USE_USD_SDK
-		if ( PtrType& Ptr = Impl->GetInner() )
+		if (PtrType& Ptr = Impl->GetInner())
 		{
-			return Ptr->SetMetadata( pxr::TfToken{ TCHAR_TO_ANSI( Key ) }, Value.GetUsdValue() );
+			return Ptr->SetMetadata(pxr::TfToken{TCHAR_TO_ANSI(Key)}, Value.GetUsdValue());
 		}
-#endif // #if USE_USD_SDK
+#endif	  // #if USE_USD_SDK
 
 		return false;
 	}
 
 	template<typename PtrType>
-	bool FUsdStageBase<PtrType>::ClearMetadata( const TCHAR * Key ) const
+	bool FUsdStageBase<PtrType>::ClearMetadata(const TCHAR* Key) const
 	{
 #if USE_USD_SDK
-		if ( PtrType& Ptr = Impl->GetInner() )
+		if (PtrType& Ptr = Impl->GetInner())
 		{
-			return Ptr->ClearMetadata( pxr::TfToken{ TCHAR_TO_ANSI( Key ) } );
+			return Ptr->ClearMetadata(pxr::TfToken{TCHAR_TO_ANSI(Key)});
 		}
-#endif // #if USE_USD_SDK
+#endif	  // #if USE_USD_SDK
 
 		return false;
 	}
@@ -516,11 +513,11 @@ namespace UE
 	double FUsdStageBase<PtrType>::GetStartTimeCode() const
 	{
 #if USE_USD_SDK
-		if ( PtrType& Ptr = Impl->GetInner() )
+		if (PtrType& Ptr = Impl->GetInner())
 		{
 			return Ptr->GetStartTimeCode();
 		}
-#endif // #if USE_USD_SDK
+#endif	  // #if USE_USD_SDK
 
 		return 0.0;
 	}
@@ -529,66 +526,66 @@ namespace UE
 	double FUsdStageBase<PtrType>::GetEndTimeCode() const
 	{
 #if USE_USD_SDK
-		if ( PtrType& Ptr = Impl->GetInner() )
+		if (PtrType& Ptr = Impl->GetInner())
 		{
 			return Ptr->GetEndTimeCode();
 		}
-#endif // #if USE_USD_SDK
+#endif	  // #if USE_USD_SDK
 
 		return 0.0;
 	}
 
 	template<typename PtrType>
-	void FUsdStageBase<PtrType>::SetStartTimeCode( double TimeCode )
+	void FUsdStageBase<PtrType>::SetStartTimeCode(double TimeCode)
 	{
 #if USE_USD_SDK
-		if ( PtrType& Ptr = Impl->GetInner() )
+		if (PtrType& Ptr = Impl->GetInner())
 		{
-			if ( !FMath::IsNearlyEqual( TimeCode, Ptr->GetStartTimeCode() ) || !Ptr->HasAuthoredTimeCodeRange() )
+			if (!FMath::IsNearlyEqual(TimeCode, Ptr->GetStartTimeCode()) || !Ptr->HasAuthoredTimeCodeRange())
 			{
-				Ptr->SetStartTimeCode( TimeCode );
+				Ptr->SetStartTimeCode(TimeCode);
 			}
 		}
-#endif // #if USE_USD_SDK
+#endif	  // #if USE_USD_SDK
 	}
 
 	template<typename PtrType>
-	void FUsdStageBase<PtrType>::SetEndTimeCode( double TimeCode )
+	void FUsdStageBase<PtrType>::SetEndTimeCode(double TimeCode)
 	{
 #if USE_USD_SDK
-		if ( PtrType& Ptr = Impl->GetInner() )
+		if (PtrType& Ptr = Impl->GetInner())
 		{
-			if ( !FMath::IsNearlyEqual( TimeCode, Ptr->GetEndTimeCode() ) || !Ptr->HasAuthoredTimeCodeRange() )
+			if (!FMath::IsNearlyEqual(TimeCode, Ptr->GetEndTimeCode()) || !Ptr->HasAuthoredTimeCodeRange())
 			{
-				Ptr->SetEndTimeCode( TimeCode );
+				Ptr->SetEndTimeCode(TimeCode);
 			}
 		}
-#endif // #if USE_USD_SDK
+#endif	  // #if USE_USD_SDK
 	}
 
 	template<typename PtrType>
 	double FUsdStageBase<PtrType>::GetTimeCodesPerSecond() const
 	{
 #if USE_USD_SDK
-		if ( PtrType& Ptr = Impl->GetInner() )
+		if (PtrType& Ptr = Impl->GetInner())
 		{
 			return Ptr->GetTimeCodesPerSecond();
 		}
-#endif // #if USE_USD_SDK
+#endif	  // #if USE_USD_SDK
 
 		return 24.0;
 	}
 
 	template<typename PtrType>
-	void FUsdStageBase<PtrType>::SetTimeCodesPerSecond( double TimeCodesPerSecond )
+	void FUsdStageBase<PtrType>::SetTimeCodesPerSecond(double TimeCodesPerSecond)
 	{
 #if USE_USD_SDK
-		if ( PtrType& Ptr = Impl->GetInner() )
+		if (PtrType& Ptr = Impl->GetInner())
 		{
 			FScopedUsdAllocs Allocs;
 
-			pxr::UsdEditContext EditContext( Ptr, Ptr->GetRootLayer() );
-			Ptr->SetTimeCodesPerSecond( TimeCodesPerSecond );
+			pxr::UsdEditContext EditContext(Ptr, Ptr->GetRootLayer());
+			Ptr->SetTimeCodesPerSecond(TimeCodesPerSecond);
 		}
 #endif
 	}
@@ -597,38 +594,40 @@ namespace UE
 	double FUsdStageBase<PtrType>::GetFramesPerSecond() const
 	{
 #if USE_USD_SDK
-		if ( PtrType& Ptr = Impl->GetInner() )
+		if (PtrType& Ptr = Impl->GetInner())
 		{
 			return Ptr->GetFramesPerSecond();
 		}
-#endif // #if USE_USD_SDK
+#endif	  // #if USE_USD_SDK
 
 		return 24.0;
 	}
 
 	template<typename PtrType>
-	void FUsdStageBase<PtrType>::SetFramesPerSecond( double FramesPerSecond )
+	void FUsdStageBase<PtrType>::SetFramesPerSecond(double FramesPerSecond)
 	{
 #if USE_USD_SDK
-		if ( PtrType& Ptr = Impl->GetInner() )
+		if (PtrType& Ptr = Impl->GetInner())
 		{
 			FScopedUsdAllocs Allocs;
 
-			pxr::UsdEditContext EditContext( Ptr, Ptr->GetRootLayer() );
-			Ptr->SetFramesPerSecond( FramesPerSecond );
+			pxr::UsdEditContext EditContext(Ptr, Ptr->GetRootLayer());
+			Ptr->SetFramesPerSecond(FramesPerSecond);
 		}
 #endif
 	}
 
 	template<typename PtrType>
-	void FUsdStageBase<PtrType>::SetInterpolationType( EUsdInterpolationType InterpolationType )
+	void FUsdStageBase<PtrType>::SetInterpolationType(EUsdInterpolationType InterpolationType)
 	{
 #if USE_USD_SDK
-		if ( PtrType& Ptr = Impl->GetInner() )
+		if (PtrType& Ptr = Impl->GetInner())
 		{
 			FScopedUsdAllocs Allocs;
 
-			Ptr->SetInterpolationType( InterpolationType == EUsdInterpolationType::Held ? pxr::UsdInterpolationTypeHeld : pxr::UsdInterpolationTypeLinear );
+			Ptr->SetInterpolationType(
+				InterpolationType == EUsdInterpolationType::Held ? pxr::UsdInterpolationTypeHeld : pxr::UsdInterpolationTypeLinear
+			);
 		}
 #endif
 	}
@@ -637,9 +636,9 @@ namespace UE
 	EUsdInterpolationType FUsdStageBase<PtrType>::GetInterpolationType() const
 	{
 #if USE_USD_SDK
-		if ( PtrType& Ptr = Impl->GetInner() )
+		if (PtrType& Ptr = Impl->GetInner())
 		{
-			if ( Ptr->GetInterpolationType() == pxr::UsdInterpolationTypeHeld )
+			if (Ptr->GetInterpolationType() == pxr::UsdInterpolationTypeHeld)
 			{
 				return EUsdInterpolationType::Held;
 			}
@@ -650,38 +649,38 @@ namespace UE
 	}
 
 	template<typename PtrType>
-	void FUsdStageBase<PtrType>::SetDefaultPrim( const FUsdPrim & Prim )
+	void FUsdStageBase<PtrType>::SetDefaultPrim(const FUsdPrim& Prim)
 	{
 #if USE_USD_SDK
-		if ( PtrType& Ptr = Impl->GetInner() )
+		if (PtrType& Ptr = Impl->GetInner())
 		{
-			Ptr->SetDefaultPrim( Prim );
+			Ptr->SetDefaultPrim(Prim);
 		}
-#endif // #if USE_USD_SDK
+#endif	  // #if USE_USD_SDK
 	}
 
 	template<typename PtrType>
-	FUsdPrim FUsdStageBase<PtrType>::OverridePrim( const FSdfPath& Path )
+	FUsdPrim FUsdStageBase<PtrType>::OverridePrim(const FSdfPath& Path)
 	{
 #if USE_USD_SDK
-		if ( PtrType& Ptr = Impl->GetInner() )
+		if (PtrType& Ptr = Impl->GetInner())
 		{
-			return FUsdPrim( Ptr->OverridePrim( Path ) );
+			return FUsdPrim(Ptr->OverridePrim(Path));
 		}
-#endif // #if USE_USD_SDK
+#endif	  // #if USE_USD_SDK
 
 		return FUsdPrim();
 	}
 
 	template<typename PtrType>
-	FUsdPrim FUsdStageBase<PtrType>::DefinePrim( const FSdfPath & Path, const TCHAR * TypeName )
+	FUsdPrim FUsdStageBase<PtrType>::DefinePrim(const FSdfPath& Path, const TCHAR* TypeName)
 	{
 #if USE_USD_SDK
-		if ( PtrType& Ptr = Impl->GetInner() )
+		if (PtrType& Ptr = Impl->GetInner())
 		{
-			return FUsdPrim( Ptr->DefinePrim( Path, pxr::TfToken( TCHAR_TO_ANSI( TypeName ) ) ) );
+			return FUsdPrim(Ptr->DefinePrim(Path, pxr::TfToken(TCHAR_TO_ANSI(TypeName))));
 		}
-#endif // #if USE_USD_SDK
+#endif	  // #if USE_USD_SDK
 
 		return FUsdPrim();
 	}
@@ -694,20 +693,20 @@ namespace UE
 		{
 			return FUsdPrim(Ptr->CreateClassPrim(RootPrimPath));
 		}
-#endif // #if USE_USD_SDK
+#endif	  // #if USE_USD_SDK
 
 		return FUsdPrim();
 	}
 
 	template<typename PtrType>
-	bool FUsdStageBase<PtrType>::RemovePrim( const FSdfPath & Path )
+	bool FUsdStageBase<PtrType>::RemovePrim(const FSdfPath& Path)
 	{
 #if USE_USD_SDK
-		if ( PtrType& Ptr = Impl->GetInner() )
+		if (PtrType& Ptr = Impl->GetInner())
 		{
-			return Ptr->RemovePrim( Path );
+			return Ptr->RemovePrim(Path);
 		}
-#endif // #if USE_USD_SDK
+#endif	  // #if USE_USD_SDK
 
 		return false;
 	}
@@ -718,5 +717,5 @@ namespace UE
 #else
 	template class UNREALUSDWRAPPER_API FUsdStageBase<FDummyRefPtrType>;
 	template class UNREALUSDWRAPPER_API FUsdStageBase<FDummyWeakPtrType>;
-#endif // #if USE_USD_SDK
-}
+#endif	  // #if USE_USD_SDK
+}	 // namespace UE
diff --git a/Engine/Plugins/Importers/USDImporter/Source/UnrealUSDWrapper/Private/UsdWrappers/UsdTyped.cpp b/Engine/Plugins/Importers/USDImporter/Source/UnrealUSDWrapper/Private/UsdWrappers/UsdTyped.cpp
--- a/Engine/Plugins/Importers/USDImporter/Source/UnrealUSDWrapper/Private/UsdWrappers/UsdTyped.cpp
+++ b/Engine/Plugins/Importers/USDImporter/Source/UnrealUSDWrapper/Private/UsdWrappers/UsdTyped.cpp
@@ -8,12 +8,10 @@
 #include "USDMemory.h"
 
 #if USE_USD_SDK
-
 #include "USDIncludesStart.h"
-	#include "pxr/usd/usd/typed.h"
+#include "pxr/usd/usd/typed.h"
 #include "USDIncludesEnd.h"
-
-#endif //  #if USE_USD_SDK
+#endif	  //  #if USE_USD_SDK
 
 namespace UE
 {
@@ -25,63 +23,63 @@ namespace UE
 			FUsdTypedImpl() = default;
 
 #if USE_USD_SDK
-			explicit FUsdTypedImpl( const pxr::UsdTyped& InUsdTyped )
-				: PxrUsdTyped( InUsdTyped )
+			explicit FUsdTypedImpl(const pxr::UsdTyped& InUsdTyped)
+				: PxrUsdTyped(InUsdTyped)
 			{
 			}
 
-			explicit FUsdTypedImpl( pxr::UsdTyped&& InUsdTyped )
-				: PxrUsdTyped( MoveTemp( InUsdTyped ) )
+			explicit FUsdTypedImpl(pxr::UsdTyped&& InUsdTyped)
+				: PxrUsdTyped(MoveTemp(InUsdTyped))
 			{
 			}
 
-			explicit FUsdTypedImpl( const pxr::UsdPrim& InUsdPrim )
-				: PxrUsdTyped( pxr::UsdTyped( InUsdPrim ) )
+			explicit FUsdTypedImpl(const pxr::UsdPrim& InUsdPrim)
+				: PxrUsdTyped(pxr::UsdTyped(InUsdPrim))
 			{
 			}
 
-			TUsdStore< pxr::UsdTyped > PxrUsdTyped;
-#endif // #if USE_USD_SDK
+			TUsdStore<pxr::UsdTyped> PxrUsdTyped;
+#endif	  // #if USE_USD_SDK
 		};
 	}
 
 	FUsdTyped::FUsdTyped()
 	{
 		FScopedUnrealAllocs UnrealAllocs;
-		Impl = MakeUnique< Internal::FUsdTypedImpl >();
+		Impl = MakeUnique<Internal::FUsdTypedImpl>();
 	}
 
-	FUsdTyped::FUsdTyped( const FUsdTyped& Other )
+	FUsdTyped::FUsdTyped(const FUsdTyped& Other)
 	{
 #if USE_USD_SDK
 		FScopedUnrealAllocs UnrealAllocs;
-		Impl = MakeUnique< Internal::FUsdTypedImpl >( Other.Impl->PxrUsdTyped.Get() );
-#endif // #if USE_USD_SDK
+		Impl = MakeUnique<Internal::FUsdTypedImpl>(Other.Impl->PxrUsdTyped.Get());
+#endif	  // #if USE_USD_SDK
 	}
 
-	FUsdTyped::FUsdTyped( FUsdTyped&& Other ) = default;
+	FUsdTyped::FUsdTyped(FUsdTyped&& Other) = default;
 
-	FUsdTyped::FUsdTyped( const FUsdPrim& Prim )
+	FUsdTyped::FUsdTyped(const FUsdPrim& Prim)
 	{
 #if USE_USD_SDK
 		FScopedUnrealAllocs UnrealAllocs;
-		Impl = MakeUnique< Internal::FUsdTypedImpl >( Prim );
-#endif // #if USE_USD_SDK
+		Impl = MakeUnique<Internal::FUsdTypedImpl>(Prim);
+#endif	  // #if USE_USD_SDK
 	}
 
-	FUsdTyped& FUsdTyped::operator=( const FUsdTyped& Other )
+	FUsdTyped& FUsdTyped::operator=(const FUsdTyped& Other)
 	{
 #if USE_USD_SDK
 		FScopedUnrealAllocs UnrealAllocs;
-		Impl = MakeUnique< Internal::FUsdTypedImpl >( Other.Impl->PxrUsdTyped.Get() );
-#endif // #if USE_USD_SDK
+		Impl = MakeUnique<Internal::FUsdTypedImpl>(Other.Impl->PxrUsdTyped.Get());
+#endif	  // #if USE_USD_SDK
 		return *this;
 	}
 
-	FUsdTyped& FUsdTyped::operator=( FUsdTyped&& Other )
+	FUsdTyped& FUsdTyped::operator=(FUsdTyped&& Other)
 	{
 		FScopedUnrealAllocs UnrealAllocs;
-		Impl = MoveTemp( Other.Impl );
+		Impl = MoveTemp(Other.Impl);
 
 		return *this;
 	}
@@ -98,26 +96,26 @@ namespace UE
 		return (bool)Impl->PxrUsdTyped.Get();
 #else
 		return false;
-#endif // #if USE_USD_SDK
+#endif	  // #if USE_USD_SDK
 	}
 
 #if USE_USD_SDK
-	FUsdTyped::FUsdTyped( const pxr::UsdTyped& InUsdTyped )		
+	FUsdTyped::FUsdTyped(const pxr::UsdTyped& InUsdTyped)
 	{
 		FScopedUnrealAllocs UnrealAllocs;
-		Impl = MakeUnique< Internal::FUsdTypedImpl >( InUsdTyped );
+		Impl = MakeUnique<Internal::FUsdTypedImpl>(InUsdTyped);
 	}
 
-	FUsdTyped::FUsdTyped( pxr::UsdTyped&& InUsdTyped )
+	FUsdTyped::FUsdTyped(pxr::UsdTyped&& InUsdTyped)
 	{
 		FScopedUnrealAllocs UnrealAllocs;
-		Impl = MakeUnique< Internal::FUsdTypedImpl >( MoveTemp( InUsdTyped ) );
+		Impl = MakeUnique<Internal::FUsdTypedImpl>(MoveTemp(InUsdTyped));
 	}
 
-	FUsdTyped::FUsdTyped( const pxr::UsdPrim& Prim )
+	FUsdTyped::FUsdTyped(const pxr::UsdPrim& Prim)
 	{
 		FScopedUnrealAllocs UnrealAllocs;
-		Impl = MakeUnique< Internal::FUsdTypedImpl >( Prim );
+		Impl = MakeUnique<Internal::FUsdTypedImpl>(Prim);
 	}
 
 	FUsdTyped::operator pxr::UsdTyped&()
@@ -129,23 +127,23 @@ namespace UE
 	{
 		return Impl->PxrUsdTyped.Get();
 	}
-#endif // #if USE_USD_SDK
+#endif	  // #if USE_USD_SDK
 
 	FSdfPath FUsdTyped::GetPath() const
 	{
 #if USE_USD_SDK
-		return FSdfPath( Impl->PxrUsdTyped.Get().GetPath() );
+		return FSdfPath(Impl->PxrUsdTyped.Get().GetPath());
 #else
 		return FSdfPath();
-#endif // #if USE_USD_SDK
+#endif	  // #if USE_USD_SDK
 	}
 
 	FUsdPrim FUsdTyped::GetPrim() const
 	{
 #if USE_USD_SDK
-		return FUsdPrim( Impl->PxrUsdTyped.Get().GetPrim() );
+		return FUsdPrim(Impl->PxrUsdTyped.Get().GetPrim());
 #else
 		return FUsdPrim();
-#endif // #if USE_USD_SDK
+#endif	  // #if USE_USD_SDK
 	}
-}
+}	 // namespace UE
diff --git a/Engine/Plugins/Importers/USDImporter/Source/UnrealUSDWrapper/Private/UsdWrappers/UsdVariantSets.cpp b/Engine/Plugins/Importers/USDImporter/Source/UnrealUSDWrapper/Private/UsdWrappers/UsdVariantSets.cpp
--- a/Engine/Plugins/Importers/USDImporter/Source/UnrealUSDWrapper/Private/UsdWrappers/UsdVariantSets.cpp
+++ b/Engine/Plugins/Importers/USDImporter/Source/UnrealUSDWrapper/Private/UsdWrappers/UsdVariantSets.cpp
@@ -7,17 +7,15 @@
 #include "UsdWrappers/UsdPrim.h"
 
 #if USE_USD_SDK
-
 #include "USDIncludesStart.h"
-	#include "pxr/usd/usd/prim.h"
-	#include "pxr/usd/usd/variantSets.h"
+#include "pxr/usd/usd/prim.h"
+#include "pxr/usd/usd/variantSets.h"
 #include "USDIncludesEnd.h"
 
 #include <map>
 #include <utility>
 #include <vector>
-
-#endif // #if USE_USD_SDK
+#endif	  // #if USE_USD_SDK
 
 namespace UE
 {
@@ -57,7 +55,7 @@ namespace UE
 			// and the variant set name independently and fetch the object on demand.
 			TUsdStore<pxr::UsdPrim> PxrUsdPrim;
 			std::string VariantSetName;
-#endif // #if USE_USD_SDK
+#endif	  // #if USE_USD_SDK
 		};
 
 		class FUsdVariantSetsImpl
@@ -82,10 +80,9 @@ namespace UE
 			// so copy and move constructing from a pxr::UsdVariantSets object
 			// are not supported.
 			TUsdStore<pxr::UsdPrim> PxrUsdPrim;
-#endif // #if USE_USD_SDK
+#endif	  // #if USE_USD_SDK
 		};
-	}
-
+	}	  // namespace Internal
 
 	// FUsdVariantSet
 
@@ -97,10 +94,8 @@ namespace UE
 	FUsdVariantSet::FUsdVariantSet(const FUsdVariantSet& Other)
 	{
 #if USE_USD_SDK
-		Impl = MakeUnique<Internal::FUsdVariantSetImpl>(
-			Other.Impl->PxrUsdPrim.Get(),
-			Other.Impl->VariantSetName);
-#endif // #if USE_USD_SDK
+		Impl = MakeUnique<Internal::FUsdVariantSetImpl>(Other.Impl->PxrUsdPrim.Get(), Other.Impl->VariantSetName);
+#endif	  // #if USE_USD_SDK
 	}
 
 	FUsdVariantSet::FUsdVariantSet(FUsdVariantSet&& Other) = default;
@@ -108,10 +103,8 @@ namespace UE
 	FUsdVariantSet& FUsdVariantSet::operator=(const FUsdVariantSet& Other)
 	{
 #if USE_USD_SDK
-		Impl = MakeUnique<Internal::FUsdVariantSetImpl>(
-			Other.Impl->PxrUsdPrim.Get(),
-			Other.Impl->VariantSetName);
-#endif // #if USE_USD_SDK
+		Impl = MakeUnique<Internal::FUsdVariantSetImpl>(Other.Impl->PxrUsdPrim.Get(), Other.Impl->VariantSetName);
+#endif	  // #if USE_USD_SDK
 
 		return *this;
 	}
@@ -171,7 +164,7 @@ namespace UE
 	{
 		Impl = MakeUnique<Internal::FUsdVariantSetImpl>(InPrim, InVariantSetName);
 	}
-#endif // #if USE_USD_SDK
+#endif	  // #if USE_USD_SDK
 
 	bool FUsdVariantSet::AddVariant(const FString& VariantName, EUsdListPosition Position)
 	{
@@ -183,11 +176,9 @@ namespace UE
 		{
 			const std::string UsdVariantName(TCHAR_TO_ANSI(*VariantName));
 
-			return UsdVariantSet.AddVariant(
-				UsdVariantName,
-				static_cast<pxr::UsdListPosition>(Position));
+			return UsdVariantSet.AddVariant(UsdVariantName, static_cast<pxr::UsdListPosition>(Position));
 		}
-#endif // #if USE_USD_SDK
+#endif	  // #if USE_USD_SDK
 
 		return false;
 	}
@@ -215,7 +206,7 @@ namespace UE
 				VariantNames.Emplace(ANSI_TO_TCHAR(UsdVariantName.c_str()));
 			}
 		}
-#endif // #if USE_USD_SDK
+#endif	  // #if USE_USD_SDK
 
 		return VariantNames;
 	}
@@ -232,7 +223,7 @@ namespace UE
 
 			return UsdVariantSet.HasAuthoredVariant(UsdVariantName);
 		}
-#endif // #if USE_USD_SDK
+#endif	  // #if USE_USD_SDK
 
 		return false;
 	}
@@ -249,7 +240,7 @@ namespace UE
 
 			return FString(ANSI_TO_TCHAR(UsdVariantSelection.c_str()));
 		}
-#endif // #if USE_USD_SDK
+#endif	  // #if USE_USD_SDK
 
 		return FString{};
 	}
@@ -272,7 +263,7 @@ namespace UE
 
 			return bResult;
 		}
-#endif // #if USE_USD_SDK
+#endif	  // #if USE_USD_SDK
 
 		return false;
 	}
@@ -289,7 +280,7 @@ namespace UE
 
 			return UsdVariantSet.SetVariantSelection(UsdVariantName);
 		}
-#endif // #if USE_USD_SDK
+#endif	  // #if USE_USD_SDK
 
 		return false;
 	}
@@ -304,7 +295,7 @@ namespace UE
 		{
 			return UsdVariantSet.ClearVariantSelection();
 		}
-#endif // #if USE_USD_SDK
+#endif	  // #if USE_USD_SDK
 
 		return false;
 	}
@@ -319,7 +310,7 @@ namespace UE
 		{
 			return UsdVariantSet.BlockVariantSelection();
 		}
-#endif // #if USE_USD_SDK
+#endif	  // #if USE_USD_SDK
 
 		return false;
 	}
@@ -334,7 +325,7 @@ namespace UE
 		{
 			return FUsdPrim(UsdVariantSet.GetPrim());
 		}
-#endif // #if USE_USD_SDK
+#endif	  // #if USE_USD_SDK
 
 		return FUsdPrim{};
 	}
@@ -351,7 +342,7 @@ namespace UE
 
 			return FString(ANSI_TO_TCHAR(UsdVariantName.c_str()));
 		}
-#endif // #if USE_USD_SDK
+#endif	  // #if USE_USD_SDK
 
 		return FString{};
 	}
@@ -360,12 +351,11 @@ namespace UE
 	{
 #if USE_USD_SDK
 		return Impl->GetInner().IsValid();
-#endif // #if USE_USD_SDK
+#endif	  // #if USE_USD_SDK
 
 		return false;
 	}
 
-
 	// FUsdVariantSets
 
 	FUsdVariantSets::FUsdVariantSets()
@@ -376,9 +366,8 @@ namespace UE
 	FUsdVariantSets::FUsdVariantSets(const FUsdVariantSets& Other)
 	{
 #if USE_USD_SDK
-		Impl = MakeUnique<Internal::FUsdVariantSetsImpl>(
-			Other.Impl->PxrUsdPrim.Get());
-#endif // #if USE_USD_SDK
+		Impl = MakeUnique<Internal::FUsdVariantSetsImpl>(Other.Impl->PxrUsdPrim.Get());
+#endif	  // #if USE_USD_SDK
 	}
 
 	FUsdVariantSets::FUsdVariantSets(FUsdVariantSets&& Other) = default;
@@ -386,9 +375,8 @@ namespace UE
 	FUsdVariantSets& FUsdVariantSets::operator=(const FUsdVariantSets& Other)
 	{
 #if USE_USD_SDK
-		Impl = MakeUnique<Internal::FUsdVariantSetsImpl>(
-			Other.Impl->PxrUsdPrim.Get());
-#endif // #if USE_USD_SDK
+		Impl = MakeUnique<Internal::FUsdVariantSetsImpl>(Other.Impl->PxrUsdPrim.Get());
+#endif	  // #if USE_USD_SDK
 
 		return *this;
 	}
@@ -411,7 +399,7 @@ namespace UE
 		return (bool)Impl->PxrUsdPrim.Get();
 #else
 		return false;
-#endif // #if USE_USD_SDK
+#endif	  // #if USE_USD_SDK
 	}
 
 #if USE_USD_SDK
@@ -430,7 +418,7 @@ namespace UE
 	{
 		Impl = MakeUnique<Internal::FUsdVariantSetsImpl>(InPrim);
 	}
-#endif // #if USE_USD_SDK
+#endif	  // #if USE_USD_SDK
 
 	FUsdVariantSet FUsdVariantSets::AddVariantSet(const FString& VariantName, EUsdListPosition Position)
 	{
@@ -441,13 +429,10 @@ namespace UE
 
 		const std::string UsdVariantName(TCHAR_TO_ANSI(*VariantName));
 
-		return FUsdVariantSet(
-			UsdVariantSets.AddVariantSet(
-				UsdVariantName,
-				static_cast<pxr::UsdListPosition>(Position)));
+		return FUsdVariantSet(UsdVariantSets.AddVariantSet(UsdVariantName, static_cast<pxr::UsdListPosition>(Position)));
 #else
 		return FUsdVariantSet{};
-#endif // #if USE_USD_SDK
+#endif	  // #if USE_USD_SDK
 	}
 
 	TArray<FString> FUsdVariantSets::GetNames() const
@@ -458,7 +443,7 @@ namespace UE
 		FScopedUsdAllocs UsdAllocs;
 
 		pxr::UsdVariantSets UsdVariantSets = Impl->GetInner();
-		
+
 		const std::vector<std::string> UsdVariantSetNames = UsdVariantSets.GetNames();
 		if (UsdVariantSetNames.empty())
 		{
@@ -471,7 +456,7 @@ namespace UE
 		{
 			VariantSetNames.Emplace(ANSI_TO_TCHAR(UsdVariantSetName.c_str()));
 		}
-#endif // #if USE_USD_SDK
+#endif	  // #if USE_USD_SDK
 
 		return VariantSetNames;
 	}
@@ -488,7 +473,7 @@ namespace UE
 		return FUsdVariantSet(UsdVariantSets.GetVariantSet(UsdVariantSetName));
 #else
 		return FUsdVariantSet{};
-#endif // #if USE_USD_SDK
+#endif	  // #if USE_USD_SDK
 	}
 
 	bool FUsdVariantSets::HasVariantSet(const FString& VariantSetName) const
@@ -497,13 +482,13 @@ namespace UE
 		FScopedUsdAllocs UsdAllocs;
 
 		pxr::UsdVariantSets UsdVariantSets = Impl->GetInner();
-		
+
 		const std::string UsdVariantSetName(TCHAR_TO_ANSI(*VariantSetName));
 
 		return UsdVariantSets.HasVariantSet(UsdVariantSetName);
 #else
 		return false;
-#endif // #if USE_USD_SDK
+#endif	  // #if USE_USD_SDK
 	}
 
 	FString FUsdVariantSets::GetVariantSelection(const FString& VariantSetName) const
@@ -520,7 +505,7 @@ namespace UE
 		return FString(ANSI_TO_TCHAR(UsdVariantSelection.c_str()));
 #else
 		return FString{};
-#endif // #if USE_USD_SDK
+#endif	  // #if USE_USD_SDK
 	}
 
 	bool FUsdVariantSets::SetSelection(const FString& VariantSetName, const FString& VariantName)
@@ -536,7 +521,7 @@ namespace UE
 		return UsdVariantSets.SetSelection(UsdVariantSetName, UsdVariantName);
 #else
 		return false;
-#endif // #if USE_USD_SDK
+#endif	  // #if USE_USD_SDK
 	}
 
 	TMap<FString, FString> FUsdVariantSets::GetAllVariantSelections() const
@@ -558,12 +543,10 @@ namespace UE
 
 		for (const std::pair<const std::string, std::string>& UsdVariantSelection : UsdVariantSelections)
 		{
-			VariantSelections.Emplace(
-				ANSI_TO_TCHAR(UsdVariantSelection.first.c_str()),
-				ANSI_TO_TCHAR(UsdVariantSelection.second.c_str()));
+			VariantSelections.Emplace(ANSI_TO_TCHAR(UsdVariantSelection.first.c_str()), ANSI_TO_TCHAR(UsdVariantSelection.second.c_str()));
 		}
-#endif // #if USE_USD_SDK
+#endif	  // #if USE_USD_SDK
 
 		return VariantSelections;
 	}
-}
+}	 // namespace UE
diff --git a/Engine/Plugins/Importers/USDImporter/Source/UnrealUSDWrapper/Private/UsdWrappers/VtValue.cpp b/Engine/Plugins/Importers/USDImporter/Source/UnrealUSDWrapper/Private/UsdWrappers/VtValue.cpp
--- a/Engine/Plugins/Importers/USDImporter/Source/UnrealUSDWrapper/Private/UsdWrappers/VtValue.cpp
+++ b/Engine/Plugins/Importers/USDImporter/Source/UnrealUSDWrapper/Private/UsdWrappers/VtValue.cpp
@@ -5,12 +5,10 @@
 #include "USDMemory.h"
 
 #if USE_USD_SDK
-
 #include "USDIncludesStart.h"
-	#include "pxr/base/vt/value.h"
+#include "pxr/base/vt/value.h"
 #include "USDIncludesEnd.h"
-
-#endif // #if USE_USD_SDK
+#endif	  // #if USE_USD_SDK
 
 namespace UE
 {
@@ -22,50 +20,50 @@ namespace UE
 			FVtValueImpl() = default;
 
 #if USE_USD_SDK
-			explicit FVtValueImpl( const pxr::VtValue& InVtValue )
-				: PxrVtValue( InVtValue )
+			explicit FVtValueImpl(const pxr::VtValue& InVtValue)
+				: PxrVtValue(InVtValue)
 			{
 			}
 
-			explicit FVtValueImpl( pxr::VtValue&& InVtValue )
-				: PxrVtValue( MoveTemp( InVtValue ) )
+			explicit FVtValueImpl(pxr::VtValue&& InVtValue)
+				: PxrVtValue(MoveTemp(InVtValue))
 			{
 			}
 
-			TUsdStore< pxr::VtValue > PxrVtValue;
-#endif // #if USE_USD_SDK
+			TUsdStore<pxr::VtValue> PxrVtValue;
+#endif	  // #if USE_USD_SDK
 		};
 	}
 
 	FVtValue::FVtValue()
 	{
 		FScopedUnrealAllocs UnrealAllocs;
-		Impl = MakeUnique< Internal::FVtValueImpl >();
+		Impl = MakeUnique<Internal::FVtValueImpl>();
 	}
 
-	FVtValue::FVtValue( const FVtValue& Other )
+	FVtValue::FVtValue(const FVtValue& Other)
 	{
 #if USE_USD_SDK
 		FScopedUnrealAllocs UnrealAllocs;
-		Impl =  MakeUnique< Internal::FVtValueImpl >( Other.Impl->PxrVtValue.Get() );
-#endif // #if USE_USD_SDK
+		Impl = MakeUnique<Internal::FVtValueImpl>(Other.Impl->PxrVtValue.Get());
+#endif	  // #if USE_USD_SDK
 	}
 
-	FVtValue::FVtValue( FVtValue&& Other ) = default;
+	FVtValue::FVtValue(FVtValue&& Other) = default;
 
-	FVtValue& FVtValue::operator=( const FVtValue& Other )
+	FVtValue& FVtValue::operator=(const FVtValue& Other)
 	{
 #if USE_USD_SDK
 		FScopedUnrealAllocs UnrealAllocs;
-		Impl = MakeUnique< Internal::FVtValueImpl >( Other.Impl->PxrVtValue.Get() );
-#endif // #if USE_USD_SDK
+		Impl = MakeUnique<Internal::FVtValueImpl>(Other.Impl->PxrVtValue.Get());
+#endif	  // #if USE_USD_SDK
 		return *this;
 	}
 
-	FVtValue& FVtValue::operator=( FVtValue&& Other )
+	FVtValue& FVtValue::operator=(FVtValue&& Other)
 	{
 		FScopedUnrealAllocs UnrealAllocs;
-		Impl = MoveTemp( Other.Impl );
+		Impl = MoveTemp(Other.Impl);
 
 		return *this;
 	}
@@ -76,40 +74,40 @@ namespace UE
 		Impl.Reset();
 	}
 
-	bool FVtValue::operator==( const FVtValue& Other ) const
+	bool FVtValue::operator==(const FVtValue& Other) const
 	{
 #if USE_USD_SDK
 		return Impl->PxrVtValue.Get() == Other.Impl->PxrVtValue.Get();
 #else
 		return false;
-#endif // #if USE_USD_SDK
+#endif	  // #if USE_USD_SDK
 	}
 
-	bool FVtValue::operator!=( const FVtValue& Other ) const
+	bool FVtValue::operator!=(const FVtValue& Other) const
 	{
-		return !( *this == Other );
+		return !(*this == Other);
 	}
 
 #if USE_USD_SDK
-	FVtValue::FVtValue( const pxr::VtValue& InVtValue )
-		: Impl( MakeUnique< Internal::FVtValueImpl >( InVtValue ) )
+	FVtValue::FVtValue(const pxr::VtValue& InVtValue)
+		: Impl(MakeUnique<Internal::FVtValueImpl>(InVtValue))
 	{
 	}
 
-	FVtValue::FVtValue( pxr::VtValue&& InVtValue )
-		: Impl( MakeUnique< Internal::FVtValueImpl >( MoveTemp( InVtValue ) ) )
+	FVtValue::FVtValue(pxr::VtValue&& InVtValue)
+		: Impl(MakeUnique<Internal::FVtValueImpl>(MoveTemp(InVtValue)))
 	{
 	}
 
-	FVtValue& FVtValue::operator=(  const pxr::VtValue& InVtValue )
+	FVtValue& FVtValue::operator=(const pxr::VtValue& InVtValue)
 	{
-		Impl = MakeUnique< Internal::FVtValueImpl >( InVtValue );
+		Impl = MakeUnique<Internal::FVtValueImpl>(InVtValue);
 		return *this;
 	}
 
-	FVtValue& FVtValue::operator=( pxr::VtValue&& InVtValue )
+	FVtValue& FVtValue::operator=(pxr::VtValue&& InVtValue)
 	{
-		Impl = MakeUnique< Internal::FVtValueImpl >( MoveTemp( InVtValue ) );
+		Impl = MakeUnique<Internal::FVtValueImpl>(MoveTemp(InVtValue));
 		return *this;
 	}
 
@@ -122,17 +120,17 @@ namespace UE
 	{
 		return Impl->PxrVtValue.Get();
 	}
-#endif // #if USE_USD_SDK
+#endif	  // #if USE_USD_SDK
 
 	FString FVtValue::GetTypeName() const
 	{
 #if USE_USD_SDK
 		FScopedUsdAllocs UsdAllocs;
 
-		return FString( ANSI_TO_TCHAR( Impl->PxrVtValue.Get().GetTypeName().c_str() ) );
+		return FString(ANSI_TO_TCHAR(Impl->PxrVtValue.Get().GetTypeName().c_str()));
 #else
 		return FString();
-#endif // #if USE_USD_SDK
+#endif	  // #if USE_USD_SDK
 	}
 
 	bool FVtValue::IsArrayValued() const
@@ -141,7 +139,7 @@ namespace UE
 		return Impl->PxrVtValue.Get().IsArrayValued();
 #else
 		return true;
-#endif // #if USE_USD_SDK
+#endif	  // #if USE_USD_SDK
 	}
 
 	bool FVtValue::IsEmpty() const
@@ -150,6 +148,6 @@ namespace UE
 		return Impl->PxrVtValue.Get().IsEmpty();
 #else
 		return true;
-#endif // #if USE_USD_SDK
+#endif	  // #if USE_USD_SDK
 	}
-}
+}	 // namespace UE
diff --git a/Engine/Plugins/Importers/USDImporter/Source/UnrealUSDWrapper/Public/USDIncludesEnd.h b/Engine/Plugins/Importers/USDImporter/Source/UnrealUSDWrapper/Public/USDIncludesEnd.h
--- a/Engine/Plugins/Importers/USDImporter/Source/UnrealUSDWrapper/Public/USDIncludesEnd.h
+++ b/Engine/Plugins/Importers/USDImporter/Source/UnrealUSDWrapper/Public/USDIncludesEnd.h
@@ -2,6 +2,8 @@
 
 // HEADER_UNIT_SKIP - Special include
 
+// clang-format off
+
 #if USE_USD_SDK
 
 #pragma warning(pop)
@@ -24,3 +26,5 @@
 THIRD_PARTY_INCLUDES_END
 
 #endif // #if USE_USD_SDK
+
+// clang-format on
diff --git a/Engine/Plugins/Importers/USDImporter/Source/UnrealUSDWrapper/Public/USDIncludesStart.h b/Engine/Plugins/Importers/USDImporter/Source/UnrealUSDWrapper/Public/USDIncludesStart.h
--- a/Engine/Plugins/Importers/USDImporter/Source/UnrealUSDWrapper/Public/USDIncludesStart.h
+++ b/Engine/Plugins/Importers/USDImporter/Source/UnrealUSDWrapper/Public/USDIncludesStart.h
@@ -2,6 +2,8 @@
 
 // HEADER_UNIT_SKIP - Special include
 
+// clang-format off
+
 #if USE_USD_SDK
 
 #include "CoreMinimal.h"
@@ -68,3 +70,5 @@ THIRD_PARTY_INCLUDES_START
 #include "pxr/pxr.h"
 
 #endif // #if USE_USD_SDK
+
+// clang-format on
diff --git a/Engine/Plugins/Importers/USDImporter/Source/UnrealUSDWrapper/Public/USDListener.h b/Engine/Plugins/Importers/USDImporter/Source/UnrealUSDWrapper/Public/USDListener.h
--- a/Engine/Plugins/Importers/USDImporter/Source/UnrealUSDWrapper/Public/USDListener.h
+++ b/Engine/Plugins/Importers/USDImporter/Source/UnrealUSDWrapper/Public/USDListener.h
@@ -48,7 +48,7 @@ namespace UsdUtils
 
 		FPrimChangeFlags()
 		{
-			FMemory::Memset( this, 0, sizeof( *this ) );
+			FMemory::Memset(this, 0, sizeof(*this));
 		}
 	};
 
@@ -58,10 +58,10 @@ namespace UsdUtils
 	 */
 	struct FAttributeChange
 	{
-		FString PropertyName;	// metersPerUnit, kind, upAxis, etc.
-		FString Field;			// default, variability, timeSamples, etc.
-		UE::FVtValue OldValue;	// Can be empty when we create a new attribute opinion
-		UE::FVtValue NewValue;	// Can be empty when we clear an existing attribute opinion
+		FString PropertyName;	  // metersPerUnit, kind, upAxis, etc.
+		FString Field;			  // default, variability, timeSamples, etc.
+		UE::FVtValue OldValue;	  // Can be empty when we create a new attribute opinion
+		UE::FVtValue NewValue;	  // Can be empty when we clear an existing attribute opinion
 	};
 
 	/** Analogous to pxr::SdfChangeList::SubLayerChangeType, describes a change to a sublayer */
@@ -77,11 +77,12 @@ namespace UsdUtils
 	{
 		TArray<FAttributeChange> AttributeChanges;
 		FPrimChangeFlags Flags;
-		FString OldPath;							// Empty if Flags.bDidRename is not set
-		FString OldIdentifier;						// Empty if Flags.bDIdChangeIdentifier is not set
+		FString OldPath;		  // Empty if Flags.bDidRename is not set
+		FString OldIdentifier;	  // Empty if Flags.bDIdChangeIdentifier is not set
 		TArray<TPair<FString, ESubLayerChangeType>> SubLayerChanges;
 	};
-	using FObjectChangeNotice = FSdfChangeListEntry; // Renamed in 5.3 as it is used for layer changes now too
+
+	using FObjectChangeNotice = FSdfChangeListEntry;	// Renamed in 5.3 as it is used for layer changes now too
 
 	using FSdfChangeList = TArray<TPair<UE::FSdfPath, FSdfChangeListEntry>>;
 	using FLayerToSdfChangeList = TArray<TPair<UE::FSdfLayerWeak, FSdfChangeList>>;
@@ -92,7 +93,7 @@ namespace UsdUtils
 	 * the key can be "/" to signify a stage change, or something like "/MyRoot/SomePrim" to indicate a particular prim change
 	 */
 	using FObjectChangesByPath = TMap<FString, TArray<FSdfChangeListEntry>>;
-}
+}	 // namespace UsdUtils
 
 /**
  * Registers to Usd Notices and emits events when the Usd Stage has changed
@@ -101,17 +102,17 @@ class UNREALUSDWRAPPER_API FUsdListener
 {
 public:
 	FUsdListener();
-	FUsdListener( const UE::FUsdStage& Stage );
+	FUsdListener(const UE::FUsdStage& Stage);
 
-	FUsdListener( const FUsdListener& Other ) = delete;
-	FUsdListener( FUsdListener&& Other ) = delete;
+	FUsdListener(const FUsdListener& Other) = delete;
+	FUsdListener(FUsdListener&& Other) = delete;
 
 	virtual ~FUsdListener();
 
-	FUsdListener& operator=( const FUsdListener& Other ) = delete;
-	FUsdListener& operator=( FUsdListener&& Other ) = delete;
+	FUsdListener& operator=(const FUsdListener& Other) = delete;
+	FUsdListener& operator=(FUsdListener&& Other) = delete;
 
-	void Register( const UE::FUsdStage& Stage );
+	void Register(const UE::FUsdStage& Stage);
 
 	// Increment/decrement the block counter
 	void Block();
@@ -142,20 +143,20 @@ public:
 	FOnObjectsChanged& GetOnObjectsChanged();
 
 private:
-	TUniquePtr< FUsdListenerImpl > Impl;
+	TUniquePtr<FUsdListenerImpl> Impl;
 };
 
 class UNREALUSDWRAPPER_API FScopedBlockNotices final
 {
 public:
-	explicit FScopedBlockNotices( FUsdListener& InListener );
+	explicit FScopedBlockNotices(FUsdListener& InListener);
 	~FScopedBlockNotices();
 
 	FScopedBlockNotices() = delete;
-	FScopedBlockNotices( const FScopedBlockNotices& ) = delete;
-	FScopedBlockNotices( FScopedBlockNotices&& ) = delete;
-	FScopedBlockNotices& operator=( const FScopedBlockNotices& ) = delete;
-	FScopedBlockNotices& operator=( FScopedBlockNotices&& ) = delete;
+	FScopedBlockNotices(const FScopedBlockNotices&) = delete;
+	FScopedBlockNotices(FScopedBlockNotices&&) = delete;
+	FScopedBlockNotices& operator=(const FScopedBlockNotices&) = delete;
+	FScopedBlockNotices& operator=(FScopedBlockNotices&&) = delete;
 
 private:
 	FUsdListener& Listener;
diff --git a/Engine/Plugins/Importers/USDImporter/Source/UnrealUSDWrapper/Public/USDMemory.h b/Engine/Plugins/Importers/USDImporter/Source/UnrealUSDWrapper/Public/USDMemory.h
--- a/Engine/Plugins/Importers/USDImporter/Source/UnrealUSDWrapper/Public/USDMemory.h
+++ b/Engine/Plugins/Importers/USDImporter/Source/UnrealUSDWrapper/Public/USDMemory.h
@@ -5,12 +5,12 @@
 #include "CoreMinimal.h"
 #include "HAL/LowLevelMemTracker.h"
 #include "Misc/Optional.h"
+#include "Misc/ScopeRWLock.h"
 #include "Modules/Boilerplate/ModuleBoilerplate.h"
 #include "Templates/SharedPointer.h"
-#include "Misc/ScopeRWLock.h"
 #include "Templates/UnrealTemplate.h"
 
-LLM_DECLARE_TAG_API( Usd, UNREALUSDWRAPPER_API );
+LLM_DECLARE_TAG_API(Usd, UNREALUSDWRAPPER_API);
 
 /**
  * This file is used to handle memory allocation issues when interacting with the USD SDK.
@@ -18,11 +18,14 @@ LLM_DECLARE_TAG_API( Usd, UNREALUSDWRAPPER_API );
  * Modules looking to use the USD SDK and the memory tools provided here need to use IMPLEMENT_MODULE_USD instead of the regular IMPLEMENT_MODULE.
  * The USD memory tools are only supported in non monolithic builds at this time, since it requires overriding new and delete per module.
  *
- * The USD SDK uses the shared C runtime allocator. This means that objects returned by the USD SDK might try to delete objects that were allocated through the CRT.
- * Since UE overrides new and delete per module, USD objects that try to call delete will end up freeing memory with the UE allocator but the malloc was made using the CRT, leading in a crash.
+ * The USD SDK uses the shared C runtime allocator. This means that objects returned by the USD SDK might try to delete objects that were allocated
+ *through the CRT. Since UE overrides new and delete per module, USD objects that try to call delete will end up freeing memory with the UE allocator
+ *but the malloc was made using the CRT, leading in a crash.
  *
- * To go around this problem, modules using the USD SDK need special operators for new and delete. Those operators have the ability to redirect the malloc or free calls to either the UE allocator or to the CRT allocator.
- * The choice of the allocator is made in the FUsdMemoryManager. FUsdMemoryManager manages a stack of active allocators per thread. Using ActivateAllocator and DeactivateAllocator, we can push and pop which allocator is active on the calling thread.
+ * To go around this problem, modules using the USD SDK need special operators for new and delete. Those operators have the ability to redirect the
+ *malloc or free calls to either the UE allocator or to the CRT allocator. The choice of the allocator is made in the FUsdMemoryManager.
+ *FUsdMemoryManager manages a stack of active allocators per thread. Using ActivateAllocator and DeactivateAllocator, we can push and pop which
+ *allocator is active on the calling thread.
  *
  * To simplify the workflow, TScopedAllocs is provided to make sure a certain block of code is bound to the right allocator.
  * FScopedUsdAllocs is a TScopedAllocs that activates the CRT allocator, while FScopedUnrealAllocs is the one that activates the UE allocator.
@@ -61,16 +64,16 @@ public:
 	static void Shutdown();
 
 	/** Pushes Allocator on the stack of active allocators */
-	static void ActivateAllocator( EUsdActiveAllocator Allocator );
+	static void ActivateAllocator(EUsdActiveAllocator Allocator);
 
 	/** Pops Allocator from the stack of active allocators */
-	static bool DeactivateAllocator( EUsdActiveAllocator Allocator );
+	static bool DeactivateAllocator(EUsdActiveAllocator Allocator);
 
 	/** Redirects the call to malloc to the currently active allocator. */
-	static void* Malloc( SIZE_T Count );
+	static void* Malloc(SIZE_T Count);
 
 	/** Redirects the call to free to the currently active allocator. */
-	static void Free( void* Original );
+	static void Free(void* Original);
 
 private:
 	static FActiveAllocatorsStack& GetActiveAllocatorsStackForThread();
@@ -78,31 +81,32 @@ private:
 	/** Returns if the current active allocator is EActiveAllocator::System */
 	static bool IsUsingSystemMalloc();
 
-	static TOptional< FTlsSlot > ActiveAllocatorsStackTLS;
+	static TOptional<FTlsSlot> ActiveAllocatorsStackTLS;
 
 	struct FThreadSafeSet
 	{
 	public:
-		void Add( void* Ptr )
+		void Add(void* Ptr)
 		{
-			FSubSet& SubSet = SubSets[GetTypeHash( Ptr ) % BucketCount];
-			FWriteScopeLock ScopeLock( SubSet.Lock );
-			SubSet.Set.Add( Ptr );
+			FSubSet& SubSet = SubSets[GetTypeHash(Ptr) % BucketCount];
+			FWriteScopeLock ScopeLock(SubSet.Lock);
+			SubSet.Set.Add(Ptr);
 		}
 
-		bool Remove( void* Ptr )
+		bool Remove(void* Ptr)
 		{
-			FSubSet& SubSet = SubSets[GetTypeHash( Ptr ) % BucketCount];
-			FWriteScopeLock ScopeLock( SubSet.Lock );
-			return SubSet.Set.Remove( Ptr ) > 0;
+			FSubSet& SubSet = SubSets[GetTypeHash(Ptr) % BucketCount];
+			FWriteScopeLock ScopeLock(SubSet.Lock);
+			return SubSet.Set.Remove(Ptr) > 0;
 		}
 
 	private:
 		struct FSubSet
 		{
 			FRWLock Lock;
-			TSet< void* > Set;
+			TSet<void*> Set;
 		};
+
 		static constexpr int32 BucketCount = 61; /* Prime number for better modulo distribution */
 		FSubSet SubSets[BucketCount];
 	};
@@ -115,133 +119,148 @@ private:
 /**
  * Activates an allocator on construction and deactivates it on destruction
  */
-template< EUsdActiveAllocator AllocatorType >
+template<EUsdActiveAllocator AllocatorType>
 class TScopedAllocs final
 {
 public:
 	TScopedAllocs()
 	{
-		FUsdMemoryManager::ActivateAllocator( AllocatorType );
+		FUsdMemoryManager::ActivateAllocator(AllocatorType);
 	}
 
 	~TScopedAllocs()
 	{
-		FUsdMemoryManager::DeactivateAllocator( AllocatorType );
+		FUsdMemoryManager::DeactivateAllocator(AllocatorType);
 	}
 
-	TScopedAllocs( const TScopedAllocs< AllocatorType >& Other ) = delete;
-	TScopedAllocs( TScopedAllocs< AllocatorType >&& Other ) = delete;
+	TScopedAllocs(const TScopedAllocs<AllocatorType>& Other) = delete;
+	TScopedAllocs(TScopedAllocs<AllocatorType>&& Other) = delete;
 
-	TScopedAllocs& operator=( const TScopedAllocs< AllocatorType >& Other ) = delete;
-	TScopedAllocs& operator=( TScopedAllocs< AllocatorType >&& Other ) = delete;
+	TScopedAllocs& operator=(const TScopedAllocs<AllocatorType>& Other) = delete;
+	TScopedAllocs& operator=(TScopedAllocs<AllocatorType>&& Other) = delete;
 };
 
-using FScopedUsdAllocs = TScopedAllocs< EUsdActiveAllocator::System >;
-using FScopedUnrealAllocs = TScopedAllocs< EUsdActiveAllocator::Unreal >;
+using FScopedUsdAllocs = TScopedAllocs<EUsdActiveAllocator::System>;
+using FScopedUnrealAllocs = TScopedAllocs<EUsdActiveAllocator::Unreal>;
 
 /**
  * Stores a USD object.
  * Ensures that its constructed, copied, moved and destroyed using the USD allocator.
  */
-template < typename UsdObjectType >
+template<typename UsdObjectType>
 class TUsdStore final
 {
 public:
 	TUsdStore()
 	{
-		StoredUsdObject = TOptional< UsdObjectType >{}; // Force init with current allocator
+		StoredUsdObject = TOptional<UsdObjectType>{};	 // Force init with current allocator
 
 		{
 			FScopedUsdAllocs UsdAllocs;
-			StoredUsdObject.Emplace(); // Create UsdObjectType with USD allocator
+			StoredUsdObject.Emplace();	  // Create UsdObjectType with USD allocator
 		}
 	}
 
-	TUsdStore( const TUsdStore< UsdObjectType >& Other )
+	TUsdStore(const TUsdStore<UsdObjectType>& Other)
 	{
 		FScopedUsdAllocs UsdAllocs;
-		StoredUsdObject.Emplace( Other.StoredUsdObject.GetValue() );
+		StoredUsdObject.Emplace(Other.StoredUsdObject.GetValue());
 	}
 
-	TUsdStore( TUsdStore< UsdObjectType >&& Other )
+	TUsdStore(TUsdStore<UsdObjectType>&& Other)
 	{
 		FScopedUsdAllocs UsdAllocs;
-		StoredUsdObject = MoveTemp( Other.StoredUsdObject );
+		StoredUsdObject = MoveTemp(Other.StoredUsdObject);
 	}
 
-	TUsdStore< UsdObjectType >& operator=( const TUsdStore< UsdObjectType >& Other )
+	TUsdStore<UsdObjectType>& operator=(const TUsdStore<UsdObjectType>& Other)
 	{
 		FScopedUsdAllocs UsdAllocs;
-		StoredUsdObject.Emplace( Other.StoredUsdObject.GetValue() );
+		StoredUsdObject.Emplace(Other.StoredUsdObject.GetValue());
 		return *this;
 	}
 
-	TUsdStore< UsdObjectType >& operator=( TUsdStore< UsdObjectType >&& Other )
+	TUsdStore<UsdObjectType>& operator=(TUsdStore<UsdObjectType>&& Other)
 	{
 		FScopedUsdAllocs UsdAllocs;
-		StoredUsdObject = MoveTemp( Other.StoredUsdObject );
+		StoredUsdObject = MoveTemp(Other.StoredUsdObject);
 		return *this;
 	}
 
-	TUsdStore( const UsdObjectType& UsdObject )
+	TUsdStore(const UsdObjectType& UsdObject)
 	{
 		FScopedUsdAllocs UsdAllocs;
 		StoredUsdObject = UsdObject;
 	}
 
-	TUsdStore( UsdObjectType&& UsdObject )
+	TUsdStore(UsdObjectType&& UsdObject)
 	{
 		FScopedUsdAllocs UsdAllocs;
-		StoredUsdObject = MoveTemp( UsdObject );
+		StoredUsdObject = MoveTemp(UsdObject);
 	}
 
-	TUsdStore< UsdObjectType >& operator=( const UsdObjectType& UsdObject )
+	TUsdStore<UsdObjectType>& operator=(const UsdObjectType& UsdObject)
 	{
 		FScopedUsdAllocs UsdAllocs;
 		StoredUsdObject = UsdObject;
 		return *this;
 	}
 
-	TUsdStore< UsdObjectType >& operator=( UsdObjectType&& UsdObject )
+	TUsdStore<UsdObjectType>& operator=(UsdObjectType&& UsdObject)
 	{
 		FScopedUsdAllocs UsdAllocs;
-		StoredUsdObject = MoveTemp( UsdObject );
+		StoredUsdObject = MoveTemp(UsdObject);
 		return *this;
 	}
 
 	~TUsdStore()
 	{
 		FScopedUsdAllocs UsdAllocs;
-		StoredUsdObject.Reset(); // Destroy UsdObjectType with USD allocator
+		StoredUsdObject.Reset();	// Destroy UsdObjectType with USD allocator
+	}
+
+	UsdObjectType& Get()
+	{
+		return StoredUsdObject.GetValue();
+	}
+
+	const UsdObjectType& Get() const
+	{
+		return StoredUsdObject.GetValue();
 	}
 
-	UsdObjectType& Get() { return StoredUsdObject.GetValue(); }
-	const UsdObjectType& Get() const { return StoredUsdObject.GetValue(); }
+	UsdObjectType& operator*()
+	{
+		return Get();
+	}
 
-	UsdObjectType& operator*() { return Get(); }
-	const UsdObjectType& operator*() const { return Get(); }
+	const UsdObjectType& operator*() const
+	{
+		return Get();
+	}
 
 private:
-	TOptional< UsdObjectType > StoredUsdObject;
+	TOptional<UsdObjectType> StoredUsdObject;
 };
 
-template< typename UsdObjectType, typename... ArgTypes >
-TUsdStore< UsdObjectType > MakeUsdStore( ArgTypes&&... Args )
+template<typename UsdObjectType, typename... ArgTypes>
+TUsdStore<UsdObjectType> MakeUsdStore(ArgTypes&&... Args)
 {
 	FScopedUsdAllocs UsdAllocs;
-	return TUsdStore< UsdObjectType >( UsdObjectType( Forward< ArgTypes >( Args )... ) );
+	return TUsdStore<UsdObjectType>(UsdObjectType(Forward<ArgTypes>(Args)...));
 }
 
 // MakeShared version that makes sure that the SharedPointer allocs are made with the Unreal allocator
-template< typename ObjectType, typename... ArgTypes >
-TSharedRef< ObjectType > MakeSharedUnreal( ArgTypes&&... Args )
+template<typename ObjectType, typename... ArgTypes>
+TSharedRef<ObjectType> MakeSharedUnreal(ArgTypes&&... Args)
 {
 	FScopedUnrealAllocs UnrealAllocs;
-	return MakeShared< ObjectType >( Forward< ArgTypes >( Args )... );
+	return MakeShared<ObjectType>(Forward<ArgTypes>(Args)...);
 }
 
 // See comment on UnrealUSDWrapper.Build.cs to understand why we disable these for monolithic builds (everything will still work,
 // they're just unnecessary)
+// clang-format off
 #if !FORCE_ANSI_ALLOCATOR && !IS_MONOLITHIC
 	#define REPLACEMENT_OPERATOR_NEW_AND_DELETE_USD \
 		OPERATOR_NEW_MSVC_PRAGMA void* operator new  ( size_t Size                        ) OPERATOR_NEW_THROW_SPEC      { return FUsdMemoryManager::Malloc( Size ); } \
@@ -278,4 +297,5 @@ TSharedRef< ObjectType > MakeSharedUnreal( ArgTypes&&... Args )
 		REPLACEMENT_OPERATOR_NEW_AND_DELETE_USD
 #else
 	#define IMPLEMENT_MODULE_USD( ModuleImplClass, ModuleName ) IMPLEMENT_MODULE( ModuleImplClass, ModuleName )
-#endif
\ No newline at end of file
+#endif
+// clang-format on
diff --git a/Engine/Plugins/Importers/USDImporter/Source/UnrealUSDWrapper/Public/UnrealUSDWrapper.h b/Engine/Plugins/Importers/USDImporter/Source/UnrealUSDWrapper/Public/UnrealUSDWrapper.h
--- a/Engine/Plugins/Importers/USDImporter/Source/UnrealUSDWrapper/Public/UnrealUSDWrapper.h
+++ b/Engine/Plugins/Importers/USDImporter/Source/UnrealUSDWrapper/Public/UnrealUSDWrapper.h
@@ -12,15 +12,15 @@
 
 #include "UsdWrappers/ForwardDeclarations.h"
 
+#include <memory>
 #include <string>
 #include <vector>
-#include <memory>
 
 #if USE_USD_SDK
 #include "USDIncludesStart.h"
-	#include "pxr/pxr.h"
+#include "pxr/pxr.h"
 #include "USDIncludesEnd.h"
-#endif // #if USE_USD_SDK
+#endif	  // #if USE_USD_SDK
 
 #include "UnrealUSDWrapper.generated.h"
 
@@ -39,9 +39,10 @@ PXR_NAMESPACE_OPEN_SCOPE
 	class UsdStage;
 	class UsdStageCache;
 
-	template< typename T > class TfRefPtr;
+	template<typename T>
+	class TfRefPtr;
 PXR_NAMESPACE_CLOSE_SCOPE
-#endif // #if USE_USD_SDK
+#endif	  // #if USE_USD_SDK
 
 class IUsdPrim;
 class FUsdDiagnosticDelegate;
@@ -89,24 +90,24 @@ enum class EUsdPurpose : int32
 };
 ENUM_CLASS_FLAGS(EUsdPurpose);
 
-UENUM( meta = ( Bitflags, UseEnumValuesAsMaskValuesInEditor = "true" ) )
+UENUM(meta = (Bitflags, UseEnumValuesAsMaskValuesInEditor = "true"))
 enum class EUsdDefaultKind : int32
 {
-	None = 0 UMETA( Hidden ),
+	None = 0 UMETA(Hidden),
 	Model = 1,
 	Component = 2,
 	Group = 4,
 	Assembly = 8,
 	Subcomponent = 16
 };
-ENUM_CLASS_FLAGS( EUsdDefaultKind );
+ENUM_CLASS_FLAGS(EUsdDefaultKind);
 
 /** Corresponds to pxr::UsdLoadPolicy, refer to the USD SDK documentation */
 UENUM()
 enum class EUsdLoadPolicy : uint8
 {
-	UsdLoadWithDescendants,    // Load a prim plus all its descendants.
-	UsdLoadWithoutDescendants  // Load a prim by itself with no descendants.
+	UsdLoadWithDescendants,		 // Load a prim plus all its descendants.
+	UsdLoadWithoutDescendants	 // Load a prim by itself with no descendants.
 };
 
 UENUM()
@@ -173,6 +174,7 @@ enum class EUsdListPosition : uint8
  * these types from USD.
  */
 USTRUCT(BlueprintType)
+
 struct FMatrix2D
 {
 	GENERATED_BODY()
@@ -189,6 +191,7 @@ struct FMatrix2D
  * these types from USD.
  */
 USTRUCT(BlueprintType)
+
 struct FMatrix3D
 {
 	GENERATED_BODY()
@@ -212,7 +215,7 @@ class UnrealUSDWrapper
 public:
 #if USE_USD_SDK
 	UNREALUSDWRAPPER_API static double GetDefaultTimeCode();
-#endif  // #if USE_USD_SDK
+#endif	  // #if USE_USD_SDK
 
 	/**
 	 * Registers all USD plug-ins discovered at PathToPlugInfo.
@@ -245,7 +248,8 @@ public:
 
 	/**
 	 * Opens a USD stage from a file on disk or existing layers, with a population mask or not.
-	 * @param Identifier - Path to a file that the USD SDK can open (or the identifier of a root layer), which will become the root layer of the new stage
+	 * @param Identifier - Path to a file that the USD SDK can open (or the identifier of a root layer), which will become the root layer of the new
+	 *stage
 	 * @param RootLayer - Existing root layer to use for the new stage, instead of reading it from disk
 	 * @param SessionLayer - Existing session layer to use for the new stage, instead of creating a new one
 	 * @param InitialLoadSet - How to handle USD payloads when opening this stage
@@ -287,7 +291,7 @@ public:
 	);
 
 	/** Creates a new USD root layer file, opens it as a new stage and returns that stage */
-	UNREALUSDWRAPPER_API static UE::FUsdStage NewStage( const TCHAR* FilePath );
+	UNREALUSDWRAPPER_API static UE::FUsdStage NewStage(const TCHAR* FilePath);
 
 	/** Creates a new memory USD root layer, opens it as a new stage and returns that stage */
 	UNREALUSDWRAPPER_API static UE::FUsdStage NewStage();
@@ -300,10 +304,10 @@ public:
 	UNREALUSDWRAPPER_API static UE::FUsdStage GetClipboardStage();
 
 	/** Returns all the stages that are currently opened in the USD utils stage cache, shared between C++ and Python */
-	UNREALUSDWRAPPER_API static TArray< UE::FUsdStage > GetAllStagesFromCache();
+	UNREALUSDWRAPPER_API static TArray<UE::FUsdStage> GetAllStagesFromCache();
 
 	/** Removes the stage from the stage cache. See UsdStageCache::Erase. */
-	UNREALUSDWRAPPER_API static void EraseStageFromCache( const UE::FUsdStage& Stage );
+	UNREALUSDWRAPPER_API static void EraseStageFromCache(const UE::FUsdStage& Stage);
 
 	/**
 	 * Set the directories that will be used as the default search path by USD's default resolver during asset resolution.
@@ -313,7 +317,7 @@ public:
 	 * Note that the default search path must be set before the first invocation of USD's resolver system, so this function
 	 * must be called before that to have any effect.
 	 */
-	UNREALUSDWRAPPER_API static void SetDefaultResolverDefaultSearchPath( const TArray<FDirectoryPath>& SearchPath );
+	UNREALUSDWRAPPER_API static void SetDefaultResolverDefaultSearchPath(const TArray<FDirectoryPath>& SearchPath);
 
 	/** Starts listening to error/warning/log messages emitted by USD */
 	UNREALUSDWRAPPER_API static void SetupDiagnosticDelegate();
@@ -343,8 +347,8 @@ public:
 	static UNREALUSDWRAPPER_API bool ClearKind(const pxr::UsdPrim& Prim);
 
 	static UNREALUSDWRAPPER_API pxr::GfMatrix4d GetLocalTransform(const pxr::UsdPrim& Prim);
-	static UNREALUSDWRAPPER_API pxr::GfMatrix4d GetLocalToWorldTransform(const pxr::UsdPrim& Prim );
-	static UNREALUSDWRAPPER_API pxr::GfMatrix4d GetLocalToWorldTransform(const pxr::UsdPrim& Prim, double Time );
+	static UNREALUSDWRAPPER_API pxr::GfMatrix4d GetLocalToWorldTransform(const pxr::UsdPrim& Prim);
+	static UNREALUSDWRAPPER_API pxr::GfMatrix4d GetLocalToWorldTransform(const pxr::UsdPrim& Prim, double Time);
 	static UNREALUSDWRAPPER_API pxr::GfMatrix4d GetLocalToWorldTransform(const pxr::UsdPrim& Prim, double Time, const pxr::SdfPath& AbsoluteRootPath);
 	static UNREALUSDWRAPPER_API bool HasTransform(const pxr::UsdPrim& Prim);
 
@@ -352,7 +356,7 @@ public:
 
 	static UNREALUSDWRAPPER_API EUsdGeomOrientation GetGeometryOrientation(const pxr::UsdGeomMesh& Mesh);
 	static UNREALUSDWRAPPER_API EUsdGeomOrientation GetGeometryOrientation(const pxr::UsdGeomMesh& Mesh, double Time);
-#endif // #if USE_USD_SDK
+#endif	  // #if USE_USD_SDK
 };
 
 namespace UnrealIdentifiers
@@ -436,7 +440,7 @@ namespace UnrealIdentifiers
 	extern UNREALUSDWRAPPER_API const pxr::TfToken UnrealAssetType;
 	extern UNREALUSDWRAPPER_API const pxr::TfToken UnrealExportTime;
 	extern UNREALUSDWRAPPER_API const pxr::TfToken UnrealEngineVersion;
-#endif // #if USE_USD_SDK
+#endif	  // #if USE_USD_SDK
 
 	extern UNREALUSDWRAPPER_API const TCHAR* LayerSavedComment;
 	extern UNREALUSDWRAPPER_API const TCHAR* TwoSidedMaterialSuffix;
@@ -484,7 +488,7 @@ namespace UnrealIdentifiers
 
 	// Material purpose tokens that we convert from USD, if available
 	extern UNREALUSDWRAPPER_API FString MaterialAllPurpose;
-	extern UNREALUSDWRAPPER_API FString MaterialAllPurposeText; // Text to show on UI for "allPurpose", as its value is actually the empty string
+	extern UNREALUSDWRAPPER_API FString MaterialAllPurposeText;	   // Text to show on UI for "allPurpose", as its value is actually the empty string
 	extern UNREALUSDWRAPPER_API FString MaterialPreviewPurpose;
 	extern UNREALUSDWRAPPER_API FString MaterialFullPurpose;
 
@@ -498,11 +502,11 @@ namespace UnrealIdentifiers
 
 	// The character used to separate property namespaces (usually just ':')
 	extern UNREALUSDWRAPPER_API FString UsdNamespaceDelimiter;
-}
+}	 // namespace UnrealIdentifiers
 
 struct UNREALUSDWRAPPER_API FUsdDelegates
 {
-	DECLARE_MULTICAST_DELEGATE_OneParam( FUsdImportDelegate, FString /* FilePath */);
+	DECLARE_MULTICAST_DELEGATE_OneParam(FUsdImportDelegate, FString /* FilePath */);
 	static FUsdImportDelegate OnPreUsdImport;
 	static FUsdImportDelegate OnPostUsdImport;
-};
\ No newline at end of file
+};
diff --git a/Engine/Plugins/Importers/USDImporter/Source/UnrealUSDWrapper/Public/UsdWrappers/ForwardDeclarations.h b/Engine/Plugins/Importers/USDImporter/Source/UnrealUSDWrapper/Public/UsdWrappers/ForwardDeclarations.h
--- a/Engine/Plugins/Importers/USDImporter/Source/UnrealUSDWrapper/Public/UsdWrappers/ForwardDeclarations.h
+++ b/Engine/Plugins/Importers/USDImporter/Source/UnrealUSDWrapper/Public/UsdWrappers/ForwardDeclarations.h
@@ -5,35 +5,39 @@
 #if USE_USD_SDK
 
 #include "USDIncludesStart.h"
-	#include "pxr/pxr.h"
+#include "pxr/pxr.h"
 #include "USDIncludesEnd.h"
 
 PXR_NAMESPACE_OPEN_SCOPE
-class UsdStage;
-class SdfLayer;
+	class UsdStage;
+	class SdfLayer;
 
-template< typename T > class TfRefPtr;
-template< typename T > class TfWeakPtr;
+	template<typename T>
+	class TfRefPtr;
+	template<typename T>
+	class TfWeakPtr;
 
-using UsdStageRefPtr = TfRefPtr< UsdStage >;
-using UsdStageWeakPtr = TfWeakPtr< UsdStage >;
+	using UsdStageRefPtr = TfRefPtr<UsdStage>;
+	using UsdStageWeakPtr = TfWeakPtr<UsdStage>;
 
-using SdfLayerRefPtr = TfRefPtr< SdfLayer >;
-using SdfLayerWeakPtr = TfWeakPtr< SdfLayer >;
+	using SdfLayerRefPtr = TfRefPtr<SdfLayer>;
+	using SdfLayerWeakPtr = TfWeakPtr<SdfLayer>;
 PXR_NAMESPACE_CLOSE_SCOPE
 
 namespace UE
 {
-	template<typename PtrType> class FUsdStageBase;
-	using FUsdStage = FUsdStageBase< pxr::UsdStageRefPtr >;
-	using FUsdStageWeak = FUsdStageBase< pxr::UsdStageWeakPtr >;
+	template<typename PtrType>
+	class FUsdStageBase;
+	using FUsdStage = FUsdStageBase<pxr::UsdStageRefPtr>;
+	using FUsdStageWeak = FUsdStageBase<pxr::UsdStageWeakPtr>;
 
-	template<typename PtrType> class FSdfLayerBase;
-	using FSdfLayer = FSdfLayerBase< pxr::SdfLayerRefPtr >;
-	using FSdfLayerWeak = FSdfLayerBase< pxr::SdfLayerWeakPtr >;
+	template<typename PtrType>
+	class FSdfLayerBase;
+	using FSdfLayer = FSdfLayerBase<pxr::SdfLayerRefPtr>;
+	using FSdfLayerWeak = FSdfLayerBase<pxr::SdfLayerWeakPtr>;
 }
 
-#else // When the USD SDK is not available
+#else	  // When the USD SDK is not available
 namespace UE
 {
 	// We use these types so that we can keep the same FUsdStage, etc. type aliases
@@ -42,31 +46,59 @@ namespace UE
 	// We use two to prevent multiple symbol definitions on the wrapper copy
 	// constructor/assignment operators.
 	class FDummyWeakPtrType;
+
 	class FDummyRefPtrType
 	{
 	public:
 		FDummyRefPtrType(){};
-		FDummyRefPtrType( const FDummyWeakPtrType& ) {};
-		explicit operator bool() const { return false; }
-		bool operator==( const FDummyRefPtrType& ) const { return false; }
-		bool operator!=( const FDummyRefPtrType& ) const { return false; }
+		FDummyRefPtrType(const FDummyWeakPtrType&){};
+
+		explicit operator bool() const
+		{
+			return false;
+		}
+
+		bool operator==(const FDummyRefPtrType&) const
+		{
+			return false;
+		}
+
+		bool operator!=(const FDummyRefPtrType&) const
+		{
+			return false;
+		}
 	};
+
 	class FDummyWeakPtrType
 	{
 	public:
 		FDummyWeakPtrType(){};
-		FDummyWeakPtrType( const FDummyRefPtrType& ) {};
-		explicit operator bool() const { return false; }
-		bool operator==( const FDummyWeakPtrType& ) const { return false; }
-		bool operator!=( const FDummyWeakPtrType& ) const { return false; }
+		FDummyWeakPtrType(const FDummyRefPtrType&){};
+
+		explicit operator bool() const
+		{
+			return false;
+		}
+
+		bool operator==(const FDummyWeakPtrType&) const
+		{
+			return false;
+		}
+
+		bool operator!=(const FDummyWeakPtrType&) const
+		{
+			return false;
+		}
 	};
 
-	template<typename PtrType> class FUsdStageBase;
-	using FUsdStage = FUsdStageBase< FDummyRefPtrType >;
-	using FUsdStageWeak = FUsdStageBase< FDummyWeakPtrType >;
+	template<typename PtrType>
+	class FUsdStageBase;
+	using FUsdStage = FUsdStageBase<FDummyRefPtrType>;
+	using FUsdStageWeak = FUsdStageBase<FDummyWeakPtrType>;
 
-	template<typename PtrType> class FSdfLayerBase;
-	using FSdfLayer = FSdfLayerBase< FDummyRefPtrType >;
-	using FSdfLayerWeak = FSdfLayerBase< FDummyWeakPtrType >;
-}
-#endif // #if USE_USD_SDK
\ No newline at end of file
+	template<typename PtrType>
+	class FSdfLayerBase;
+	using FSdfLayer = FSdfLayerBase<FDummyRefPtrType>;
+	using FSdfLayerWeak = FSdfLayerBase<FDummyWeakPtrType>;
+}	 // namespace UE
+#endif	  // #if USE_USD_SDK
diff --git a/Engine/Plugins/Importers/USDImporter/Source/UnrealUSDWrapper/Public/UsdWrappers/SdfAttributeSpec.h b/Engine/Plugins/Importers/USDImporter/Source/UnrealUSDWrapper/Public/UsdWrappers/SdfAttributeSpec.h
--- a/Engine/Plugins/Importers/USDImporter/Source/UnrealUSDWrapper/Public/UsdWrappers/SdfAttributeSpec.h
+++ b/Engine/Plugins/Importers/USDImporter/Source/UnrealUSDWrapper/Public/UsdWrappers/SdfAttributeSpec.h
@@ -7,16 +7,17 @@
 #if USE_USD_SDK
 
 #include "USDIncludesStart.h"
-	#include "pxr/pxr.h"
+#include "pxr/pxr.h"
 #include "USDIncludesEnd.h"
 
 PXR_NAMESPACE_OPEN_SCOPE
 	class SdfAttributeSpec;
-	template <class T> class SdfHandle;
-	using SdfAttributeSpecHandle = SdfHandle< SdfAttributeSpec >;
+	template<class T>
+	class SdfHandle;
+	using SdfAttributeSpecHandle = SdfHandle<SdfAttributeSpec>;
 PXR_NAMESPACE_CLOSE_SCOPE
 
-#endif // #if USE_USD_SDK
+#endif	  // #if USE_USD_SDK
 
 namespace UE
 {
@@ -35,38 +36,38 @@ namespace UE
 	public:
 		FSdfAttributeSpec();
 
-		FSdfAttributeSpec( const FSdfAttributeSpec& Other );
-		FSdfAttributeSpec( FSdfAttributeSpec&& Other );
+		FSdfAttributeSpec(const FSdfAttributeSpec& Other);
+		FSdfAttributeSpec(FSdfAttributeSpec&& Other);
 		~FSdfAttributeSpec();
 
-		FSdfAttributeSpec& operator=( const FSdfAttributeSpec& Other );
-		FSdfAttributeSpec& operator=( FSdfAttributeSpec&& Other );
+		FSdfAttributeSpec& operator=(const FSdfAttributeSpec& Other);
+		FSdfAttributeSpec& operator=(FSdfAttributeSpec&& Other);
 
-		bool operator==( const FSdfAttributeSpec& Other ) const;
-		bool operator!=( const FSdfAttributeSpec& Other ) const;
+		bool operator==(const FSdfAttributeSpec& Other) const;
+		bool operator!=(const FSdfAttributeSpec& Other) const;
 
 		explicit operator bool() const;
 
-	// Auto conversion from/to pxr::SdfAttributeSpecHandle
+		// Auto conversion from/to pxr::SdfAttributeSpecHandle
 	public:
 #if USE_USD_SDK
-		explicit FSdfAttributeSpec( const pxr::SdfAttributeSpecHandle& InSdfAttributeSpecHandle );
-		explicit FSdfAttributeSpec( pxr::SdfAttributeSpecHandle&& InSdfAttributeSpecHandle );
-		FSdfAttributeSpec& operator=( const pxr::SdfAttributeSpecHandle& InSdfAttributeSpecHandle );
-		FSdfAttributeSpec& operator=( pxr::SdfAttributeSpecHandle&& InSdfAttributeSpecHandle );
+		explicit FSdfAttributeSpec(const pxr::SdfAttributeSpecHandle& InSdfAttributeSpecHandle);
+		explicit FSdfAttributeSpec(pxr::SdfAttributeSpecHandle&& InSdfAttributeSpecHandle);
+		FSdfAttributeSpec& operator=(const pxr::SdfAttributeSpecHandle& InSdfAttributeSpecHandle);
+		FSdfAttributeSpec& operator=(pxr::SdfAttributeSpecHandle&& InSdfAttributeSpecHandle);
 
 		operator pxr::SdfAttributeSpecHandle&();
 		operator const pxr::SdfAttributeSpecHandle&() const;
-#endif // #if USE_USD_SDK
+#endif	  // #if USE_USD_SDK
 
-	// Wrapped pxr::SdfAttributeSpec functions, refer to the USD SDK documentation
+		  // Wrapped pxr::SdfAttributeSpec functions, refer to the USD SDK documentation
 	public:
 		UE::FVtValue GetDefaultValue() const;
-		bool SetDefaultValue( const UE::FVtValue& Value );
+		bool SetDefaultValue(const UE::FVtValue& Value);
 
 		FString GetName() const;
 
 	private:
-		TUniquePtr< Internal::FSdfAttributeSpecImpl > Impl;
+		TUniquePtr<Internal::FSdfAttributeSpecImpl> Impl;
 	};
-}
\ No newline at end of file
+}	 // namespace UE
diff --git a/Engine/Plugins/Importers/USDImporter/Source/UnrealUSDWrapper/Public/UsdWrappers/SdfLayer.h b/Engine/Plugins/Importers/USDImporter/Source/UnrealUSDWrapper/Public/UsdWrappers/SdfLayer.h
--- a/Engine/Plugins/Importers/USDImporter/Source/UnrealUSDWrapper/Public/UsdWrappers/SdfLayer.h
+++ b/Engine/Plugins/Importers/USDImporter/Source/UnrealUSDWrapper/Public/UsdWrappers/SdfLayer.h
@@ -2,9 +2,9 @@
 
 #pragma once
 
-#include "CoreMinimal.h"
 #include "Containers/Map.h"
 #include "Containers/UnrealString.h"
+#include "CoreMinimal.h"
 
 #include "UsdWrappers/ForwardDeclarations.h"
 #include "UsdWrappers/SdfPath.h"
@@ -15,15 +15,17 @@ namespace UE
 
 	namespace Internal
 	{
-		template < typename PtrType > class FSdfLayerImpl;
+		template<typename PtrType>
+		class FSdfLayerImpl;
 	}
 
 	struct UNREALUSDWRAPPER_API FSdfLayerOffset
 	{
 		FSdfLayerOffset() = default;
-		FSdfLayerOffset( double InOffset, double InScale )
-			: Offset( InOffset )
-			, Scale( InScale )
+
+		FSdfLayerOffset(double InOffset, double InScale)
+			: Offset(InOffset)
+			, Scale(InScale)
 		{
 		}
 
@@ -33,10 +35,7 @@ namespace UE
 
 	struct UNREALUSDWRAPPER_API FSdfPayload
 	{
-		FSdfPayload(
-			const FString& InAssetPath = {},
-			const FSdfPath& InPrimPath = {},
-			const FSdfLayerOffset& InLayerOffset = {})
+		FSdfPayload(const FString& InAssetPath = {}, const FSdfPath& InPrimPath = {}, const FSdfLayerOffset& InLayerOffset = {})
 			: AssetPath(InAssetPath)
 			, PrimPath(InPrimPath)
 			, LayerOffset(InLayerOffset)
@@ -57,13 +56,10 @@ namespace UE
 	// references is added to the FSdfReference wrapper in the future.
 	struct UNREALUSDWRAPPER_API FSdfReference
 	{
-		FSdfReference(
-				const FString& InAssetPath = {},
-				const FSdfPath& InPrimPath = {},
-				const FSdfLayerOffset& InLayerOffset = {})
-			: AssetPath( InAssetPath )
-			, PrimPath( InPrimPath )
-			, LayerOffset ( InLayerOffset )
+		FSdfReference(const FString& InAssetPath = {}, const FSdfPath& InPrimPath = {}, const FSdfLayerOffset& InLayerOffset = {})
+			: AssetPath(InAssetPath)
+			, PrimPath(InPrimPath)
+			, LayerOffset(InLayerOffset)
 		{
 		}
 
@@ -76,46 +72,46 @@ namespace UE
 	 * Minimal pxr::SdfLayer pointer wrapper for Unreal that can be used from no-rtti modules.
 	 * Use the aliases FSdfLayer and FSdfLayerWeak instead (defined on ForwardDeclarations.h)
 	 */
-	template< typename PtrType >
+	template<typename PtrType>
 	class UNREALUSDWRAPPER_API FSdfLayerBase
 	{
 	public:
 		FSdfLayerBase();
 
-		FSdfLayerBase( const FSdfLayer& Other );
-		FSdfLayerBase( FSdfLayer&& Other );
-		FSdfLayerBase( const FSdfLayerWeak& Other );
-		FSdfLayerBase( FSdfLayerWeak&& Other );
+		FSdfLayerBase(const FSdfLayer& Other);
+		FSdfLayerBase(FSdfLayer&& Other);
+		FSdfLayerBase(const FSdfLayerWeak& Other);
+		FSdfLayerBase(FSdfLayerWeak&& Other);
 
 		~FSdfLayerBase();
 
-		FSdfLayerBase& operator=( const FSdfLayer& Other );
-		FSdfLayerBase& operator=( FSdfLayer&& Other );
-		FSdfLayerBase& operator=( const FSdfLayerWeak& Other );
-		FSdfLayerBase& operator=( FSdfLayerWeak&& Other );
+		FSdfLayerBase& operator=(const FSdfLayer& Other);
+		FSdfLayerBase& operator=(FSdfLayer&& Other);
+		FSdfLayerBase& operator=(const FSdfLayerWeak& Other);
+		FSdfLayerBase& operator=(FSdfLayerWeak&& Other);
 
 		template<typename OtherPtrType>
-		bool operator==( const FSdfLayerBase<OtherPtrType>& Other ) const;
+		bool operator==(const FSdfLayerBase<OtherPtrType>& Other) const;
 		template<typename OtherPtrType>
-		bool operator!=( const FSdfLayerBase<OtherPtrType>& Other ) const;
+		bool operator!=(const FSdfLayerBase<OtherPtrType>& Other) const;
 
 		explicit operator bool() const;
 
-		friend UNREALUSDWRAPPER_API uint32 GetTypeHash( const FSdfLayerWeak& Layer );
+		friend UNREALUSDWRAPPER_API uint32 GetTypeHash(const FSdfLayerWeak& Layer);
 
 		// Auto conversion from/to PtrType. We use concrete pointer types here
 		// because we should also be able to convert between them
 	public:
 #if USE_USD_SDK
-		explicit FSdfLayerBase( const pxr::SdfLayerRefPtr& InSdfLayer );
-		explicit FSdfLayerBase( pxr::SdfLayerRefPtr&& InSdfLayer );
-		explicit FSdfLayerBase( const pxr::SdfLayerWeakPtr& InSdfLayer );
-		explicit FSdfLayerBase( pxr::SdfLayerWeakPtr&& InSdfLayer );
+		explicit FSdfLayerBase(const pxr::SdfLayerRefPtr& InSdfLayer);
+		explicit FSdfLayerBase(pxr::SdfLayerRefPtr&& InSdfLayer);
+		explicit FSdfLayerBase(const pxr::SdfLayerWeakPtr& InSdfLayer);
+		explicit FSdfLayerBase(pxr::SdfLayerWeakPtr&& InSdfLayer);
 
-		FSdfLayerBase& operator=( const pxr::SdfLayerRefPtr& InSdfLayer );
-		FSdfLayerBase& operator=( pxr::SdfLayerRefPtr&& InSdfLayer );
-		FSdfLayerBase& operator=( const pxr::SdfLayerWeakPtr& InSdfLayer );
-		FSdfLayerBase& operator=( pxr::SdfLayerWeakPtr&& InSdfLayer );
+		FSdfLayerBase& operator=(const pxr::SdfLayerRefPtr& InSdfLayer);
+		FSdfLayerBase& operator=(pxr::SdfLayerRefPtr&& InSdfLayer);
+		FSdfLayerBase& operator=(const pxr::SdfLayerWeakPtr& InSdfLayer);
+		FSdfLayerBase& operator=(pxr::SdfLayerWeakPtr&& InSdfLayer);
 
 		// We can provide reference cast operators for the type we do have, but
 		// need to settle on providing copy cast operators for the other types
@@ -124,32 +120,26 @@ namespace UE
 
 		operator pxr::SdfLayerRefPtr() const;
 		operator pxr::SdfLayerWeakPtr() const;
-#endif // #if USE_USD_SDK
+#endif	  // #if USE_USD_SDK
 
-	// Wrapped pxr::SdfLayer functions, refer to the USD SDK documentation
+		  // Wrapped pxr::SdfLayer functions, refer to the USD SDK documentation
 	public:
 		FString GetComment() const;
-		void SetComment( const TCHAR* Comment ) const;
+		void SetComment(const TCHAR* Comment) const;
 
-		void TransferContent( const FSdfLayer& SourceLayer );
+		void TransferContent(const FSdfLayer& SourceLayer);
 
 		static TSet<FSdfLayerWeak> GetLoadedLayers();
-		static FSdfLayer FindOrOpen(
-			const TCHAR* Identifier,
-			const TMap< FString, FString >& FileFormatArguments = {} );
-		static FSdfLayer CreateNew(
-			const TCHAR* Identifier,
-			const TMap< FString, FString >& FileFormatArguments = {} );
+		static FSdfLayer FindOrOpen(const TCHAR* Identifier, const TMap<FString, FString>& FileFormatArguments = {});
+		static FSdfLayer CreateNew(const TCHAR* Identifier, const TMap<FString, FString>& FileFormatArguments = {});
 
-		TMap< FString, FString > GetFileFormatArguments() const;
+		TMap<FString, FString> GetFileFormatArguments() const;
 
-		bool Save( bool bForce = false ) const;
+		bool Save(bool bForce = false) const;
 
 		TSet<FString> GetCompositionAssetDependencies() const;
 
-		bool UpdateCompositionAssetDependency(
-			const TCHAR* OldAssetPath,
-			const TCHAR* NewAssetPath = nullptr);
+		bool UpdateCompositionAssetDependency(const TCHAR* OldAssetPath, const TCHAR* NewAssetPath = nullptr);
 
 		FString GetRealPath() const;
 		FString GetIdentifier() const;
@@ -161,28 +151,25 @@ namespace UE
 		bool IsEmpty() const;
 		bool IsAnonymous() const;
 
-		bool Export(
-			const TCHAR* Filename,
-			const FString& Comment = {},
-			const TMap< FString, FString >& FileFormatArguments = {} ) const;
+		bool Export(const TCHAR* Filename, const FString& Comment = {}, const TMap<FString, FString>& FileFormatArguments = {}) const;
 
 		void Clear();
 
 		bool HasStartTimeCode() const;
 		double GetStartTimeCode() const;
-		void SetStartTimeCode( double TimeCode );
+		void SetStartTimeCode(double TimeCode);
 
 		bool HasEndTimeCode() const;
 		double GetEndTimeCode() const;
-		void SetEndTimeCode( double TimeCode );
+		void SetEndTimeCode(double TimeCode);
 
 		bool HasTimeCodesPerSecond() const;
 		double GetTimeCodesPerSecond() const;
-		void SetTimeCodesPerSecond( double TimeCodesPerSecond );
+		void SetTimeCodesPerSecond(double TimeCodesPerSecond);
 
 		bool HasFramesPerSecond() const;
 		double GetFramesPerSecond() const;
-		void SetFramesPerSecond( double FramesPerSecond );
+		void SetFramesPerSecond(double FramesPerSecond);
 
 		TArray<FString> GetSubLayerPaths() const;
 		void SetSubLayerPaths(const TArray<FString>& NewPaths);
@@ -193,13 +180,13 @@ namespace UE
 		FSdfLayerOffset GetSubLayerOffset(int32 Index) const;
 		void SetSubLayerOffset(const FSdfLayerOffset& Offset, int32 Index);
 
-		bool HasSpec( const FSdfPath& Path ) const;
+		bool HasSpec(const FSdfPath& Path) const;
 
 		FSdfPrimSpec GetPseudoRoot() const;
-		FSdfPrimSpec GetPrimAtPath( const FSdfPath& Path ) const;
+		FSdfPrimSpec GetPrimAtPath(const FSdfPath& Path) const;
 
-		TSet< double > ListTimeSamplesForPath( const FSdfPath& Path ) const;
-		void EraseTimeSample( const FSdfPath& Path, double Time );
+		TSet<double> ListTimeSamplesForPath(const FSdfPath& Path) const;
+		void EraseTimeSample(const FSdfPath& Path, double Time);
 
 		bool IsMuted() const;
 		void SetMuted(bool bMuted);
@@ -208,7 +195,7 @@ namespace UE
 		friend FSdfLayer;
 		friend FSdfLayerWeak;
 
-		TUniquePtr< Internal::FSdfLayerImpl<PtrType> > Impl;
+		TUniquePtr<Internal::FSdfLayerImpl<PtrType>> Impl;
 	};
 
 	/**
@@ -217,6 +204,6 @@ namespace UE
 	class UNREALUSDWRAPPER_API FSdfLayerUtils
 	{
 	public:
-		static FString SdfComputeAssetPathRelativeToLayer( const FSdfLayer& Anchor, const TCHAR* AssetPath );
+		static FString SdfComputeAssetPathRelativeToLayer(const FSdfLayer& Anchor, const TCHAR* AssetPath);
 	};
-}
\ No newline at end of file
+}	 // namespace UE
diff --git a/Engine/Plugins/Importers/USDImporter/Source/UnrealUSDWrapper/Public/UsdWrappers/SdfPath.h b/Engine/Plugins/Importers/USDImporter/Source/UnrealUSDWrapper/Public/UsdWrappers/SdfPath.h
--- a/Engine/Plugins/Importers/USDImporter/Source/UnrealUSDWrapper/Public/UsdWrappers/SdfPath.h
+++ b/Engine/Plugins/Importers/USDImporter/Source/UnrealUSDWrapper/Public/UsdWrappers/SdfPath.h
@@ -5,16 +5,14 @@
 #include "CoreMinimal.h"
 
 #if USE_USD_SDK
-
 #include "USDIncludesStart.h"
-	#include "pxr/pxr.h"
+#include "pxr/pxr.h"
 #include "USDIncludesEnd.h"
 
 PXR_NAMESPACE_OPEN_SCOPE
 	class SdfPath;
 PXR_NAMESPACE_CLOSE_SCOPE
-
-#endif // #if USE_USD_SDK
+#endif	  // #if USE_USD_SDK
 
 namespace UE
 {
@@ -30,37 +28,37 @@ namespace UE
 	{
 	public:
 		FSdfPath();
-		explicit FSdfPath( const TCHAR* InPath );
+		explicit FSdfPath(const TCHAR* InPath);
 
-		FSdfPath( const FSdfPath& Other );
-		FSdfPath( FSdfPath&& Other );
+		FSdfPath(const FSdfPath& Other);
+		FSdfPath(FSdfPath&& Other);
 
-		FSdfPath& operator=( const FSdfPath& Other );
-		FSdfPath& operator=( FSdfPath&& Other );
+		FSdfPath& operator=(const FSdfPath& Other);
+		FSdfPath& operator=(FSdfPath&& Other);
 
 		~FSdfPath();
 
-		bool operator==( const FSdfPath& Other ) const;
-		bool operator!=( const FSdfPath& Other ) const;
+		bool operator==(const FSdfPath& Other) const;
+		bool operator!=(const FSdfPath& Other) const;
 		bool operator<(const FSdfPath& Other) const;
 		bool operator<=(const FSdfPath& Other) const;
 
-		friend UNREALUSDWRAPPER_API uint32 GetTypeHash( const UE::FSdfPath& Path );
-		friend UNREALUSDWRAPPER_API FArchive& operator<<( FArchive& Ar, FSdfPath& Path );
+		friend UNREALUSDWRAPPER_API uint32 GetTypeHash(const UE::FSdfPath& Path);
+		friend UNREALUSDWRAPPER_API FArchive& operator<<(FArchive& Ar, FSdfPath& Path);
 
-	// Auto conversion from/to pxr::SdfPath
+		// Auto conversion from/to pxr::SdfPath
 	public:
 #if USE_USD_SDK
-		explicit FSdfPath( const pxr::SdfPath& InSdfPath );
-		explicit FSdfPath( pxr::SdfPath&& InSdfPath );
-		FSdfPath& operator=(  const pxr::SdfPath& InSdfPath );
-		FSdfPath& operator=( pxr::SdfPath&& InSdfPath );
+		explicit FSdfPath(const pxr::SdfPath& InSdfPath);
+		explicit FSdfPath(pxr::SdfPath&& InSdfPath);
+		FSdfPath& operator=(const pxr::SdfPath& InSdfPath);
+		FSdfPath& operator=(pxr::SdfPath&& InSdfPath);
 
 		operator pxr::SdfPath&();
 		operator const pxr::SdfPath&() const;
-#endif // #if USE_USD_SDK
+#endif	  // #if USE_USD_SDK
 
-	// Wrapped pxr::SdfPath functions, refer to the USD SDK documentation
+		  // Wrapped pxr::SdfPath functions, refer to the USD SDK documentation
 	public:
 		static const FSdfPath& AbsoluteRootPath();
 
@@ -72,12 +70,12 @@ namespace UE
 		FString GetElementString() const;
 		FSdfPath GetAbsoluteRootOrPrimPath() const;
 
-		FSdfPath ReplaceName( const TCHAR* NewLeafName ) const;
+		FSdfPath ReplaceName(const TCHAR* NewLeafName) const;
 
 		FSdfPath GetParentPath() const;
 		FSdfPath AppendPath(const UE::FSdfPath& NewRelativeSuffix) const;
-		FSdfPath AppendChild( const TCHAR* ChildName ) const;
-		FSdfPath AppendProperty( FName PropertyName ) const;
+		FSdfPath AppendChild(const TCHAR* ChildName) const;
+		FSdfPath AppendProperty(FName PropertyName) const;
 
 		FSdfPath StripAllVariantSelections() const;
 
@@ -91,6 +89,6 @@ namespace UE
 		FSdfPath MakeRelativePath(const FSdfPath& Anchor) const;
 
 	private:
-		TUniquePtr< Internal::FSdfPathImpl > Impl;
+		TUniquePtr<Internal::FSdfPathImpl> Impl;
 	};
-}
\ No newline at end of file
+}	 // namespace UE
diff --git a/Engine/Plugins/Importers/USDImporter/Source/UnrealUSDWrapper/Public/UsdWrappers/SdfPrimSpec.h b/Engine/Plugins/Importers/USDImporter/Source/UnrealUSDWrapper/Public/UsdWrappers/SdfPrimSpec.h
--- a/Engine/Plugins/Importers/USDImporter/Source/UnrealUSDWrapper/Public/UsdWrappers/SdfPrimSpec.h
+++ b/Engine/Plugins/Importers/USDImporter/Source/UnrealUSDWrapper/Public/UsdWrappers/SdfPrimSpec.h
@@ -8,18 +8,17 @@
 #include "UsdWrappers/ForwardDeclarations.h"
 
 #if USE_USD_SDK
-
 #include "USDIncludesStart.h"
-	#include "pxr/pxr.h"
+#include "pxr/pxr.h"
 #include "USDIncludesEnd.h"
 
 PXR_NAMESPACE_OPEN_SCOPE
 	class SdfPrimSpec;
-	template <class T> class SdfHandle;
-	using SdfPrimSpecHandle = SdfHandle< SdfPrimSpec >;
+	template<class T>
+	class SdfHandle;
+	using SdfPrimSpecHandle = SdfHandle<SdfPrimSpec>;
 PXR_NAMESPACE_CLOSE_SCOPE
-
-#endif // #if USE_USD_SDK
+#endif	  // #if USE_USD_SDK
 
 namespace UE
 {
@@ -58,31 +57,31 @@ namespace UE
 	public:
 		FSdfPrimSpec();
 
-		FSdfPrimSpec( const FSdfPrimSpec& Other );
-		FSdfPrimSpec( FSdfPrimSpec&& Other );
+		FSdfPrimSpec(const FSdfPrimSpec& Other);
+		FSdfPrimSpec(FSdfPrimSpec&& Other);
 		~FSdfPrimSpec();
 
-		FSdfPrimSpec& operator=( const FSdfPrimSpec& Other );
-		FSdfPrimSpec& operator=( FSdfPrimSpec&& Other );
+		FSdfPrimSpec& operator=(const FSdfPrimSpec& Other);
+		FSdfPrimSpec& operator=(FSdfPrimSpec&& Other);
 
-		bool operator==( const FSdfPrimSpec& Other ) const;
-		bool operator!=( const FSdfPrimSpec& Other ) const;
+		bool operator==(const FSdfPrimSpec& Other) const;
+		bool operator!=(const FSdfPrimSpec& Other) const;
 
 		explicit operator bool() const;
 
-	// Auto conversion from/to pxr::SdfPrimSpecHandle
+		// Auto conversion from/to pxr::SdfPrimSpecHandle
 	public:
 #if USE_USD_SDK
-		explicit FSdfPrimSpec( const pxr::SdfPrimSpecHandle& InSdfPrimSpec );
-		explicit FSdfPrimSpec( pxr::SdfPrimSpecHandle&& InSdfPrimSpec );
-		FSdfPrimSpec& operator=( const pxr::SdfPrimSpecHandle& InSdfPrimSpec );
-		FSdfPrimSpec& operator=( pxr::SdfPrimSpecHandle&& InSdfPrimSpec );
+		explicit FSdfPrimSpec(const pxr::SdfPrimSpecHandle& InSdfPrimSpec);
+		explicit FSdfPrimSpec(pxr::SdfPrimSpecHandle&& InSdfPrimSpec);
+		FSdfPrimSpec& operator=(const pxr::SdfPrimSpecHandle& InSdfPrimSpec);
+		FSdfPrimSpec& operator=(pxr::SdfPrimSpecHandle&& InSdfPrimSpec);
 
 		operator pxr::SdfPrimSpecHandle&();
 		operator const pxr::SdfPrimSpecHandle&() const;
-#endif // #if USE_USD_SDK
+#endif	  // #if USE_USD_SDK
 
-	// Wrapped pxr::SdfPrimSpecHandle functions, refer to the USD SDK documentation
+		  // Wrapped pxr::SdfPrimSpecHandle functions, refer to the USD SDK documentation
 	public:
 		ESdfSpecType GetSpecType() const;
 		FSdfLayerWeak GetLayer() const;
@@ -91,12 +90,12 @@ namespace UE
 		FSdfPrimSpec GetRealNameParent() const;
 		bool RemoveNameChild(const FSdfPrimSpec& Child);
 
-		UE::FSdfAttributeSpec GetAttributeAtPath( const UE::FSdfPath& Path ) const;
+		UE::FSdfAttributeSpec GetAttributeAtPath(const UE::FSdfPath& Path) const;
 
 		FName GetTypeName() const;
 		FName GetName() const;
 
 	private:
-		TUniquePtr< Internal::FSdfPrimSpecImpl > Impl;
+		TUniquePtr<Internal::FSdfPrimSpecImpl> Impl;
 	};
-}
\ No newline at end of file
+}	 // namespace UE
diff --git a/Engine/Plugins/Importers/USDImporter/Source/UnrealUSDWrapper/Public/UsdWrappers/UsdAttribute.h b/Engine/Plugins/Importers/USDImporter/Source/UnrealUSDWrapper/Public/UsdWrappers/UsdAttribute.h
--- a/Engine/Plugins/Importers/USDImporter/Source/UnrealUSDWrapper/Public/UsdWrappers/UsdAttribute.h
+++ b/Engine/Plugins/Importers/USDImporter/Source/UnrealUSDWrapper/Public/UsdWrappers/UsdAttribute.h
@@ -2,21 +2,19 @@
 
 #pragma once
 
-#include "Templates/UniquePtr.h"
 #include "Misc/Optional.h"
+#include "Templates/UniquePtr.h"
 
 #if USE_USD_SDK
-
 #include "USDIncludesStart.h"
-	#include "pxr/pxr.h"
+#include "pxr/pxr.h"
 #include "USDIncludesEnd.h"
 
 PXR_NAMESPACE_OPEN_SCOPE
 	class UsdAttribute;
 	class UsdProperty;
 PXR_NAMESPACE_CLOSE_SCOPE
-
-#endif // #if USE_USD_SDK
+#endif	  // #if USE_USD_SDK
 
 namespace UE
 {
@@ -37,25 +35,25 @@ namespace UE
 	public:
 		FUsdAttribute();
 
-		FUsdAttribute( const FUsdAttribute& Other );
-		FUsdAttribute( FUsdAttribute&& Other );
+		FUsdAttribute(const FUsdAttribute& Other);
+		FUsdAttribute(FUsdAttribute&& Other);
 		~FUsdAttribute();
 
-		FUsdAttribute& operator=( const FUsdAttribute& Other );
-		FUsdAttribute& operator=( FUsdAttribute&& Other );
+		FUsdAttribute& operator=(const FUsdAttribute& Other);
+		FUsdAttribute& operator=(FUsdAttribute&& Other);
 
-		bool operator==( const FUsdAttribute& Other ) const;
-		bool operator!=( const FUsdAttribute& Other ) const;
+		bool operator==(const FUsdAttribute& Other) const;
+		bool operator!=(const FUsdAttribute& Other) const;
 
 		explicit operator bool() const;
 
-	// Auto conversion from/to pxr::UsdAttribute
+		// Auto conversion from/to pxr::UsdAttribute
 	public:
 #if USE_USD_SDK
-		explicit FUsdAttribute( const pxr::UsdAttribute& InUsdAttribute);
-		explicit FUsdAttribute( pxr::UsdAttribute&& InUsdAttribute );
-		FUsdAttribute& operator=( const pxr::UsdAttribute& InUsdAttribute );
-		FUsdAttribute& operator=( pxr::UsdAttribute&& InUsdAttribute );
+		explicit FUsdAttribute(const pxr::UsdAttribute& InUsdAttribute);
+		explicit FUsdAttribute(pxr::UsdAttribute&& InUsdAttribute);
+		FUsdAttribute& operator=(const pxr::UsdAttribute& InUsdAttribute);
+		FUsdAttribute& operator=(pxr::UsdAttribute&& InUsdAttribute);
 
 		operator pxr::UsdAttribute&();
 		operator const pxr::UsdAttribute&() const;
@@ -65,22 +63,22 @@ namespace UE
 		// pxr::UsdProperty directly
 		operator pxr::UsdProperty&();
 		operator const pxr::UsdProperty&() const;
-#endif // #if USE_USD_SDK
+#endif	  // #if USE_USD_SDK
 
-	// Wrapped pxr::UsdObject functions, refer to the USD SDK documentation
+		  // Wrapped pxr::UsdObject functions, refer to the USD SDK documentation
 	public:
-		bool GetMetadata( const TCHAR* Key, UE::FVtValue& Value ) const;
-		bool HasMetadata( const TCHAR* Key ) const;
-		bool SetMetadata( const TCHAR* Key, const UE::FVtValue& Value ) const;
-		bool ClearMetadata( const TCHAR* Key ) const;
+		bool GetMetadata(const TCHAR* Key, UE::FVtValue& Value) const;
+		bool HasMetadata(const TCHAR* Key) const;
+		bool SetMetadata(const TCHAR* Key, const UE::FVtValue& Value) const;
+		bool ClearMetadata(const TCHAR* Key) const;
 
-	// Wrapped pxr::UsdAttribute functions, refer to the USD SDK documentation
+		// Wrapped pxr::UsdAttribute functions, refer to the USD SDK documentation
 	public:
 		FName GetName() const;
 		FName GetBaseName() const;
 		FName GetTypeName() const;
 
-		bool GetTimeSamples( TArray<double>& Times ) const;
+		bool GetTimeSamples(TArray<double>& Times) const;
 		size_t GetNumTimeSamples() const;
 
 		bool HasValue() const;
@@ -89,18 +87,18 @@ namespace UE
 
 		bool ValueMightBeTimeVarying() const;
 
-		bool Get( UE::FVtValue& Value, TOptional<double> Time = {} ) const;
-		bool Set( const UE::FVtValue& Value, TOptional<double> Time = {} ) const;
+		bool Get(UE::FVtValue& Value, TOptional<double> Time = {}) const;
+		bool Set(const UE::FVtValue& Value, TOptional<double> Time = {}) const;
 
 		bool Clear() const;
-		bool ClearAtTime( double Time ) const;
+		bool ClearAtTime(double Time) const;
 
-		static bool GetUnionedTimeSamples( const TArray<UE::FUsdAttribute>& Attrs, TArray<double>& OutTimes );
+		static bool GetUnionedTimeSamples(const TArray<UE::FUsdAttribute>& Attrs, TArray<double>& OutTimes);
 
 		FSdfPath GetPath() const;
 		FUsdPrim GetPrim() const;
 
 	private:
-		TUniquePtr< Internal::FUsdAttributeImpl > Impl;
+		TUniquePtr<Internal::FUsdAttributeImpl> Impl;
 	};
-}
\ No newline at end of file
+}	 // namespace UE
