USD: Reformat the UnrealUSDWrapper module.

diff --git a/Engine/Plugins/Importers/USDImporter/Source/UnrealUSDWrapper/Public/UsdWrappers/UsdGeomXformable.h b/Engine/Plugins/Importers/USDImporter/Source/UnrealUSDWrapper/Public/UsdWrappers/UsdGeomXformable.h
--- a/Engine/Plugins/Importers/USDImporter/Source/UnrealUSDWrapper/Public/UsdWrappers/UsdGeomXformable.h
+++ b/Engine/Plugins/Importers/USDImporter/Source/UnrealUSDWrapper/Public/UsdWrappers/UsdGeomXformable.h
@@ -7,7 +7,6 @@
 #include "UsdTyped.h"
 
 #if USE_USD_SDK
-
 #include "USDIncludesStart.h"
 #include "pxr/pxr.h"
 #include "USDIncludesEnd.h"
@@ -16,8 +15,7 @@ PXR_NAMESPACE_OPEN_SCOPE
 	class UsdPrim;
 	class UsdGeomXformable;
 PXR_NAMESPACE_CLOSE_SCOPE
-
-#endif // #if USE_USD_SDK
+#endif	  // #if USE_USD_SDK
 
 namespace UE
 {
@@ -35,13 +33,13 @@ namespace UE
 	{
 	public:
 		FUsdGeomXformable();
-		FUsdGeomXformable( const FUsdGeomXformable& Other );
-		FUsdGeomXformable( FUsdGeomXformable&& Other );
+		FUsdGeomXformable(const FUsdGeomXformable& Other);
+		FUsdGeomXformable(FUsdGeomXformable&& Other);
 
-		explicit FUsdGeomXformable( const FUsdPrim& Prim );
+		explicit FUsdGeomXformable(const FUsdPrim& Prim);
 
-		FUsdGeomXformable& operator=( const FUsdGeomXformable& Other );
-		FUsdGeomXformable& operator=( FUsdGeomXformable&& Other );
+		FUsdGeomXformable& operator=(const FUsdGeomXformable& Other);
+		FUsdGeomXformable& operator=(FUsdGeomXformable&& Other);
 
 		~FUsdGeomXformable();
 
@@ -50,25 +48,25 @@ namespace UE
 		// Auto conversion from/to pxr::UsdGeomXformable
 	public:
 #if USE_USD_SDK
-		explicit FUsdGeomXformable( const pxr::UsdGeomXformable& InUsdGeomXformable );
-		explicit FUsdGeomXformable( pxr::UsdGeomXformable&& InUsdGeomXformable );
-		explicit FUsdGeomXformable( const pxr::UsdPrim& Prim );
+		explicit FUsdGeomXformable(const pxr::UsdGeomXformable& InUsdGeomXformable);
+		explicit FUsdGeomXformable(pxr::UsdGeomXformable&& InUsdGeomXformable);
+		explicit FUsdGeomXformable(const pxr::UsdPrim& Prim);
 
 		operator pxr::UsdGeomXformable&();
 		operator const pxr::UsdGeomXformable&() const;
-#endif // #if USE_USD_SDK
+#endif	  // #if USE_USD_SDK
 
-		// Wrapped pxr::UsdGeomXformable functions, refer to the USD SDK documentation
+		  // Wrapped pxr::UsdGeomXformable functions, refer to the USD SDK documentation
 	public:
 		bool GetResetXformStack() const;
 
 		bool TransformMightBeTimeVarying() const;
-		bool GetTimeSamples( TArray< double >* Times) const;
+		bool GetTimeSamples(TArray<double>* Times) const;
 
 		FUsdAttribute GetXformOpOrderAttr() const;
 		bool ClearXformOpOrder() const;
 
 	private:
-		TUniquePtr< Internal::FUsdGeomXformableImpl > Impl;
+		TUniquePtr<Internal::FUsdGeomXformableImpl> Impl;
 	};
-}
\ No newline at end of file
+}	 // namespace UE
diff --git a/Engine/Plugins/Importers/USDImporter/Source/UnrealUSDWrapper/Public/UsdWrappers/UsdPayloads.h b/Engine/Plugins/Importers/USDImporter/Source/UnrealUSDWrapper/Public/UsdWrappers/UsdPayloads.h
--- a/Engine/Plugins/Importers/USDImporter/Source/UnrealUSDWrapper/Public/UsdWrappers/UsdPayloads.h
+++ b/Engine/Plugins/Importers/USDImporter/Source/UnrealUSDWrapper/Public/UsdWrappers/UsdPayloads.h
@@ -10,17 +10,15 @@
 #include "UsdWrappers/SdfLayer.h"
 
 #if USE_USD_SDK
-
 #include "USDIncludesStart.h"
-	#include "pxr/pxr.h"
+#include "pxr/pxr.h"
 #include "USDIncludesEnd.h"
 
 PXR_NAMESPACE_OPEN_SCOPE
 	class UsdPayloads;
 	class UsdPrim;
 PXR_NAMESPACE_CLOSE_SCOPE
-
-#endif // #if USE_USD_SDK
+#endif	  // #if USE_USD_SDK
 
 namespace UE
 {
@@ -64,29 +62,30 @@ namespace UE
 		// FUsdPrim.
 		// Use FUsdPrim::GetPayloads() to construct an instance.
 		explicit FUsdPayloads(const pxr::UsdPrim& InPrim);
-#endif // #if USE_USD_SDK
+#endif	  // #if USE_USD_SDK
 
-		// Wrapped pxr::UsdPayloads functions, refer to the USD SDK documentation
+		  // Wrapped pxr::UsdPayloads functions, refer to the USD SDK documentation
 	public:
-		bool AddPayload(
-			const FSdfPayload& Payload,
-			EUsdListPosition Position = EUsdListPosition::BackOfPrependList);
+		bool AddPayload(const FSdfPayload& Payload, EUsdListPosition Position = EUsdListPosition::BackOfPrependList);
 
 		bool AddPayload(
 			const FString& Identifier,
 			const FSdfPath& PrimPath,
 			const FSdfLayerOffset& LayerOffset = {},
-			EUsdListPosition Position = EUsdListPosition::BackOfPrependList);
+			EUsdListPosition Position = EUsdListPosition::BackOfPrependList
+		);
 
 		bool AddPayload(
 			const FString& Identifier,
 			const FSdfLayerOffset& LayerOffset = {},
-			EUsdListPosition Position = EUsdListPosition::BackOfPrependList);
+			EUsdListPosition Position = EUsdListPosition::BackOfPrependList
+		);
 
 		bool AddInternalPayload(
 			const FSdfPath& PrimPath,
 			const FSdfLayerOffset& LayerOffset = {},
-			EUsdListPosition Position = EUsdListPosition::BackOfPrependList);
+			EUsdListPosition Position = EUsdListPosition::BackOfPrependList
+		);
 
 		bool RemovePayload(const FSdfPayload& Payload);
 
@@ -101,4 +100,4 @@ namespace UE
 
 		friend class FUsdPrim;
 	};
-}
+}	 // namespace UE
diff --git a/Engine/Plugins/Importers/USDImporter/Source/UnrealUSDWrapper/Public/UsdWrappers/UsdPrim.h b/Engine/Plugins/Importers/USDImporter/Source/UnrealUSDWrapper/Public/UsdWrappers/UsdPrim.h
--- a/Engine/Plugins/Importers/USDImporter/Source/UnrealUSDWrapper/Public/UsdWrappers/UsdPrim.h
+++ b/Engine/Plugins/Importers/USDImporter/Source/UnrealUSDWrapper/Public/UsdWrappers/UsdPrim.h
@@ -11,16 +11,14 @@
 #include "UsdWrappers/ForwardDeclarations.h"
 
 #if USE_USD_SDK
-
 #include "USDIncludesStart.h"
-	#include "pxr/pxr.h"
+#include "pxr/pxr.h"
 #include "USDIncludesEnd.h"
 
 PXR_NAMESPACE_OPEN_SCOPE
 	class UsdPrim;
 PXR_NAMESPACE_CLOSE_SCOPE
-
-#endif // #if USE_USD_SDK
+#endif	  // #if USE_USD_SDK
 
 namespace UE
 {
@@ -39,10 +37,10 @@ namespace UE
 	/** Corresponds to pxr::SdfSpecifier, refer to the USD SDK documentation */
 	enum class ESdfSpecifier
 	{
-		Def,	// Defines a concrete prim
-		Over,	// Overrides an existing prim
-		Class,	// Defines an abstract prim
-		Num		// The number of specifiers
+		Def,	  // Defines a concrete prim
+		Over,	  // Overrides an existing prim
+		Class,	  // Defines an abstract prim
+		Num		  // The number of specifiers
 	};
 
 	/**
@@ -53,36 +51,36 @@ namespace UE
 	public:
 		FUsdPrim();
 
-		FUsdPrim( const FUsdPrim& Other );
-		FUsdPrim( FUsdPrim&& Other );
+		FUsdPrim(const FUsdPrim& Other);
+		FUsdPrim(FUsdPrim&& Other);
 		~FUsdPrim();
 
-		FUsdPrim& operator=( const FUsdPrim& Other );
-		FUsdPrim& operator=( FUsdPrim&& Other );
+		FUsdPrim& operator=(const FUsdPrim& Other);
+		FUsdPrim& operator=(FUsdPrim&& Other);
 
-		bool operator==( const FUsdPrim& Other ) const;
-		bool operator!=( const FUsdPrim& Other ) const;
+		bool operator==(const FUsdPrim& Other) const;
+		bool operator!=(const FUsdPrim& Other) const;
 
 		explicit operator bool() const;
 
-	// Auto conversion from/to pxr::UsdPrim
+		// Auto conversion from/to pxr::UsdPrim
 	public:
 #if USE_USD_SDK
-		explicit FUsdPrim( const pxr::UsdPrim& InUsdPrim );
-		explicit FUsdPrim( pxr::UsdPrim&& InUsdPrim );
-		FUsdPrim& operator=( const pxr::UsdPrim& InUsdPrim );
-		FUsdPrim& operator=( pxr::UsdPrim&& InUsdPrim );
+		explicit FUsdPrim(const pxr::UsdPrim& InUsdPrim);
+		explicit FUsdPrim(pxr::UsdPrim&& InUsdPrim);
+		FUsdPrim& operator=(const pxr::UsdPrim& InUsdPrim);
+		FUsdPrim& operator=(pxr::UsdPrim&& InUsdPrim);
 
 		operator pxr::UsdPrim&();
 		operator const pxr::UsdPrim&() const;
-#endif // #if USE_USD_SDK
+#endif	  // #if USE_USD_SDK
 
-	// Wrapped pxr::UsdPrim functions, refer to the USD SDK documentation
+		  // Wrapped pxr::UsdPrim functions, refer to the USD SDK documentation
 	public:
-		bool SetSpecifier( ESdfSpecifier Specifier );
+		bool SetSpecifier(ESdfSpecifier Specifier);
 
 		bool IsActive() const;
-		bool SetActive( bool bActive );
+		bool SetActive(bool bActive);
 
 		bool IsValid() const;
 		bool IsPseudoRoot() const;
@@ -91,21 +89,24 @@ namespace UE
 
 		TArray<FName> GetAppliedSchemas() const;
 
-		bool IsA( FName SchemaIdentifier ) const;
-		bool HasAPI( FName SchemaIdentifier ) const;
-		bool HasAPI( FName SchemaIdentifier, FName InstanceName ) const;
+		bool IsA(FName SchemaIdentifier) const;
+		bool HasAPI(FName SchemaIdentifier) const;
+		bool HasAPI(FName SchemaIdentifier, FName InstanceName) const;
 
-		UE_DEPRECATED( 5.3, "Please use the one FName overload for single apply API schemas, or the two FName overload for multiple apply API schemas." )
-		bool HasAPI( FName SchemaType, TOptional<FName> InstanceName ) const;
+		UE_DEPRECATED(
+			5.3,
+			"Please use the one FName overload for single apply API schemas, or the two FName overload for multiple apply API schemas."
+		)
+		bool HasAPI(FName SchemaType, TOptional<FName> InstanceName) const;
 
-		bool CanApplyAPI( FName SchemaIdentifier, FString* OutWhyNot = nullptr ) const;
-		bool CanApplyAPI( FName SchemaIdentifier, FName InstanceName, FString* OutWhyNot = nullptr ) const;
+		bool CanApplyAPI(FName SchemaIdentifier, FString* OutWhyNot = nullptr) const;
+		bool CanApplyAPI(FName SchemaIdentifier, FName InstanceName, FString* OutWhyNot = nullptr) const;
 
-		bool ApplyAPI( FName SchemaIdentifier ) const;
-		bool ApplyAPI( FName SchemaIdentifier, FName InstanceName ) const;
+		bool ApplyAPI(FName SchemaIdentifier) const;
+		bool ApplyAPI(FName SchemaIdentifier, FName InstanceName) const;
 
-		bool RemoveAPI( FName SchemaIdentifier ) const;
-		bool RemoveAPI( FName SchemaIdentifier, FName InstanceName ) const;
+		bool RemoveAPI(FName SchemaIdentifier) const;
+		bool RemoveAPI(FName SchemaIdentifier, FName InstanceName) const;
 
 		const FSdfPath GetPrimPath() const;
 		FUsdStage GetStage() const;
@@ -119,8 +120,8 @@ namespace UE
 
 		FUsdPrim GetParent() const;
 
-		TArray< FUsdPrim > GetChildren() const;
-		TArray< FUsdPrim > GetFilteredChildren( bool bTraverseInstanceProxies ) const;
+		TArray<FUsdPrim> GetChildren() const;
+		TArray<FUsdPrim> GetFilteredChildren(bool bTraverseInstanceProxies) const;
 
 		FUsdVariantSets GetVariantSets() const;
 		FUsdVariantSet GetVariantSet(const FString& VariantSetName) const;
@@ -131,16 +132,16 @@ namespace UE
 		UE_DEPRECATED(5.3, "Please use HasAuthoredPayloads() instead.")
 		bool HasPayload() const;
 		bool IsLoaded() const;
-		void Load( EUsdLoadPolicy Policy = EUsdLoadPolicy::UsdLoadWithDescendants );
+		void Load(EUsdLoadPolicy Policy = EUsdLoadPolicy::UsdLoadWithDescendants);
 		void Unload();
 
 		FUsdReferences GetReferences() const;
 		bool HasAuthoredReferences() const;
 
-		bool RemoveProperty( FName PropName ) const;
+		bool RemoveProperty(FName PropName) const;
 
-		FUsdAttribute CreateAttribute( const TCHAR* AttrName, FName TypeName ) const;
-		TArray< FUsdAttribute > GetAttributes() const;
+		FUsdAttribute CreateAttribute(const TCHAR* AttrName, FName TypeName) const;
+		TArray<FUsdAttribute> GetAttributes() const;
 		FUsdAttribute GetAttribute(const TCHAR* AttrName) const;
 		bool HasAttribute(const TCHAR* AttrName) const;
 
@@ -159,6 +160,6 @@ namespace UE
 		static bool IsPathInPrototype(const FSdfPath& Path);
 
 	private:
-		TUniquePtr< Internal::FUsdPrimImpl > Impl;
+		TUniquePtr<Internal::FUsdPrimImpl> Impl;
 	};
-}
\ No newline at end of file
+}	 // namespace UE
diff --git a/Engine/Plugins/Importers/USDImporter/Source/UnrealUSDWrapper/Public/UsdWrappers/UsdReferences.h b/Engine/Plugins/Importers/USDImporter/Source/UnrealUSDWrapper/Public/UsdWrappers/UsdReferences.h
--- a/Engine/Plugins/Importers/USDImporter/Source/UnrealUSDWrapper/Public/UsdWrappers/UsdReferences.h
+++ b/Engine/Plugins/Importers/USDImporter/Source/UnrealUSDWrapper/Public/UsdWrappers/UsdReferences.h
@@ -10,17 +10,15 @@
 #include "UsdWrappers/SdfLayer.h"
 
 #if USE_USD_SDK
-
 #include "USDIncludesStart.h"
-	#include "pxr/pxr.h"
+#include "pxr/pxr.h"
 #include "USDIncludesEnd.h"
 
 PXR_NAMESPACE_OPEN_SCOPE
 	class UsdPrim;
 	class UsdReferences;
 PXR_NAMESPACE_CLOSE_SCOPE
-
-#endif // #if USE_USD_SDK
+#endif	  // #if USE_USD_SDK
 
 namespace UE
 {
@@ -64,29 +62,30 @@ namespace UE
 		// FUsdPrim.
 		// Use FUsdPrim::GetReferences() to construct an instance.
 		explicit FUsdReferences(const pxr::UsdPrim& InPrim);
-#endif // #if USE_USD_SDK
+#endif	  // #if USE_USD_SDK
 
-		// Wrapped pxr::UsdReferences functions, refer to the USD SDK documentation
+		  // Wrapped pxr::UsdReferences functions, refer to the USD SDK documentation
 	public:
-		bool AddReference(
-			const FSdfReference& Reference,
-			EUsdListPosition Position = EUsdListPosition::BackOfPrependList);
+		bool AddReference(const FSdfReference& Reference, EUsdListPosition Position = EUsdListPosition::BackOfPrependList);
 
 		bool AddReference(
 			const FString& Identifier,
 			const FSdfPath& PrimPath,
 			const FSdfLayerOffset& LayerOffset = {},
-			EUsdListPosition Position = EUsdListPosition::BackOfPrependList);
+			EUsdListPosition Position = EUsdListPosition::BackOfPrependList
+		);
 
 		bool AddReference(
 			const FString& Identifier,
 			const FSdfLayerOffset& LayerOffset = {},
-			EUsdListPosition Position = EUsdListPosition::BackOfPrependList);
+			EUsdListPosition Position = EUsdListPosition::BackOfPrependList
+		);
 
 		bool AddInternalReference(
 			const FSdfPath& PrimPath,
 			const FSdfLayerOffset& LayerOffset = {},
-			EUsdListPosition Position = EUsdListPosition::BackOfPrependList);
+			EUsdListPosition Position = EUsdListPosition::BackOfPrependList
+		);
 
 		bool RemoveReference(const FSdfReference& Reference);
 
@@ -101,4 +100,4 @@ namespace UE
 
 		friend class FUsdPrim;
 	};
-}
+}	 // namespace UE
diff --git a/Engine/Plugins/Importers/USDImporter/Source/UnrealUSDWrapper/Public/UsdWrappers/UsdStage.h b/Engine/Plugins/Importers/USDImporter/Source/UnrealUSDWrapper/Public/UsdWrappers/UsdStage.h
--- a/Engine/Plugins/Importers/USDImporter/Source/UnrealUSDWrapper/Public/UsdWrappers/UsdStage.h
+++ b/Engine/Plugins/Importers/USDImporter/Source/UnrealUSDWrapper/Public/UsdWrappers/UsdStage.h
@@ -15,109 +15,114 @@ namespace UE
 
 	namespace Internal
 	{
-		template< typename PtrType > class FUsdStageImpl;
+		template<typename PtrType>
+		class FUsdStageImpl;
 	}
 
 	/**
 	 * Minimal pxr::UsdStage pointer wrapper for Unreal that can be used from no-rtti modules.
 	 * Use the aliases FUsdStage and FUsdStageWeak instead (defined on ForwardDeclarations.h)
 	 */
-	template< typename PtrType >
+	template<typename PtrType>
 	class UNREALUSDWRAPPER_API FUsdStageBase
 	{
 	public:
 		FUsdStageBase();
 
-		FUsdStageBase( const FUsdStage& Other );
-		FUsdStageBase( FUsdStage&& Other );
-		FUsdStageBase( const FUsdStageWeak& Other );
-		FUsdStageBase( FUsdStageWeak&& Other );
+		FUsdStageBase(const FUsdStage& Other);
+		FUsdStageBase(FUsdStage&& Other);
+		FUsdStageBase(const FUsdStageWeak& Other);
+		FUsdStageBase(FUsdStageWeak&& Other);
 
-		FUsdStageBase& operator=( const FUsdStage& Other );
-		FUsdStageBase& operator=( FUsdStage&& Other );
-		FUsdStageBase& operator=( const FUsdStageWeak& Other );
-		FUsdStageBase& operator=( FUsdStageWeak&& Other );
+		FUsdStageBase& operator=(const FUsdStage& Other);
+		FUsdStageBase& operator=(FUsdStage&& Other);
+		FUsdStageBase& operator=(const FUsdStageWeak& Other);
+		FUsdStageBase& operator=(FUsdStageWeak&& Other);
 
 		~FUsdStageBase();
 
 		explicit operator bool() const;
 
-		bool operator==( const FUsdStageBase& Other ) const;
-		bool operator!=( const FUsdStageBase& Other ) const;
+		bool operator==(const FUsdStageBase& Other) const;
+		bool operator!=(const FUsdStageBase& Other) const;
 
 		// Auto conversion from/to PtrType. We use concrete pointer types here
 		// because we should also be able to convert between them
 	public:
 #if USE_USD_SDK
-		explicit FUsdStageBase( const pxr::UsdStageRefPtr& InUsdPtr );
-		explicit FUsdStageBase( pxr::UsdStageRefPtr&& InUsdPtr );
-		explicit FUsdStageBase( const pxr::UsdStageWeakPtr& InUsdPtr );
-		explicit FUsdStageBase( pxr::UsdStageWeakPtr&& InUsdPtr );
+		explicit FUsdStageBase(const pxr::UsdStageRefPtr& InUsdPtr);
+		explicit FUsdStageBase(pxr::UsdStageRefPtr&& InUsdPtr);
+		explicit FUsdStageBase(const pxr::UsdStageWeakPtr& InUsdPtr);
+		explicit FUsdStageBase(pxr::UsdStageWeakPtr&& InUsdPtr);
 
 		operator PtrType&();
 		operator const PtrType&() const;
 
-		operator pxr::UsdStageRefPtr ( ) const;
-		operator pxr::UsdStageWeakPtr ( ) const;
-#endif // USE_USD_SDK
+		operator pxr::UsdStageRefPtr() const;
+		operator pxr::UsdStageWeakPtr() const;
+#endif	  // USE_USD_SDK
 
-	// Wrapped pxr::UsdStage functions, refer to the USD SDK documentation
+		  // Wrapped pxr::UsdStage functions, refer to the USD SDK documentation
 	public:
-		void LoadAndUnload( const TSet<UE::FSdfPath>& LoadSet, const TSet<UE::FSdfPath>& UnloadSet, EUsdLoadPolicy Policy = EUsdLoadPolicy::UsdLoadWithDescendants );
+		void LoadAndUnload(
+			const TSet<UE::FSdfPath>& LoadSet,
+			const TSet<UE::FSdfPath>& UnloadSet,
+			EUsdLoadPolicy Policy = EUsdLoadPolicy::UsdLoadWithDescendants
+		);
 
 		/**
-		 * Saves a flattened copy of the stage to the given path (e.g. "C:/Folder/FlattenedStage.usda"). Will use the corresponding file writer depending on FilePath extension.
-		 * Will not alter the current stage.
+		 * Saves a flattened copy of the stage to the given path (e.g. "C:/Folder/FlattenedStage.usda"). Will use the corresponding file writer
+		 * depending on FilePath extension. Will not alter the current stage.
 		 */
-		bool Export( const TCHAR* FileName, bool bAddSourceFileComment = true, const TMap<FString, FString>& FileFormatArguments = {} ) const;
+		bool Export(const TCHAR* FileName, bool bAddSourceFileComment = true, const TMap<FString, FString>& FileFormatArguments = {}) const;
 
 		FSdfLayer GetRootLayer() const;
 		FSdfLayer GetSessionLayer() const;
-		bool HasLocalLayer( const FSdfLayer& Layer ) const;
+		bool HasLocalLayer(const FSdfLayer& Layer) const;
 
 		FUsdPrim GetPseudoRoot() const;
 		FUsdPrim GetDefaultPrim() const;
-		FUsdPrim GetPrimAtPath( const FSdfPath& Path ) const;
+		FUsdPrim GetPrimAtPath(const FSdfPath& Path) const;
 
-		TArray<FSdfLayer> GetLayerStack( bool bIncludeSessionLayers = true ) const;
-		TArray<FSdfLayer> GetUsedLayers( bool bIncludeClipLayers = true ) const;
+		TArray<FSdfLayer> GetLayerStack(bool bIncludeSessionLayers = true) const;
+		TArray<FSdfLayer> GetUsedLayers(bool bIncludeClipLayers = true) const;
 
-		void MuteAndUnmuteLayers( const TArray<FString>& MuteLayers, const TArray<FString>& UnmuteLayers );
-		bool IsLayerMuted( const FString& LayerIdentifier ) const;
+		void MuteAndUnmuteLayers(const TArray<FString>& MuteLayers, const TArray<FString>& UnmuteLayers);
+		bool IsLayerMuted(const FString& LayerIdentifier) const;
 
 		bool IsEditTargetValid() const;
-		void SetEditTarget( const FSdfLayer& Layer );
+		void SetEditTarget(const FSdfLayer& Layer);
 		FSdfLayer GetEditTarget() const;
 
-		bool GetMetadata( const TCHAR* Key, UE::FVtValue& Value ) const;
-		bool HasMetadata( const TCHAR* Key ) const;
-		bool SetMetadata( const TCHAR* Key, const UE::FVtValue& Value ) const;
-		bool ClearMetadata( const TCHAR* Key ) const;
+		bool GetMetadata(const TCHAR* Key, UE::FVtValue& Value) const;
+		bool HasMetadata(const TCHAR* Key) const;
+		bool SetMetadata(const TCHAR* Key, const UE::FVtValue& Value) const;
+		bool ClearMetadata(const TCHAR* Key) const;
 
 		double GetStartTimeCode() const;
 		double GetEndTimeCode() const;
-		void SetStartTimeCode( double TimeCode );
-		void SetEndTimeCode( double TimeCode );
+		void SetStartTimeCode(double TimeCode);
+		void SetEndTimeCode(double TimeCode);
 		double GetTimeCodesPerSecond() const;
-		void SetTimeCodesPerSecond( double TimeCodesPerSecond );
+		void SetTimeCodesPerSecond(double TimeCodesPerSecond);
 		double GetFramesPerSecond() const;
-		void SetFramesPerSecond( double FramesPerSecond );
+		void SetFramesPerSecond(double FramesPerSecond);
 
-		void SetInterpolationType( EUsdInterpolationType InterpolationType );
+		void SetInterpolationType(EUsdInterpolationType InterpolationType);
 		EUsdInterpolationType GetInterpolationType() const;
 
-		void SetDefaultPrim( const FUsdPrim& Prim );
+		void SetDefaultPrim(const FUsdPrim& Prim);
 
-		FUsdPrim OverridePrim( const FSdfPath& Path );
-		FUsdPrim DefinePrim( const FSdfPath& Path, const TCHAR* TypeName = TEXT("") );
+		FUsdPrim OverridePrim(const FSdfPath& Path);
+		FUsdPrim DefinePrim(const FSdfPath& Path, const TCHAR* TypeName = TEXT(""));
 		FUsdPrim CreateClassPrim(const FSdfPath& RootPrimPath);
-		bool RemovePrim( const FSdfPath& Path );
+		bool RemovePrim(const FSdfPath& Path);
 
 	private:
 		// So we can use the Other's Impl on copy constructor/operators
 		friend FUsdStage;
 		friend FUsdStageWeak;
 
-		TUniquePtr< Internal::FUsdStageImpl<PtrType> > Impl;
+		TUniquePtr<Internal::FUsdStageImpl<PtrType>> Impl;
 	};
-}
\ No newline at end of file
+}	 // namespace UE
diff --git a/Engine/Plugins/Importers/USDImporter/Source/UnrealUSDWrapper/Public/UsdWrappers/UsdTyped.h b/Engine/Plugins/Importers/USDImporter/Source/UnrealUSDWrapper/Public/UsdWrappers/UsdTyped.h
--- a/Engine/Plugins/Importers/USDImporter/Source/UnrealUSDWrapper/Public/UsdWrappers/UsdTyped.h
+++ b/Engine/Plugins/Importers/USDImporter/Source/UnrealUSDWrapper/Public/UsdWrappers/UsdTyped.h
@@ -5,9 +5,8 @@
 #include "CoreMinimal.h"
 
 #if USE_USD_SDK
-
 #include "USDIncludesStart.h"
-	#include "pxr/pxr.h"
+#include "pxr/pxr.h"
 #include "USDIncludesEnd.h"
 
 PXR_NAMESPACE_OPEN_SCOPE
@@ -15,8 +14,7 @@ PXR_NAMESPACE_OPEN_SCOPE
 	class UsdSchemaBase;
 	class UsdTyped;
 PXR_NAMESPACE_CLOSE_SCOPE
-
-#endif // #if USE_USD_SDK
+#endif	  // #if USE_USD_SDK
 
 namespace UE
 {
@@ -35,35 +33,35 @@ namespace UE
 	{
 	public:
 		FUsdTyped();
-		FUsdTyped( const FUsdTyped& Other );
-		FUsdTyped( FUsdTyped&& Other );
+		FUsdTyped(const FUsdTyped& Other);
+		FUsdTyped(FUsdTyped&& Other);
 
-		explicit FUsdTyped( const FUsdPrim& Prim );
+		explicit FUsdTyped(const FUsdPrim& Prim);
 
-		FUsdTyped& operator=( const FUsdTyped& Other );
-		FUsdTyped& operator=( FUsdTyped&& Other );
+		FUsdTyped& operator=(const FUsdTyped& Other);
+		FUsdTyped& operator=(FUsdTyped&& Other);
 
 		~FUsdTyped();
 
 		explicit operator bool() const;
 
-	// Auto conversion from/to pxr::UsdTyped
+		// Auto conversion from/to pxr::UsdTyped
 	public:
 #if USE_USD_SDK
-		explicit FUsdTyped( const pxr::UsdTyped& InUsdTyped );
-		explicit FUsdTyped( pxr::UsdTyped&& InUsdTyped );
-		explicit FUsdTyped( const pxr::UsdPrim& Prim );
+		explicit FUsdTyped(const pxr::UsdTyped& InUsdTyped);
+		explicit FUsdTyped(pxr::UsdTyped&& InUsdTyped);
+		explicit FUsdTyped(const pxr::UsdPrim& Prim);
 
 		operator pxr::UsdTyped&();
 		operator const pxr::UsdTyped&() const;
-#endif // #if USE_USD_SDK
+#endif	  // #if USE_USD_SDK
 
-	// Wrapped pxr::UsdTyped functions, refer to the USD SDK documentation
+		  // Wrapped pxr::UsdTyped functions, refer to the USD SDK documentation
 	public:
 		FSdfPath GetPath() const;
 		FUsdPrim GetPrim() const;
 
 	private:
-		TUniquePtr< Internal::FUsdTypedImpl > Impl;
+		TUniquePtr<Internal::FUsdTypedImpl> Impl;
 	};
-}
\ No newline at end of file
+}	 // namespace UE
diff --git a/Engine/Plugins/Importers/USDImporter/Source/UnrealUSDWrapper/Public/UsdWrappers/UsdVariantSets.h b/Engine/Plugins/Importers/USDImporter/Source/UnrealUSDWrapper/Public/UsdWrappers/UsdVariantSets.h
--- a/Engine/Plugins/Importers/USDImporter/Source/UnrealUSDWrapper/Public/UsdWrappers/UsdVariantSets.h
+++ b/Engine/Plugins/Importers/USDImporter/Source/UnrealUSDWrapper/Public/UsdWrappers/UsdVariantSets.h
@@ -10,9 +10,8 @@
 #include "UnrealUSDWrapper.h"
 
 #if USE_USD_SDK
-
 #include "USDIncludesStart.h"
-	#include "pxr/pxr.h"
+#include "pxr/pxr.h"
 #include "USDIncludesEnd.h"
 
 #include <string>
@@ -22,8 +21,7 @@ PXR_NAMESPACE_OPEN_SCOPE
 	class UsdVariantSet;
 	class UsdVariantSets;
 PXR_NAMESPACE_CLOSE_SCOPE
-
-#endif // #if USE_USD_SDK
+#endif	  // #if USE_USD_SDK
 
 namespace UE
 {
@@ -68,13 +66,11 @@ namespace UE
 		// Use FUsdPrim::GetVariantSet() or FUsdVariantSets::GetVariantSet()
 		// to construct an instance.
 		explicit FUsdVariantSet(const pxr::UsdPrim& InPrim, const std::string& InVariantSetName);
-#endif // #if USE_USD_SDK
+#endif	  // #if USE_USD_SDK
 
-		// Wrapped pxr::UsdVariantSet functions, refer to the USD SDK documentation
+		  // Wrapped pxr::UsdVariantSet functions, refer to the USD SDK documentation
 	public:
-		bool AddVariant(
-			const FString& VariantName,
-			EUsdListPosition Position = EUsdListPosition::BackOfPrependList);
+		bool AddVariant(const FString& VariantName, EUsdListPosition Position = EUsdListPosition::BackOfPrependList);
 
 		TArray<FString> GetVariantNames() const;
 
@@ -106,7 +102,6 @@ namespace UE
 		friend class FUsdVariantSets;
 	};
 
-
 	/**
 	 * Minimal pxr::UsdVariantSets wrapper for Unreal that can be used from no-rtti modules.
 	 */
@@ -139,13 +134,11 @@ namespace UE
 		// FUsdPrim.
 		// Use FUsdPrim::GetVariantSets() to construct an instance.
 		explicit FUsdVariantSets(const pxr::UsdPrim& InPrim);
-#endif // #if USE_USD_SDK
+#endif	  // #if USE_USD_SDK
 
-		// Wrapped pxr::UsdVariantSets functions, refer to the USD SDK documentation
+		  // Wrapped pxr::UsdVariantSets functions, refer to the USD SDK documentation
 	public:
-		FUsdVariantSet AddVariantSet(
-			const FString& VariantSetName,
-			EUsdListPosition Position = EUsdListPosition::BackOfPrependList);
+		FUsdVariantSet AddVariantSet(const FString& VariantSetName, EUsdListPosition Position = EUsdListPosition::BackOfPrependList);
 
 		TArray<FString> GetNames() const;
 
@@ -160,9 +153,7 @@ namespace UE
 
 		FString GetVariantSelection(const FString& VariantSetName) const;
 
-		bool SetSelection(
-			const FString& VariantSetName,
-			const FString& VariantName);
+		bool SetSelection(const FString& VariantSetName, const FString& VariantName);
 
 		TMap<FString, FString> GetAllVariantSelections() const;
 
@@ -171,4 +162,4 @@ namespace UE
 
 		friend class FUsdPrim;
 	};
-}
+}	 // namespace UE
diff --git a/Engine/Plugins/Importers/USDImporter/Source/UnrealUSDWrapper/Public/UsdWrappers/VtValue.h b/Engine/Plugins/Importers/USDImporter/Source/UnrealUSDWrapper/Public/UsdWrappers/VtValue.h
--- a/Engine/Plugins/Importers/USDImporter/Source/UnrealUSDWrapper/Public/UsdWrappers/VtValue.h
+++ b/Engine/Plugins/Importers/USDImporter/Source/UnrealUSDWrapper/Public/UsdWrappers/VtValue.h
@@ -5,16 +5,14 @@
 #include "CoreMinimal.h"
 
 #if USE_USD_SDK
-
 #include "USDIncludesStart.h"
-	#include "pxr/pxr.h"
+#include "pxr/pxr.h"
 #include "USDIncludesEnd.h"
 
 PXR_NAMESPACE_OPEN_SCOPE
 	class VtValue;
 PXR_NAMESPACE_CLOSE_SCOPE
-
-#endif // #if USE_USD_SDK
+#endif	  // #if USE_USD_SDK
 
 namespace UE
 {
@@ -31,39 +29,39 @@ namespace UE
 	public:
 		FVtValue();
 
-		FVtValue( const FVtValue& Other );
-		FVtValue( FVtValue&& Other );
+		FVtValue(const FVtValue& Other);
+		FVtValue(FVtValue&& Other);
 
-		FVtValue& operator=( const FVtValue& Other );
-		FVtValue& operator=( FVtValue&& Other );
+		FVtValue& operator=(const FVtValue& Other);
+		FVtValue& operator=(FVtValue&& Other);
 
 		~FVtValue();
 
-		bool operator==( const FVtValue& Other ) const;
-		bool operator!=( const FVtValue& Other ) const;
+		bool operator==(const FVtValue& Other) const;
+		bool operator!=(const FVtValue& Other) const;
 
-	// Auto conversion from/to pxr::VtValue
-	// We define a GetUsdValue() instead of `operator pxr::VtValue&` because pxr::VtValue itself has lower
-	// precedence, implicit generic conversions and operators that can conflict/generate hidden bugs if we try using them
-	// without having this header included (e.g. we could end up with a pxr::VtValue containing an UE::FVtValue instead)
+		// Auto conversion from/to pxr::VtValue
+		// We define a GetUsdValue() instead of `operator pxr::VtValue&` because pxr::VtValue itself has lower
+		// precedence, implicit generic conversions and operators that can conflict/generate hidden bugs if we try using them
+		// without having this header included (e.g. we could end up with a pxr::VtValue containing an UE::FVtValue instead)
 	public:
 #if USE_USD_SDK
-		explicit FVtValue( const pxr::VtValue& InVtValue );
-		explicit FVtValue( pxr::VtValue&& InVtValue );
-		FVtValue& operator=( const pxr::VtValue& InVtValue );
-		FVtValue& operator=( pxr::VtValue&& InVtValue );
+		explicit FVtValue(const pxr::VtValue& InVtValue);
+		explicit FVtValue(pxr::VtValue&& InVtValue);
+		FVtValue& operator=(const pxr::VtValue& InVtValue);
+		FVtValue& operator=(pxr::VtValue&& InVtValue);
 
 		pxr::VtValue& GetUsdValue();
 		const pxr::VtValue& GetUsdValue() const;
-#endif // #if USE_USD_SDK
+#endif	  // #if USE_USD_SDK
 
-	// Wrapped pxr::VtValue functions, refer to the USD SDK documentation
+		  // Wrapped pxr::VtValue functions, refer to the USD SDK documentation
 	public:
 		FString GetTypeName() const;
 		bool IsArrayValued() const;
 		bool IsEmpty() const;
 
 	private:
-		TUniquePtr< Internal::FVtValueImpl > Impl;
+		TUniquePtr<Internal::FVtValueImpl> Impl;
 	};
-}
\ No newline at end of file
+}	 // namespace UE
