USD: Reformat USDClasses module.

#rnx

diff --git a/Engine/Plugins/Importers/USDImporter/Source/USDClasses/Private/USDAssetCache2.cpp b/Engine/Plugins/Importers/USDImporter/Source/USDClasses/Private/USDAssetCache2.cpp
--- a/Engine/Plugins/Importers/USDImporter/Source/USDClasses/Private/USDAssetCache2.cpp
+++ b/Engine/Plugins/Importers/USDImporter/Source/USDClasses/Private/USDAssetCache2.cpp
@@ -13,7 +13,6 @@
 #include "Interfaces/ITargetPlatformManagerModule.h"
 #include "Materials/MaterialInterface.h"
 #include "MaterialShared.h"
-#include "MeshDescription.h"
 #include "Misc/App.h"
 #include "Misc/ScopeRWLock.h"
 #include "PhysicsEngine/PhysicsAsset.h"
@@ -28,7 +27,9 @@ static bool GOnDemandCachedAssetLoading = true;
 static FAutoConsoleVariableRef COnDemandCachedAssetLoading(
 	TEXT("USD.OnDemandCachedAssetLoading"),
 	GOnDemandCachedAssetLoading,
-	TEXT("When true will cause the USD Asset Cache to only load cached UObjects from disk once they are needed, regardless of when the actual Asset Cache asset itself is loaded."));
+	TEXT("When true will cause the USD Asset Cache to only load cached UObjects from disk once they are needed, regardless of when the actual Asset "
+		 "Cache asset itself is loaded.")
+);
 
 static int32 GCurrentCacheVersion = 1;
 static int32 GCurrentAssetInfoVersion = 2;
@@ -44,7 +45,12 @@ namespace UE::AssetCache::Private
 			return Result;
 		}
 
-		if (UObjectRedirector* RenamedClassRedirector = FindFirstObject<UObjectRedirector>(ClassName, EFindFirstObjectOptions::None, ELogVerbosity::Warning, TEXT("USDAssetCache2")))
+		if (UObjectRedirector* RenamedClassRedirector = FindFirstObject<UObjectRedirector>(
+				ClassName,
+				EFindFirstObjectOptions::None,
+				ELogVerbosity::Warning,
+				TEXT("USDAssetCache2")
+			))
 		{
 			return CastChecked<UClass>(RenamedClassRedirector->DestinationObject);
 		}
@@ -73,8 +79,7 @@ namespace UE::AssetCache::Private
 			UStaticMesh::StaticClass(),
 			UPhysicsAsset::StaticClass(),
 			USkeleton::StaticClass(),
-			USkeletalMesh::StaticClass()
-		};
+			USkeletalMesh::StaticClass()};
 
 		for (const UClass* AllowedClass : AllowedPersistentClasses)
 		{
@@ -102,7 +107,7 @@ namespace UE::AssetCache::Private
 				return static_cast<uint64>(Texture->Source.GetSizeOnDisk());
 			}
 			else
-#endif // WITH_EDITOR
+#endif	  // WITH_EDITOR
 			{
 				// TODO: Try doing this without calling GetPlatformData() if this is ever a speed bottleneck,
 				// as apparently that can stall.
@@ -110,11 +115,7 @@ namespace UE::AssetCache::Private
 				if (FTexturePlatformData* PlatformData = Texture->GetPlatformData())
 				{
 					const uint32 MipCount = 1;
-					return GPixelFormats[PlatformData->PixelFormat].Get2DTextureSizeInBytes(
-						PlatformData->SizeX,
-						PlatformData->SizeY,
-						MipCount
-					);
+					return GPixelFormats[PlatformData->PixelFormat].Get2DTextureSizeInBytes(PlatformData->SizeX, PlatformData->SizeY, MipCount);
 				}
 			}
 		}
@@ -124,7 +125,7 @@ namespace UE::AssetCache::Private
 
 	// Extends FObjectAndNameAsStringProxyArchive to support FLazyObjectPtr
 	// Copied from FSnapshotCustomArchive
-	struct FArchiveObjectPtrAsStringWrapper: public FObjectAndNameAsStringProxyArchive
+	struct FArchiveObjectPtrAsStringWrapper : public FObjectAndNameAsStringProxyArchive
 	{
 		FArchiveObjectPtrAsStringWrapper(FArchive& InInnerArchive)
 			: FObjectAndNameAsStringProxyArchive(InInnerArchive, false)
@@ -184,7 +185,7 @@ namespace UE::AssetCache::Private
 	class FAssetCacheObjectWriter final : public FObjectWriter
 	{
 	public:
-		FAssetCacheObjectWriter(UObject& Object, TArray< uint8 >& Bytes)
+		FAssetCacheObjectWriter(UObject& Object, TArray<uint8>& Bytes)
 			: FObjectWriter(Bytes)
 		{
 			SetIsLoading(false);
@@ -197,7 +198,7 @@ namespace UE::AssetCache::Private
 	class FAssetCacheObjectReader final : public FObjectReader
 	{
 	public:
-		FAssetCacheObjectReader(UObject& Object, const TArray< uint8 >& Bytes)
+		FAssetCacheObjectReader(UObject& Object, const TArray<uint8>& Bytes)
 			: FObjectReader(Bytes)
 		{
 			SetIsLoading(true);
@@ -220,7 +221,8 @@ namespace UE::AssetCache::Private
 			FObjectDependencyAnalyzer(UObject* InSourceObject, const TSet<UObject*>& InValidObjects)
 				: SourceObject(InSourceObject)
 				, ValidObjects(InValidObjects)
-			{ }
+			{
+			}
 
 			virtual FArchive& operator<<(UObject*& Obj) override
 			{
@@ -246,8 +248,8 @@ namespace UE::AssetCache::Private
 			TSet<UObject*> ConsumerObjects;
 		};
 
-		FAssetCacheObjectWriter Writer{ *Object, Buffer };
-		FArchiveObjectPtrAsStringWrapper Ar{ Writer };
+		FAssetCacheObjectWriter Writer{*Object, Buffer};
+		FArchiveObjectPtrAsStringWrapper Ar{Writer};
 
 		// Transfer over whether we're cooking or not, but otherwise leave at the default settings, as that is
 		// how we'll try to deserialize on the other end, which is potentially at runtime
@@ -260,7 +262,7 @@ namespace UE::AssetCache::Private
 		ensure(!Object->IsEditorOnly() || !Ar.IsCooking());
 
 		// Collect sub-objects depending on input object including nested objects
-		TArray< UObject* > SubObjectsArray;
+		TArray<UObject*> SubObjectsArray;
 		GetObjectsWithOuter(Object, SubObjectsArray, /*bIncludeNestedObjects = */ true);
 		TSet<UObject*> AllowedSubObjectsSet(SubObjectsArray);
 
@@ -279,7 +281,7 @@ namespace UE::AssetCache::Private
 		// Sort array of sub-objects based on their inter-dependency
 		{
 			// Create and initialize graph of dependency between sub-objects
-			TMap< UObject*, TSet<UObject*> > SubObjectDependencyGraph;
+			TMap<UObject*, TSet<UObject*>> SubObjectDependencyGraph;
 			SubObjectDependencyGraph.Reserve(SubObjectsArray.Num());
 
 			// Build graph of dependency: each entry contains the set of sub-objects to create before itself.
@@ -335,9 +337,7 @@ namespace UE::AssetCache::Private
 			FilteredSubObjects.Reserve(SubObjectsArray.Num());
 			for (UObject* SubObject : SubObjectsArray)
 			{
-				if (SubObject &&
-					!SubObject->HasAnyFlags(RF_DefaultSubObject) &&
-					(!SubObject->IsEditorOnly() || !Ar.IsCooking()))
+				if (SubObject && !SubObject->HasAnyFlags(RF_DefaultSubObject) && (!SubObject->IsEditorOnly() || !Ar.IsCooking()))
 				{
 					FilteredSubObjects.Add(SubObject);
 				}
@@ -345,7 +345,6 @@ namespace UE::AssetCache::Private
 			Swap(FilteredSubObjects, SubObjectsArray);
 		}
 
-
 		// Serialize size of array
 		int32 SubObjectsCount = SubObjectsArray.Num();
 		Ar << SubObjectsCount;
@@ -413,7 +412,7 @@ namespace UE::AssetCache::Private
 		// This method must obviously be called just after the InObject is created
 		auto RemoveDefaultDependencies = [](UObject* InObject)
 		{
-			TArray< UObject* > ObjectsWithOuter;
+			TArray<UObject*> ObjectsWithOuter;
 			GetObjectsWithOuter(InObject, ObjectsWithOuter, /*bIncludeNestedObjects = */ true);
 
 			for (UObject* ObjectWithOuter : ObjectsWithOuter)
@@ -423,7 +422,8 @@ namespace UE::AssetCache::Private
 				{
 					if (ObjectWithOuter != nullptr)
 					{
-						ObjectWithOuter->Rename(nullptr, GetTransientPackage(), REN_NonTransactional | REN_DontCreateRedirectors | REN_ForceNoResetLoaders);
+						ObjectWithOuter
+							->Rename(nullptr, GetTransientPackage(), REN_NonTransactional | REN_DontCreateRedirectors | REN_ForceNoResetLoaders);
 					}
 					ObjectWithOuter->MarkAsGarbage();
 				}
@@ -432,8 +432,8 @@ namespace UE::AssetCache::Private
 
 		RemoveDefaultDependencies(Object);
 
-		FAssetCacheObjectReader Reader{ *Object, InBuffer };
-		FArchiveObjectPtrAsStringWrapper Ar{ Reader };
+		FAssetCacheObjectReader Reader{*Object, InBuffer};
+		FArchiveObjectPtrAsStringWrapper Ar{Reader};
 		Ar.bLoadIfFindFails = true;
 
 		// Deserialize count of sub-objects
@@ -441,7 +441,7 @@ namespace UE::AssetCache::Private
 		Ar << SubObjectsCount;
 
 		// Create empty sub-objects based on class and patch
-		TArray< UObject* > SubObjectsArray;
+		TArray<UObject*> SubObjectsArray;
 		SubObjectsArray.SetNumZeroed(SubObjectsCount);
 
 		// Create root name to avoid name collision
@@ -528,7 +528,7 @@ namespace UE::AssetCache::Private
 
 			if (FApp::CanEverRender() || !FPlatformProperties::RequiresCookedData())
 			{
-				StaticMesh->SetRenderData(MakeUnique< FStaticMeshRenderData >());
+				StaticMesh->SetRenderData(MakeUnique<FStaticMeshRenderData>());
 				ITargetPlatformManagerModule& TargetPlatformManager = GetTargetPlatformManagerRef();
 				ITargetPlatform* RunningPlatform = TargetPlatformManager.GetRunningTargetPlatform();
 				const FStaticMeshLODSettings& LODSettings = RunningPlatform->GetStaticMeshLODSettings();
@@ -541,7 +541,7 @@ namespace UE::AssetCache::Private
 			SkeletalMesh->CalculateInvRefMatrices();
 			SkeletalMesh->InitMorphTargetsAndRebuildRenderData();
 		}
-#endif // WITH_EDITOR
+#endif	  // WITH_EDITOR
 		if (UMaterialInterface* Material = Cast<UMaterialInterface>(Object))
 		{
 			// TODO: If this presents a performance issue we could possibly just do part of the work that this seems
@@ -550,7 +550,7 @@ namespace UE::AssetCache::Private
 			UpdateContext.AddMaterialInterface(Material);
 		}
 	}
-}
+}	 // namespace UE::AssetCache::Private
 
 FArchive& operator<<(FArchive& Ar, UUsdAssetCache2::ECacheStorageType& Type)
 {
@@ -575,7 +575,8 @@ void UUsdAssetCache2::FCachedAssetInfo::Serialize(FArchive& Ar, UObject* Owner)
 	UE_LOG(
 		LogUsd,
 		Verbose,
-		TEXT("Serializing Asset Info '%s' (%s) with version %d, storage type: %d, disk bytes: %d, memory bytes: %d, bulkdata size: %d, saving? %d, persistent? %d, "
+		TEXT("Serializing Asset Info '%s' (%s) with version %d, storage type: %d, disk bytes: %d, memory bytes: %d, bulkdata size: %d, saving? %d, "
+			 "persistent? %d, "
 			 "transacting? %d, duplicating? %d, cooking? %d"),
 		*AssetName,
 		*Hash,
@@ -618,7 +619,7 @@ void UUsdAssetCache2::FCachedAssetInfo::Serialize(FArchive& Ar, UObject* Owner)
 }
 
 UUsdAssetCache2::UUsdAssetCache2()
-	: LRUCache(-1)  // Negative MaxElements as we don't want the cache to drop elements by itself
+	: LRUCache(-1)	  // Negative MaxElements as we don't want the cache to drop elements by itself
 {
 }
 
@@ -640,7 +641,7 @@ void UUsdAssetCache2::PostEditChangeProperty(FPropertyChangedEvent& PropertyChan
 
 	Super::PostEditChangeProperty(PropertyChangedEvent);
 }
-#endif // #if WITH_EDITOR
+#endif	  // #if WITH_EDITOR
 
 void UUsdAssetCache2::CacheAsset(const FString& Hash, UObject* Asset, const UObject* Referencer)
 {
@@ -659,7 +660,10 @@ void UUsdAssetCache2::CacheAsset(const FString& Hash, UObject* Asset, const UObj
 		{
 			if (FoundInfo->Hash != Hash)
 			{
-				UE_LOG(LogUsd, Warning, TEXT("Attempted to store asset '%s' more than once into the USD Asset Cache '%s'! (old hash: '%s', new hash: '%s')"),
+				UE_LOG(
+					LogUsd,
+					Warning,
+					TEXT("Attempted to store asset '%s' more than once into the USD Asset Cache '%s'! (old hash: '%s', new hash: '%s')"),
 					*Asset->GetPathName(),
 					*GetPathName(),
 					*FoundInfo->Hash,
@@ -676,7 +680,10 @@ void UUsdAssetCache2::CacheAsset(const FString& Hash, UObject* Asset, const UObj
 	{
 		if (Asset == ExistingAsset)
 		{
-			UE_LOG(LogUsd, Verbose, TEXT("Recaching asset '%s' with hash '%s' into the USD Asset Cache '%s'"),
+			UE_LOG(
+				LogUsd,
+				Verbose,
+				TEXT("Recaching asset '%s' with hash '%s' into the USD Asset Cache '%s'"),
 				*Asset->GetPathName(),
 				*Hash,
 				*GetPathName()
@@ -698,7 +705,10 @@ void UUsdAssetCache2::CacheAsset(const FString& Hash, UObject* Asset, const UObj
 				if (UObject* OldAsset = RemoveAsset(Hash))
 				{
 					bRemoved = true;
-					UE_LOG(LogUsd, Log, TEXT("Overwriting asset '%s' with '%s' (for hash '%s') in the USD Asset Cache '%s'"),
+					UE_LOG(
+						LogUsd,
+						Log,
+						TEXT("Overwriting asset '%s' with '%s' (for hash '%s') in the USD Asset Cache '%s'"),
 						*OldAsset->GetPathName(),
 						*Asset->GetPathName(),
 						*Hash,
@@ -707,9 +717,13 @@ void UUsdAssetCache2::CacheAsset(const FString& Hash, UObject* Asset, const UObj
 				}
 			}
 
-			if(!bRemoved)
+			if (!bRemoved)
 			{
-				UE_LOG(LogUsd, Error, TEXT("Irrecoverable hash collision! Asset '%s' cannot be cached into USD Asset Cache '%s' as the existing asset '%s' with the same hash '%s' could not be discarded!"),
+				UE_LOG(
+					LogUsd,
+					Error,
+					TEXT("Irrecoverable hash collision! Asset '%s' cannot be cached into USD Asset Cache '%s' as the existing asset '%s' with the "
+						 "same hash '%s' could not be discarded!"),
 					*Asset->GetPathName(),
 					*GetPathName(),
 					*ExistingAsset->GetPathName(),
@@ -720,11 +734,7 @@ void UUsdAssetCache2::CacheAsset(const FString& Hash, UObject* Asset, const UObj
 		}
 	}
 
-	UE_LOG(LogUsd, Verbose, TEXT("Caching asset '%s' with hash '%s' into the USD Asset Cache '%s'"),
-		*Asset->GetPathName(),
-		*Hash,
-		*GetPathName()
-	);
+	UE_LOG(LogUsd, Verbose, TEXT("Caching asset '%s' with hash '%s' into the USD Asset Cache '%s'"), *Asset->GetPathName(), *Hash, *GetPathName());
 
 	Modify();
 
@@ -750,7 +760,7 @@ void UUsdAssetCache2::CacheAsset(const FString& Hash, UObject* Asset, const UObj
 	NewInfo.AssetFlags = Asset->GetFlags();
 	if (ReferencerToUse)
 	{
-		NewInfo.Referencers.Add(FObjectKey{ ReferencerToUse });
+		NewInfo.Referencers.Add(FObjectKey{ReferencerToUse});
 	}
 
 	AssetStorage.Add(Hash, Asset);
@@ -843,7 +853,10 @@ UObject* UUsdAssetCache2::GetCachedAsset(const FString& Hash)
 		FReadScopeLock Lock{RWLock};
 		FoundAsset = AssetStorage.FindRef(Hash);
 	}
-	UE_LOG(LogUsd, Verbose, TEXT("Fetching cached asset with hash '%s' for cache '%s'. Loaded in AssetStorage? %d. Size of LRU cache: %d, Pending persistent assets: %d"),
+	UE_LOG(
+		LogUsd,
+		Verbose,
+		TEXT("Fetching cached asset with hash '%s' for cache '%s'. Loaded in AssetStorage? %d. Size of LRU cache: %d, Pending persistent assets: %d"),
 		*Hash,
 		*GetPathName(),
 		FoundAsset != nullptr,
@@ -864,15 +877,12 @@ UObject* UUsdAssetCache2::GetCachedAsset(const FString& Hash)
 		FoundAssetPath = PendingPersistentStorage.Find(Hash);
 		if (!FoundAssetPath)
 		{
-			UE_LOG(LogUsd, Verbose, TEXT("Full cache miss: We do not know about any assets with hash '%s' in cache '%s'"),
-				*Hash,
-				*GetPathName()
-			);
+			UE_LOG(LogUsd, Verbose, TEXT("Full cache miss: We do not know about any assets with hash '%s' in cache '%s'"), *Hash, *GetPathName());
 			return nullptr;
 		}
 	}
 
-	TFunction<UObject* (FCachedAssetInfo* ObjectInfo)> LoadAssetFromBulkData = [this](FCachedAssetInfo* ObjectInfo) -> UObject*
+	TFunction<UObject*(FCachedAssetInfo * ObjectInfo)> LoadAssetFromBulkData = [this](FCachedAssetInfo* ObjectInfo) -> UObject*
 	{
 		if (!ObjectInfo)
 		{
@@ -883,7 +893,10 @@ UObject* UUsdAssetCache2::GetCachedAsset(const FString& Hash)
 		UClass* FoundClass = UE::AssetCache::Private::FindClass(*ClassNameString);
 		if (!FoundClass)
 		{
-			UE_LOG(LogUsd, Warning, TEXT("Failed to find object class '%s' when deserializing asset '%s' with hash '%s'"),
+			UE_LOG(
+				LogUsd,
+				Warning,
+				TEXT("Failed to find object class '%s' when deserializing asset '%s' with hash '%s'"),
 				*ClassNameString,
 				*ObjectInfo->AssetName,
 				*ObjectInfo->Hash
@@ -909,7 +922,10 @@ UObject* UUsdAssetCache2::GetCachedAsset(const FString& Hash)
 		UE::AssetCache::Private::DeserializeObjectAndSubObjects(NewAsset, BulkDataCopy);
 		ensure(NewAsset);
 
-		UE_LOG(LogUsd, Log, TEXT("Deserialized asset at '%s' (%s, hash '%s', storage %d) from bulkdata (%.3f MB)"),
+		UE_LOG(
+			LogUsd,
+			Log,
+			TEXT("Deserialized asset at '%s' (%s, hash '%s', storage %d) from bulkdata (%.3f MB)"),
 			*ObjectInfo->AssetName,
 			NewAsset ? *NewAsset->GetPathName() : TEXT("nullptr"),
 			*ObjectInfo->Hash,
@@ -944,9 +960,7 @@ UObject* UUsdAssetCache2::GetCachedAsset(const FString& Hash)
 		}
 		else
 		{
-			UE_LOG(LogUsd, Warning, TEXT("Failed to find info about asset '%s' when deserializing on-demand"),
-				*AssetPath.ToString()
-			);
+			UE_LOG(LogUsd, Warning, TEXT("Failed to find info about asset '%s' when deserializing on-demand"), *AssetPath.ToString());
 		}
 	};
 	{
@@ -955,7 +969,10 @@ UObject* UUsdAssetCache2::GetCachedAsset(const FString& Hash)
 	}
 
 	ensure(AssetsToLoad.Num() > 0);
-	UE_LOG(LogUsd, Verbose, TEXT("Partial cache miss: Will attempt to load %d assets from bulk data of cache '%s'"),
+	UE_LOG(
+		LogUsd,
+		Verbose,
+		TEXT("Partial cache miss: Will attempt to load %d assets from bulk data of cache '%s'"),
 		AssetsToLoad.Num(),
 		*GetPathName()
 	);
@@ -969,7 +986,7 @@ UObject* UUsdAssetCache2::GetCachedAsset(const FString& Hash)
 				ensure(LoadedAsset);
 
 				TouchAssetInternal(LoadedAsset, CurrentScopedReferencer);
-				FoundAsset = LoadedAsset;  // The very last asset we loaded is our main asset, as it's pushed last into AssetsToLoad
+				FoundAsset = LoadedAsset;	 // The very last asset we loaded is our main asset, as it's pushed last into AssetsToLoad
 			}
 		}
 	}
@@ -1018,10 +1035,7 @@ bool UUsdAssetCache2::RemoveAssetReference(const UObject* Asset, const UObject*
 		const int32 NumRemoved = Info->Referencers.Remove(FObjectKey{Referencer});
 		if (NumRemoved > 0)
 		{
-			UE_LOG(LogUsd, Verbose, TEXT("Removed referencer '%s' from asset '%s'"),
-				*Referencer->GetPathName(),
-				*Asset->GetPathName()
-			);
+			UE_LOG(LogUsd, Verbose, TEXT("Removed referencer '%s' from asset '%s'"), *Referencer->GetPathName(), *Asset->GetPathName());
 
 			bRemovedSomething = true;
 		}
@@ -1091,7 +1105,7 @@ TArray<FString> UUsdAssetCache2::GetAllAssetHashes() const
 	TArray<FString> Hashes;
 	Hashes.Reserve(LRUCache.Num());
 
-	for (TLruCache<FSoftObjectPath, FCachedAssetInfo>::TConstIterator Iter{ LRUCache }; Iter; ++Iter)
+	for (TLruCache<FSoftObjectPath, FCachedAssetInfo>::TConstIterator Iter{LRUCache}; Iter; ++Iter)
 	{
 		Hashes.Add(Iter.Value().Hash);
 	}
@@ -1121,7 +1135,7 @@ TArray<FString> UUsdAssetCache2::GetAllCachedAssetPaths() const
 	TArray<FString> Paths;
 	Paths.Reserve(LRUCache.Num());
 
-	for (TLruCache<FSoftObjectPath, FCachedAssetInfo>::TConstIterator Iter{ LRUCache }; Iter; ++Iter)
+	for (TLruCache<FSoftObjectPath, FCachedAssetInfo>::TConstIterator Iter{LRUCache}; Iter; ++Iter)
 	{
 		Paths.Add(Iter.Key().ToString());
 	}
@@ -1139,7 +1153,7 @@ void UUsdAssetCache2::Reset()
 	const int32 MaxElements = -1;
 	AssetStorage.Reset();
 	ActiveAssets.Reset();
-	for (TLruCache<FSoftObjectPath, FCachedAssetInfo>::TIterator Iter{ LRUCache }; Iter; ++Iter)
+	for (TLruCache<FSoftObjectPath, FCachedAssetInfo>::TIterator Iter{LRUCache}; Iter; ++Iter)
 	{
 		FCachedAssetInfo& Info = Iter.Value();
 		Info.BulkData.RemoveBulkData();
@@ -1155,12 +1169,12 @@ void UUsdAssetCache2::RefreshStorage()
 	const bool bIsTransientAssetCache = GetOutermost() == GetTransientPackage();
 
 	const uint64 MaxUnreferencedTransientBytes = UnreferencedAssetStorageSizeMB < 0.0
-		? TNumericLimits<uint64>::Max()
-		: static_cast<uint64>(FMath::Max(UnreferencedAssetStorageSizeMB, 0.0) + 0.5) * 1000000;
+													 ? TNumericLimits<uint64>::Max()
+													 : static_cast<uint64>(FMath::Max(UnreferencedAssetStorageSizeMB, 0.0) + 0.5) * 1000000;
 
 	const uint64 MaxPersistentBytes = PersistentAssetStorageSizeMB < 0.0
-		? TNumericLimits<uint64>::Max()
-		: static_cast<uint64>(FMath::Max(PersistentAssetStorageSizeMB, 0.0) + 0.5) * 1000000;
+										  ? TNumericLimits<uint64>::Max()
+										  : static_cast<uint64>(FMath::Max(PersistentAssetStorageSizeMB, 0.0) + 0.5) * 1000000;
 
 	uint64 UnreferencedAssetsSumBytes = 0;
 	uint64 PersistentAssetSumBytes = 0;
@@ -1180,7 +1194,7 @@ void UUsdAssetCache2::RefreshStorage()
 	double StartTime = FPlatformTime::Cycles64();
 
 	// Do an initial pass to reset our current storage types and fetch some remaining data
-	for (TLruCache<FSoftObjectPath, FCachedAssetInfo>::TIterator Iter{ LRUCache }; Iter; ++Iter)
+	for (TLruCache<FSoftObjectPath, FCachedAssetInfo>::TIterator Iter{LRUCache}; Iter; ++Iter)
 	{
 		FCachedAssetInfo& Info = Iter.Value();
 
@@ -1246,171 +1260,176 @@ void UUsdAssetCache2::RefreshStorage()
 	};
 
 	TFunction<ECacheStorageType(FCachedAssetInfo&, FSoftObjectPath&)> VisitEntryRecursively;
-	VisitEntryRecursively =
-		[
-			this,
-			&ConsumerTreeHasReferencedAsset,
-			&VisitEntryRecursively,
-			&VisitedAssets,
-			bIsTransientAssetCache,
-			MaxUnreferencedTransientBytes,
-			MaxPersistentBytes,
-			&UnreferencedAssetsSumBytes,
-			&PersistentAssetSumBytes,
-			&NewAssetStorage,
-			&NewPendingPersistentStorage
-		]
-		(FCachedAssetInfo& Info, FSoftObjectPath& AssetPath) -> ECacheStorageType
+	VisitEntryRecursively = [this,
+							 &ConsumerTreeHasReferencedAsset,
+							 &VisitEntryRecursively,
+							 &VisitedAssets,
+							 bIsTransientAssetCache,
+							 MaxUnreferencedTransientBytes,
+							 MaxPersistentBytes,
+							 &UnreferencedAssetsSumBytes,
+							 &PersistentAssetSumBytes,
+							 &NewAssetStorage,
+							 &NewPendingPersistentStorage](FCachedAssetInfo& Info, FSoftObjectPath& AssetPath) -> ECacheStorageType
+	{
+		if (AssetPath.IsNull())
 		{
-			if (AssetPath.IsNull())
-			{
-				ensure(false);
-				return ECacheStorageType::None;
-			}
+			ensure(false);
+			return ECacheStorageType::None;
+		}
 
-			if (VisitedAssets.Contains(AssetPath))
+		if (VisitedAssets.Contains(AssetPath))
+		{
+			return Info.CurrentStorageType;
+		}
+		VisitedAssets.Add(AssetPath);
+
+		// Always visit dependencies first: We can only ever persist an asset if its dependencies are also
+		// persisted. Here we'll also track the "worst" storage any of our dependencies got (e.g. if one of our
+		// textures managed Persistent, but the other got Referenced, we'll track Referenced)
+		ECacheStorageType WorstDependencyStorage = ECacheStorageType::Persistent;
+		for (FSoftObjectPath& Dependency : Info.Dependencies)
+		{
+			FCachedAssetInfo* DependencyInfo = LRUCache.Find(Dependency);
+			if (!ensure(DependencyInfo))
 			{
-				return Info.CurrentStorageType;
+				continue;
 			}
-			VisitedAssets.Add(AssetPath);
 
-			// Always visit dependencies first: We can only ever persist an asset if its dependencies are also
-			// persisted. Here we'll also track the "worst" storage any of our dependencies got (e.g. if one of our
-			// textures managed Persistent, but the other got Referenced, we'll track Referenced)
-			ECacheStorageType WorstDependencyStorage = ECacheStorageType::Persistent;
-			for (FSoftObjectPath& Dependency : Info.Dependencies)
-			{
-				FCachedAssetInfo* DependencyInfo = LRUCache.Find(Dependency);
-				if (!ensure(DependencyInfo))
-				{
-					continue;
-				}
+			ECacheStorageType DependencyStorage = VisitEntryRecursively(*DependencyInfo, Dependency);
+			WorstDependencyStorage = static_cast<ECacheStorageType>(
+				FMath::Min(static_cast<uint8>(WorstDependencyStorage), static_cast<uint8>(DependencyStorage))
+			);
+		}
 
-				ECacheStorageType DependencyStorage = VisitEntryRecursively(*DependencyInfo, Dependency);
-				WorstDependencyStorage = static_cast<ECacheStorageType>(FMath::Min(
-					static_cast<uint8>(WorstDependencyStorage),
-					static_cast<uint8>(DependencyStorage)
-				));
-			}
+		const bool bAssetFitsInPersistent = PersistentAssetSumBytes + Info.SizeOnDiskInBytes <= MaxPersistentBytes;
+		const bool bAssetFitsInUnreferenced = UnreferencedAssetsSumBytes + Info.SizeOnMemoryInBytes <= MaxUnreferencedTransientBytes;
+		const bool bDependenciesArePersistent = WorstDependencyStorage == ECacheStorageType::Persistent;
+		const bool bAllDependenciesInStorage = WorstDependencyStorage != ECacheStorageType::None;
+		const bool bConsumerIsReferenced = ConsumerTreeHasReferencedAsset(Info);
 
-			const bool bAssetFitsInPersistent = PersistentAssetSumBytes + Info.SizeOnDiskInBytes <= MaxPersistentBytes;
-			const bool bAssetFitsInUnreferenced = UnreferencedAssetsSumBytes + Info.SizeOnMemoryInBytes <= MaxUnreferencedTransientBytes;
-			const bool bDependenciesArePersistent = WorstDependencyStorage == ECacheStorageType::Persistent;
-			const bool bAllDependenciesInStorage = WorstDependencyStorage != ECacheStorageType::None;
-			const bool bConsumerIsReferenced = ConsumerTreeHasReferencedAsset(Info);
+		const FString AssetPathStr = AssetPath.ToString();
+		const FString AssetClassNameStr = Info.AssetClassName.ToString();
+		UObject* Asset = AssetStorage.FindRef(Info.Hash);
 
-			const FString AssetPathStr = AssetPath.ToString();
-			const FString AssetClassNameStr = Info.AssetClassName.ToString();
-			UObject* Asset = AssetStorage.FindRef(Info.Hash);
+		// Priority 1: Asset can be moved to persistent storage
+		if (!bIsTransientAssetCache && bAssetFitsInPersistent && bDependenciesArePersistent && AssetCanBePersisted(*AssetClassNameStr))
+		{
+			Info.CurrentStorageType = ECacheStorageType::Persistent;
+		}
 
-			// Priority 1: Asset can be moved to persistent storage
-			if (!bIsTransientAssetCache && bAssetFitsInPersistent && bDependenciesArePersistent && AssetCanBePersisted(*AssetClassNameStr))
-			{
-				Info.CurrentStorageType = ECacheStorageType::Persistent;
-			}
+		// Priority 2: Asset is kept on transient storage because it (or a Consumer) is referenced
+		else if (Info.Referencers.Num() > 0 || bConsumerIsReferenced)
+		{
+			Info.CurrentStorageType = ECacheStorageType::Referenced;
+		}
+
+		// Priority 3: Asset fits on transient storage even while unreferenced, as long as all dependencies also
+		// made it somewhere and weren't evicted
+		else if (bAssetFitsInUnreferenced && bAllDependenciesInStorage)
+		{
+			Info.CurrentStorageType = ECacheStorageType::Unreferenced;
+		}
+
+		UE_LOG(
+			LogUsd,
+			Verbose,
+			TEXT("Asset '%s' (class %s, hash '%s', size %.3f MB on memory, %.3f MB on disk) fits in persistent? %d, fits in unreferenced? %d, "
+				 "referencer count %d, dependency count %d. Target storage: %d"),
+			*AssetPathStr,
+			*AssetClassNameStr,
+			*Info.Hash,
+			Info.SizeOnMemoryInBytes / 1000000.0,
+			Info.SizeOnDiskInBytes / 1000000.0,
+			bAssetFitsInPersistent,
+			bAssetFitsInUnreferenced,
+			Info.Referencers.Num(),
+			Info.Dependencies.Num(),
+			int(Info.CurrentStorageType)
+		);
 
-			// Priority 2: Asset is kept on transient storage because it (or a Consumer) is referenced
-			else if (Info.Referencers.Num() > 0 || bConsumerIsReferenced)
+		switch (Info.CurrentStorageType)
+		{
+			case ECacheStorageType::None:
 			{
-				Info.CurrentStorageType = ECacheStorageType::Referenced;
+				Info.BulkData.UnloadBulkData();
+				break;
 			}
-
-			// Priority 3: Asset fits on transient storage even while unreferenced, as long as all dependencies also
-			// made it somewhere and weren't evicted
-			else if (bAssetFitsInUnreferenced && bAllDependenciesInStorage)
+			case ECacheStorageType::Referenced:
 			{
-				Info.CurrentStorageType = ECacheStorageType::Unreferenced;
+				// We'll at least unload the bulkdata here, but not actually remove it from disk yet, in case we
+				// put this asset back in persistent later and can reuse the same bulkdata.
+				// We'll get rid of the bulkdata for sure on UUsdAssetCache2::Serialize(), if the asset is not
+				// persistent by then
+				Info.BulkData.UnloadBulkData();
+
+				ensure(Asset);
+				NewAssetStorage.Add(Info.Hash, Asset);
+				break;
 			}
-
-			UE_LOG(LogUsd, Verbose, TEXT("Asset '%s' (class %s, hash '%s', size %.3f MB on memory, %.3f MB on disk) fits in persistent? %d, fits in unreferenced? %d, referencer count %d, dependency count %d. Target storage: %d"),
-				*AssetPathStr,
-				*AssetClassNameStr,
-				*Info.Hash,
-				Info.SizeOnMemoryInBytes / 1000000.0,
-				Info.SizeOnDiskInBytes / 1000000.0,
-				bAssetFitsInPersistent,
-				bAssetFitsInUnreferenced,
-				Info.Referencers.Num(),
-				Info.Dependencies.Num(),
-				int(Info.CurrentStorageType)
-			);
-
-			switch (Info.CurrentStorageType)
+			case ECacheStorageType::Unreferenced:
 			{
-				case ECacheStorageType::None:
-				{
-					Info.BulkData.UnloadBulkData();
-					break;
-				}
-				case ECacheStorageType::Referenced:
-				{
-					// We'll at least unload the bulkdata here, but not actually remove it from disk yet, in case we
-					// put this asset back in persistent later and can reuse the same bulkdata.
-					// We'll get rid of the bulkdata for sure on UUsdAssetCache2::Serialize(), if the asset is not
-					// persistent by then
-					Info.BulkData.UnloadBulkData();
+				Info.BulkData.UnloadBulkData();
 
-					ensure(Asset);
+				if (Asset)
+				{
+					ensure(AssetStorage.Contains(Info.Hash));
 					NewAssetStorage.Add(Info.Hash, Asset);
-					break;
 				}
-				case ECacheStorageType::Unreferenced:
+				// It's possible to end up with an asset that is currently persistent but living in
+				// Unreferenced storage if we load an asset cache and discover the asset fits in the latter
+				// but not in the former (e.g. if size on disk > size in memory somehow, or if the storage
+				// sizes are tweaked to cause that)
+				else
 				{
-					Info.BulkData.UnloadBulkData();
-
-					if (Asset)
-					{
-						ensure(AssetStorage.Contains(Info.Hash));
-						NewAssetStorage.Add(Info.Hash, Asset);
-					}
-					// It's possible to end up with an asset that is currently persistent but living in
-					// Unreferenced storage if we load an asset cache and discover the asset fits in the latter
-					// but not in the former (e.g. if size on disk > size in memory somehow, or if the storage
-					// sizes are tweaked to cause that)
-					else
-					{
-						ensure(!AssetStorage.Contains(Info.Hash));
-						NewPendingPersistentStorage.Add(Info.Hash, AssetPath);
-					}
-
-					UnreferencedAssetsSumBytes += Info.SizeOnMemoryInBytes;
-					UE_LOG(LogUsd, Verbose, TEXT("Updated unreferenced sum to %.3f out of max %u MB"),
-						UnreferencedAssetsSumBytes / 1000000.0,
-						MaxUnreferencedTransientBytes / 1000000
-					);
-					break;
+					ensure(!AssetStorage.Contains(Info.Hash));
+					NewPendingPersistentStorage.Add(Info.Hash, AssetPath);
 				}
-				case ECacheStorageType::Persistent:
-				{
-					if (Asset)
-					{
-						ensure(AssetStorage.Contains(Info.Hash));
-						NewAssetStorage.Add(Info.Hash, Asset);
-					}
-					else
-					{
-						ensure(!AssetStorage.Contains(Info.Hash));
-						NewPendingPersistentStorage.Add(Info.Hash, AssetPath);
-					}
 
-					PersistentAssetSumBytes += Info.SizeOnDiskInBytes;
-					UE_LOG(LogUsd, Verbose, TEXT("Updated persistent sum to %.3f out of max %u MB"),
-						PersistentAssetSumBytes / 1000000.0,
-						MaxPersistentBytes / 1000000
-					);
-					break;
+				UnreferencedAssetsSumBytes += Info.SizeOnMemoryInBytes;
+				UE_LOG(
+					LogUsd,
+					Verbose,
+					TEXT("Updated unreferenced sum to %.3f out of max %u MB"),
+					UnreferencedAssetsSumBytes / 1000000.0,
+					MaxUnreferencedTransientBytes / 1000000
+				);
+				break;
+			}
+			case ECacheStorageType::Persistent:
+			{
+				if (Asset)
+				{
+					ensure(AssetStorage.Contains(Info.Hash));
+					NewAssetStorage.Add(Info.Hash, Asset);
 				}
-				default:
+				else
 				{
-					ensure(false);
-					break;
+					ensure(!AssetStorage.Contains(Info.Hash));
+					NewPendingPersistentStorage.Add(Info.Hash, AssetPath);
 				}
+
+				PersistentAssetSumBytes += Info.SizeOnDiskInBytes;
+				UE_LOG(
+					LogUsd,
+					Verbose,
+					TEXT("Updated persistent sum to %.3f out of max %u MB"),
+					PersistentAssetSumBytes / 1000000.0,
+					MaxPersistentBytes / 1000000
+				);
+				break;
 			}
+			default:
+			{
+				ensure(false);
+				break;
+			}
+		}
 
-			return Info.CurrentStorageType;
-		};
+		return Info.CurrentStorageType;
+	};
 
-	for (TLruCache<FSoftObjectPath, FCachedAssetInfo>::TIterator Iter{ LRUCache }; Iter; ++Iter)
+	for (TLruCache<FSoftObjectPath, FCachedAssetInfo>::TIterator Iter{LRUCache}; Iter; ++Iter)
 	{
 		FSoftObjectPath& AssetPath = Iter.Key();
 		FCachedAssetInfo& Info = Iter.Value();
@@ -1424,7 +1443,7 @@ void UUsdAssetCache2::RefreshStorage()
 	// Evict from the LRUCache assets that never made it into any storage: These assets are unreferenced and don't fit
 	// anymore
 	{
-		TLruCache<FSoftObjectPath, FCachedAssetInfo>::TIterator Iter{ LRUCache };
+		TLruCache<FSoftObjectPath, FCachedAssetInfo>::TIterator Iter{LRUCache};
 		while (Iter)
 		{
 			FSoftObjectPath& AssetPath = Iter.Key();
@@ -1432,7 +1451,14 @@ void UUsdAssetCache2::RefreshStorage()
 
 			if (Info.CurrentStorageType == ECacheStorageType::None)
 			{
-				UE_LOG(LogUsd, Verbose, TEXT("Evicting asset '%s' (hash '%s') from the cache '%s'"), *AssetPath.ToString(), *Info.Hash, *GetPathName());
+				UE_LOG(
+					LogUsd,
+					Verbose,
+					TEXT("Evicting asset '%s' (hash '%s') from the cache '%s'"),
+					*AssetPath.ToString(),
+					*Info.Hash,
+					*GetPathName()
+				);
 
 				// Discard data on disk
 				TryUnloadAsset(Info);
@@ -1460,7 +1486,7 @@ void UUsdAssetCache2::RefreshStorage()
 	// unload all persistent assets that aren't being currently used
 	if (GOnDemandCachedAssetLoading)
 	{
-		for (TLruCache<FSoftObjectPath, FCachedAssetInfo>::TIterator Iter{ LRUCache }; Iter; ++Iter)
+		for (TLruCache<FSoftObjectPath, FCachedAssetInfo>::TIterator Iter{LRUCache}; Iter; ++Iter)
 		{
 			FCachedAssetInfo& Info = Iter.Value();
 			TryUnloadAsset(Info);
@@ -1469,7 +1495,11 @@ void UUsdAssetCache2::RefreshStorage()
 
 	double ElapsedSeconds = FPlatformTime::ToSeconds64(FPlatformTime::Cycles64() - StartTime);
 
-	UE_LOG(LogUsd, Log, TEXT("Finished asset cache '%s' storage refresh in %.3f seconds: %.3f out of max %u MB of unreferenced assets, %.3f out of max %u MB of persistent assets"),
+	UE_LOG(
+		LogUsd,
+		Log,
+		TEXT("Finished asset cache '%s' storage refresh in %.3f seconds: %.3f out of max %u MB of unreferenced assets, %.3f out of max %u MB of "
+			 "persistent assets"),
 		*GetName(),
 		ElapsedSeconds,
 		UnreferencedAssetsSumBytes / 1000000.0,
@@ -1567,10 +1597,7 @@ bool UUsdAssetCache2::AddAssetReferenceInternal(const UObject* Asset, const UObj
 	{
 		Info->Referencers.Add(FObjectKey{Referencer});
 
-		UE_LOG(LogUsd, Verbose, TEXT("Added referencer '%s' for asset '%s'"),
-			*Referencer->GetPathName(),
-			*Asset->GetPathName()
-		);
+		UE_LOG(LogUsd, Verbose, TEXT("Added referencer '%s' for asset '%s'"), *Referencer->GetPathName(), *Asset->GetPathName());
 		return true;
 	}
 
@@ -1647,7 +1674,10 @@ void UUsdAssetCache2::Serialize(FArchive& Ar)
 			int32 NumInfos = Ar.IsPersistent() ? PersistentAssetsWithInfo.Num() : LRUCache.Num();
 			Ar << NumInfos;
 
-			UE_LOG(LogUsd, Verbose, TEXT("Serializing USD Asset Cache '%s' with %d asset infos. Persistent? %d, transacting? %d, duplicating? %d, cooking? %d"),
+			UE_LOG(
+				LogUsd,
+				Verbose,
+				TEXT("Serializing USD Asset Cache '%s' with %d asset infos. Persistent? %d, transacting? %d, duplicating? %d, cooking? %d"),
 				*GetPathName(),
 				NumInfos,
 				Ar.IsPersistent(),
@@ -1755,10 +1785,7 @@ void UUsdAssetCache2::Serialize(FArchive& Ar)
 			TArray<FSoftObjectPath> AssetPaths;
 			AssetPaths.Reserve(NumInfos);
 
-			UE_LOG(LogUsd, Verbose, TEXT("Deserializing USD Asset Cache '%s' with %d asset infos"),
-				*GetPathName(),
-				NumInfos
-			);
+			UE_LOG(LogUsd, Verbose, TEXT("Deserializing USD Asset Cache '%s' with %d asset infos"), *GetPathName(), NumInfos);
 
 			for (int32 Index = 0; Index < NumInfos; ++Index)
 			{
diff --git a/Engine/Plugins/Importers/USDImporter/Source/USDClasses/Private/USDClassesModule.cpp b/Engine/Plugins/Importers/USDImporter/Source/USDClasses/Private/USDClassesModule.cpp
--- a/Engine/Plugins/Importers/USDImporter/Source/USDClasses/Private/USDClassesModule.cpp
+++ b/Engine/Plugins/Importers/USDImporter/Source/USDClasses/Private/USDClassesModule.cpp
@@ -13,7 +13,7 @@
 #include "Engine/SkeletalMesh.h"
 #include "Engine/SkinnedAssetCommon.h"
 #include "Engine/StaticMesh.h"
-#include "Engine/Texture2D.h"
+#include "Engine/Texture.h"
 #include "EngineAnalytics.h"
 #include "GeometryCache.h"
 #include "GroomAsset.h"
@@ -34,36 +34,35 @@
 #include "UObject/NameTypes.h"
 #include "UObject/ObjectSaveContext.h"
 #include "UObject/Package.h"
-#include "UObject/Package.h"
 #include "UObject/UObjectGlobals.h"
 
 #if WITH_EDITOR
 #include "Editor/MaterialEditor/Public/MaterialEditingLibrary.h"
-#endif // WITH_EDITOR
+#endif	  // WITH_EDITOR
 
-DEFINE_LOG_CATEGORY( LogUsd );
+DEFINE_LOG_CATEGORY(LogUsd);
 
 namespace UE::USDClasses::Private
 {
-	static const FString DisplayColorID = TEXT( "!DisplayColor" );
+	static const FString DisplayColorID = TEXT("!DisplayColor");
 
-	TSharedPtr<FJsonObject> ParseJSON( const FString& FileContents )
+	TSharedPtr<FJsonObject> ParseJSON(const FString& FileContents)
 	{
-		if ( FileContents.IsEmpty() )
+		if (FileContents.IsEmpty())
 		{
 			return nullptr;
 		}
 
-		const TSharedRef< TJsonReader<> >& Reader = TJsonReaderFactory<>::Create( FileContents );
+		const TSharedRef<TJsonReader<>>& Reader = TJsonReaderFactory<>::Create(FileContents);
 
 		TSharedPtr<FJsonObject> DescriptorObject;
-		if ( FJsonSerializer::Deserialize( Reader, DescriptorObject ) && DescriptorObject.IsValid() )
+		if (FJsonSerializer::Deserialize(Reader, DescriptorObject) && DescriptorObject.IsValid())
 		{
 			return DescriptorObject;
 		}
 		else
 		{
-			UE_LOG( LogUsd, Warning, TEXT( "Failed to parse plugInfo.json file: '%s'" ), *Reader->GetErrorMessage() );
+			UE_LOG(LogUsd, Warning, TEXT("Failed to parse plugInfo.json file: '%s'"), *Reader->GetErrorMessage());
 		}
 
 		return nullptr;
@@ -72,18 +71,18 @@ namespace UE::USDClasses::Private
 	TMap<FString, int32> PackagePathNameToDirtyCounter;
 }
 
-void IUsdClassesModule::UpdatePlugInfoFiles( const FString& PluginDirectory, const FString& TargetDllFolder )
+void IUsdClassesModule::UpdatePlugInfoFiles(const FString& PluginDirectory, const FString& TargetDllFolder)
 {
 	// Traverse all USD plugins
 	TArray<FString> JsonPaths;
 	const bool bFiles = true;
 	const bool bDirectories = false;
-	const bool bClearFileNames = false;  // Whether it needs to reset the output array or not
-	IFileManager::Get().FindFilesRecursive( JsonPaths, *PluginDirectory, TEXT( "plugInfo.json" ), bFiles, bDirectories, bClearFileNames );
-	for ( const FString& JsonFilePath : JsonPaths )
+	const bool bClearFileNames = false;	   // Whether it needs to reset the output array or not
+	IFileManager::Get().FindFilesRecursive(JsonPaths, *PluginDirectory, TEXT("plugInfo.json"), bFiles, bDirectories, bClearFileNames);
+	for (const FString& JsonFilePath : JsonPaths)
 	{
 		FString FileContent;
-		if ( !FFileHelper::LoadFileToString( FileContent, *JsonFilePath ) )
+		if (!FFileHelper::LoadFileToString(FileContent, *JsonFilePath))
 		{
 			continue;
 		}
@@ -94,134 +93,147 @@ void IUsdClassesModule::UpdatePlugInfoFiles( const FString& PluginDirectory, con
 		//		# changes to types with autoGenerated=true.
 		// We will not save this preamble back into the file, but that's not bad as the file is already in the binaries folder
 		// and the preamble is actually invalid json anyway
-		int32 FirstOpenBracketPos = FileContent.Find( TEXT( "{" ) );
+		int32 FirstOpenBracketPos = FileContent.Find(TEXT("{"));
 		int32 NewLen = FileContent.Len() - FirstOpenBracketPos;
-		FileContent.RightInline( NewLen );
+		FileContent.RightInline(NewLen);
 
-		TSharedPtr<FJsonObject> DescriptorObject = UE::USDClasses::Private::ParseJSON( FileContent );
-		if ( !DescriptorObject.IsValid() )
+		TSharedPtr<FJsonObject> DescriptorObject = UE::USDClasses::Private::ParseJSON(FileContent);
+		if (!DescriptorObject.IsValid())
 		{
 			continue;
 		}
 
-		TSharedPtr<FJsonValue> PluginsField = DescriptorObject->TryGetField( TEXT( "Plugins" ) );
-		if ( !PluginsField.IsValid() )
+		TSharedPtr<FJsonValue> PluginsField = DescriptorObject->TryGetField(TEXT("Plugins"));
+		if (!PluginsField.IsValid())
 		{
 			continue;
 		}
 
 		bool bUpdatedFile = false;
 
-		const TArray< TSharedPtr<FJsonValue> >* PluginEntriesPtr = nullptr;
-		if ( !PluginsField->TryGetArray( PluginEntriesPtr ) || !PluginEntriesPtr )
+		const TArray<TSharedPtr<FJsonValue>>* PluginEntriesPtr = nullptr;
+		if (!PluginsField->TryGetArray(PluginEntriesPtr) || !PluginEntriesPtr)
 		{
 			continue;
 		}
 
-		for ( const TSharedPtr<FJsonValue>& PluginEntry : *PluginEntriesPtr )
+		for (const TSharedPtr<FJsonValue>& PluginEntry : *PluginEntriesPtr)
 		{
 			const TSharedPtr<FJsonObject>* PluginEntryObjectPtr = nullptr;
-			if ( !PluginEntry->TryGetObject( PluginEntryObjectPtr ) || !PluginEntryObjectPtr || !PluginEntryObjectPtr->IsValid() )
+			if (!PluginEntry->TryGetObject(PluginEntryObjectPtr) || !PluginEntryObjectPtr || !PluginEntryObjectPtr->IsValid())
 			{
 				continue;
 			}
 
-			TSharedPtr<FJsonValue> LibraryPathField = ( *PluginEntryObjectPtr )->TryGetField( TEXT( "LibraryPath" ) );
-			if ( !LibraryPathField.IsValid() )
+			TSharedPtr<FJsonValue> LibraryPathField = (*PluginEntryObjectPtr)->TryGetField(TEXT("LibraryPath"));
+			if (!LibraryPathField.IsValid())
 			{
 				continue;
 			}
 
 			// e.g. "../../../../../../Win64/sdf.dll"
 			FString LibraryPath;
-			if ( !LibraryPathField->TryGetString( LibraryPath ) )
+			if (!LibraryPathField->TryGetString(LibraryPath))
 			{
 				continue;
 			}
 
 			// e.g. "sdf.dll"
-			FString DllName = FPaths::GetCleanFilename( LibraryPath );
-			if ( DllName.IsEmpty() )
+			FString DllName = FPaths::GetCleanFilename(LibraryPath);
+			if (DllName.IsEmpty())
 			{
 				continue;
 			}
 
 			// e.g. ".."
 			FString RootPath;
-			if ( TSharedPtr<FJsonValue> RootPathField = ( *PluginEntryObjectPtr )->TryGetField( TEXT( "Root" ) ) )
+			if (TSharedPtr<FJsonValue> RootPathField = (*PluginEntryObjectPtr)->TryGetField(TEXT("Root")))
 			{
-				RootPathField->TryGetString( RootPath );
+				RootPathField->TryGetString(RootPath);
 			}
 
 			// e.g. "E:/Folder/Workspace/Engine/Binaries/ThirdParty/USD/UsdResources/Win64/plugins/ar/resources/.."
-			RootPath = FPaths::Combine( FPaths::GetPath( JsonFilePath ), RootPath );
-			FPaths::NormalizeDirectoryName( RootPath );
-			FPaths::CollapseRelativeDirectories( RootPath );
+			RootPath = FPaths::Combine(FPaths::GetPath(JsonFilePath), RootPath);
+			FPaths::NormalizeDirectoryName(RootPath);
+			FPaths::CollapseRelativeDirectories(RootPath);
 
 			FString PathToDll = TargetDllFolder;
 
 			// Ensure directories end with a slash, otherwise FPaths::MakePathRelativeTo may strip the last folder thinking it's a filename
-			if ( !PathToDll.EndsWith( TEXT( "/" ) ) )
+			if (!PathToDll.EndsWith(TEXT("/")))
 			{
-				PathToDll.AppendChar( TEXT( '/' ) );
+				PathToDll.AppendChar(TEXT('/'));
 			}
-			if ( !RootPath.EndsWith( TEXT( "/" ) ) )
+			if (!RootPath.EndsWith(TEXT("/")))
 			{
-				RootPath.AppendChar( TEXT( '/' ) );
+				RootPath.AppendChar(TEXT('/'));
 			}
 
 			// e.g. "../../../../../../Win64"
 			// This can fail if we're in different drives. Then we're forced to assume that CurrentDllLocation is an absolute path and use it directly
-			if ( !FPaths::MakePathRelativeTo( PathToDll, *RootPath ) )
+			if (!FPaths::MakePathRelativeTo(PathToDll, *RootPath))
 			{
-				if ( FPaths::IsRelative( PathToDll ) && !FPaths::DirectoryExists( PathToDll ) )
+				if (FPaths::IsRelative(PathToDll) && !FPaths::DirectoryExists(PathToDll))
 				{
-					UE_LOG( LogUsd, Warning, TEXT( "Failed to handle Current DLL Location '%s' when trying to update plugInfo.json paths" ), *PathToDll );
+					UE_LOG(LogUsd, Warning, TEXT("Failed to handle Current DLL Location '%s' when trying to update plugInfo.json paths"), *PathToDll);
 					continue;
 				}
 			}
 
 			// e.g. "../../../../../../Win64/sdf.dll"
-			FString NewLibraryPath = FPaths::Combine( PathToDll, DllName );
-			( *PluginEntryObjectPtr )->SetField( TEXT( "LibraryPath" ), MakeShared<FJsonValueString>( NewLibraryPath ) );
+			FString NewLibraryPath = FPaths::Combine(PathToDll, DllName);
+			(*PluginEntryObjectPtr)->SetField(TEXT("LibraryPath"), MakeShared<FJsonValueString>(NewLibraryPath));
 			bUpdatedFile = true;
 		}
 
 		// Serialize them back to JSON if we changed anything
-		if ( bUpdatedFile )
+		if (bUpdatedFile)
 		{
 			FString OutputString;
-			TSharedRef< TJsonWriter<> > Writer = TJsonWriterFactory<>::Create( &OutputString );
-			FJsonSerializer::Serialize( DescriptorObject.ToSharedRef(), Writer );
-
-			if ( FFileHelper::SaveStringToFile( OutputString, *JsonFilePath, FFileHelper::EEncodingOptions::AutoDetect, &IFileManager::Get(), FILEWRITE_EvenIfReadOnly ) )
+			TSharedRef<TJsonWriter<>> Writer = TJsonWriterFactory<>::Create(&OutputString);
+			FJsonSerializer::Serialize(DescriptorObject.ToSharedRef(), Writer);
+
+			if (FFileHelper::SaveStringToFile(
+					OutputString,
+					*JsonFilePath,
+					FFileHelper::EEncodingOptions::AutoDetect,
+					&IFileManager::Get(),
+					FILEWRITE_EvenIfReadOnly
+				))
 			{
-				UE_LOG( LogUsd, Log, TEXT( "Updated LibraryPaths for USD plugInfo.json file '%s'" ), *JsonFilePath );
+				UE_LOG(LogUsd, Log, TEXT("Updated LibraryPaths for USD plugInfo.json file '%s'"), *JsonFilePath);
 			}
 			else
 			{
-				UE_LOG( LogUsd, Warning, TEXT( "Failed to update LibraryPath for USD plugInfo.json file '%s'" ), *JsonFilePath );
+				UE_LOG(LogUsd, Warning, TEXT("Failed to update LibraryPath for USD plugInfo.json file '%s'"), *JsonFilePath);
 			}
 		}
 	}
 }
 
-void IUsdClassesModule::SendAnalytics( TArray<FAnalyticsEventAttribute>&& InAttributes, const FString& EventName, bool bAutomated, double ElapsedSeconds, double NumberOfFrames, const FString& Extension )
+void IUsdClassesModule::SendAnalytics(
+	TArray<FAnalyticsEventAttribute>&& InAttributes,
+	const FString& EventName,
+	bool bAutomated,
+	double ElapsedSeconds,
+	double NumberOfFrames,
+	const FString& Extension
+)
 {
-	if ( FEngineAnalytics::IsAvailable() )
+	if (FEngineAnalytics::IsAvailable())
 	{
-		TArray<FAnalyticsEventAttribute> Attributes( InAttributes );
+		TArray<FAnalyticsEventAttribute> Attributes(InAttributes);
 
-		Attributes.Emplace( TEXT( "Automated" ), LexToString( bAutomated ) );
-		Attributes.Emplace( TEXT( "FileExtension" ), Extension );
-		Attributes.Emplace( TEXT( "NumberOfFrames" ), LexToString( FMath::Abs( NumberOfFrames ) ) );
-		Attributes.Emplace( TEXT( "TimeTaken.Seconds" ), ElapsedSeconds );
-		Attributes.Emplace( TEXT( "Platform" ), FPlatformProperties::IniPlatformName() );
-		Attributes.Emplace( TEXT( "EngineVersion" ), FEngineVersion::Current().ToString() );
-		Attributes.Emplace( TEXT( "EngineMode" ), FPlatformMisc::GetEngineMode() );
+		Attributes.Emplace(TEXT("Automated"), LexToString(bAutomated));
+		Attributes.Emplace(TEXT("FileExtension"), Extension);
+		Attributes.Emplace(TEXT("NumberOfFrames"), LexToString(FMath::Abs(NumberOfFrames)));
+		Attributes.Emplace(TEXT("TimeTaken.Seconds"), ElapsedSeconds);
+		Attributes.Emplace(TEXT("Platform"), FPlatformProperties::IniPlatformName());
+		Attributes.Emplace(TEXT("EngineVersion"), FEngineVersion::Current().ToString());
+		Attributes.Emplace(TEXT("EngineMode"), FPlatformMisc::GetEngineMode());
 
-		const FString EventText = FString::Printf( TEXT( "Engine.Usage.USD.%s" ), *EventName);
-		FEngineAnalytics::GetProvider().RecordEvent( EventText, Attributes );
+		const FString EventText = FString::Printf(TEXT("Engine.Usage.USD.%s"), *EventName);
+		FEngineAnalytics::GetProvider().RecordEvent(EventText, Attributes);
 	}
 }
 
@@ -240,79 +252,79 @@ void IUsdClassesModule::SendAnalytics(TArray<FAnalyticsEventAttribute>&& InAttri
 	}
 }
 
-bool IUsdClassesModule::HashObjectPackage( const UObject* Object, FSHA1& HashToUpdate )
+bool IUsdClassesModule::HashObjectPackage(const UObject* Object, FSHA1& HashToUpdate)
 {
 #if WITH_EDITOR
-	if ( !Object )
+	if (!Object)
 	{
 		return false;
 	}
 
 	UPackage* Package = Object->GetOutermost();
-	if ( !Package )
+	if (!Package)
 	{
 		return false;
 	}
 
 	// Hash package's persistent Guid
 	const FGuid& Guid = Package->GetPersistentGuid();
-	HashToUpdate.Update( reinterpret_cast< const uint8* >( &Guid ), sizeof( Guid ) );
+	HashToUpdate.Update(reinterpret_cast<const uint8*>(&Guid), sizeof(Guid));
 
 	// Hash last modified date
 	FString PackageFullName = Package->GetPathName();
 	FString FileName;
-	if ( FPackageName::TryConvertLongPackageNameToFilename( PackageFullName, FileName ) )
+	if (FPackageName::TryConvertLongPackageNameToFilename(PackageFullName, FileName))
 	{
-		FFileStatData StatData = IFileManager::Get().GetStatData( *FileName );
-		if ( StatData.bIsValid )
+		FFileStatData StatData = IFileManager::Get().GetStatData(*FileName);
+		if (StatData.bIsValid)
 		{
 			FString ModifiedTimeString = StatData.ModificationTime.ToString();
-			HashToUpdate.UpdateWithString( *ModifiedTimeString, ModifiedTimeString.Len() );
+			HashToUpdate.UpdateWithString(*ModifiedTimeString, ModifiedTimeString.Len());
 		}
 	}
 
 	// If this asset is currently dirty, also hash how many times it was dirtied in this session.
 	// If it's ever saved, we'll reset this counter but update the last saved date
-	if ( int32* DirtyCounter = UE::USDClasses::Private::PackagePathNameToDirtyCounter.Find( PackageFullName ) )
+	if (int32* DirtyCounter = UE::USDClasses::Private::PackagePathNameToDirtyCounter.Find(PackageFullName))
 	{
-		HashToUpdate.Update( reinterpret_cast< const uint8* >( DirtyCounter ), sizeof( *DirtyCounter ) );
+		HashToUpdate.Update(reinterpret_cast<const uint8*>(DirtyCounter), sizeof(*DirtyCounter));
 	}
 
 	return true;
 #else
 	return false;
-#endif // WITH_EDITOR
+#endif	  // WITH_EDITOR
 }
 
-UWorld* IUsdClassesModule::GetCurrentWorld( bool bEditorWorldsOnly )
+UWorld* IUsdClassesModule::GetCurrentWorld(bool bEditorWorldsOnly)
 {
 	UWorld* EditorWorld = nullptr;
 	UWorld* LowestPIEWorld = nullptr;
 	int32 LowestPIEInstance = TNumericLimits<int32>::Max();
 
-	for ( const FWorldContext& Context : GEngine->GetWorldContexts() )
+	for (const FWorldContext& Context : GEngine->GetWorldContexts())
 	{
 		UWorld* World = Context.World();
-		if ( !World )
+		if (!World)
 		{
 			continue;
 		}
 
-		if ( !bEditorWorldsOnly && Context.WorldType == EWorldType::PIE )
+		if (!bEditorWorldsOnly && Context.WorldType == EWorldType::PIE)
 		{
-			if ( !LowestPIEInstance || Context.PIEInstance < LowestPIEInstance )
+			if (!LowestPIEInstance || Context.PIEInstance < LowestPIEInstance)
 			{
 				LowestPIEWorld = World;
 				LowestPIEInstance = Context.PIEInstance;
 			}
 		}
-		else if ( Context.WorldType == EWorldType::Editor )
+		else if (Context.WorldType == EWorldType::Editor)
 		{
 			EditorWorld = World;
 		}
 	}
 
-	if ( LowestPIEWorld )
+	if (LowestPIEWorld)
 	{
 		return LowestPIEWorld;
 	}
@@ -570,26 +582,25 @@ UMaterialInstanceConstant* IUsdClassesModule::CreateDisplayColorMaterialInstance
 
 class FUsdClassesModule : public IUsdClassesModule
 {
-
 public:
 	virtual void StartupModule() override
 	{
 		PackageMarkedDirtyEventHandle = UPackage::PackageMarkedDirtyEvent.AddLambda(
-			[]( const UPackage* Package, bool bWasDirty )
+			[](const UPackage* Package, bool bWasDirty)
 			{
-				if ( Package )
+				if (Package)
 				{
-					UE::USDClasses::Private::PackagePathNameToDirtyCounter.FindOrAdd( Package->GetPathName() )++;
+					UE::USDClasses::Private::PackagePathNameToDirtyCounter.FindOrAdd(Package->GetPathName())++;
 				}
 			}
 		);
 
 		PackageSavedWithContextEventHandle = UPackage::PackageSavedWithContextEvent.AddLambda(
-			[]( const FString& PackageFilename, UPackage* Package, FObjectPostSaveContext ObjectSaveContext )
+			[](const FString& PackageFilename, UPackage* Package, FObjectPostSaveContext ObjectSaveContext)
 			{
-				if ( Package )
+				if (Package)
 				{
-					UE::USDClasses::Private::PackagePathNameToDirtyCounter.Remove( Package->GetPathName() );
+					UE::USDClasses::Private::PackagePathNameToDirtyCounter.Remove(Package->GetPathName());
 				}
 			}
 		);
@@ -597,8 +608,8 @@ public:
 
 	virtual void ShutdownModule() override
 	{
-		UPackage::PackageSavedWithContextEvent.Remove( PackageSavedWithContextEventHandle );
-		UPackage::PackageMarkedDirtyEvent.Remove( PackageMarkedDirtyEventHandle );
+		UPackage::PackageSavedWithContextEvent.Remove(PackageSavedWithContextEventHandle);
+		UPackage::PackageMarkedDirtyEvent.Remove(PackageMarkedDirtyEventHandle);
 	}
 
 private:
@@ -606,5 +617,4 @@ private:
 	FDelegateHandle PackageSavedWithContextEventHandle;
 };
 
-IMPLEMENT_MODULE( FUsdClassesModule, USDClasses );
-
+IMPLEMENT_MODULE(FUsdClassesModule, USDClasses);
diff --git a/Engine/Plugins/Importers/USDImporter/Source/USDClasses/Private/USDDrawModeComponent.cpp b/Engine/Plugins/Importers/USDImporter/Source/USDClasses/Private/USDDrawModeComponent.cpp
--- a/Engine/Plugins/Importers/USDImporter/Source/USDClasses/Private/USDDrawModeComponent.cpp
+++ b/Engine/Plugins/Importers/USDImporter/Source/USDClasses/Private/USDDrawModeComponent.cpp
@@ -27,13 +27,14 @@
 
 #if WITH_EDITOR
 #include "Editor/MaterialEditor/Public/MaterialEditingLibrary.h"
-#endif // WITH_EDITOR
+#endif	  // WITH_EDITOR
 
 static bool GUseWholeExtentsToDetectSelection = true;
 static FAutoConsoleVariableRef CVarUseWholeExtentsToDetectSelection(
 	TEXT("USD.Bounds.UseWholeExtentsToDetectSelection"),
 	GUseWholeExtentsToDetectSelection,
-	TEXT("When true it means we can select the UUsdDrawModeComponents in 'bounds', 'cards' or 'origin' mode by clicking anywhere on the extent's 'volume' itself, and not only by clicking directly on the drawn lines")
+	TEXT("When true it means we can select the UUsdDrawModeComponents in 'bounds', 'cards' or 'origin' mode by clicking anywhere on the extent's "
+		 "'volume' itself, and not only by clicking directly on the drawn lines")
 );
 
 static float GBoundsLineThickness = 2.0f;
@@ -86,8 +87,7 @@ namespace UE::UsdDrawModeComponentImpl::Private
 					FVector3f(-0.5, -0.5, +0.5),
 					FVector3f(-0.5, +0.5, +0.5),
 					FVector3f(+0.5, +0.5, +0.5),
-					FVector3f(+0.5, -0.5, +0.5)
-				};
+					FVector3f(+0.5, -0.5, +0.5)};
 
 				FVector2f UVs[4] = {
 					FVector2f(0, 0),
@@ -151,7 +151,6 @@ namespace UE::UsdDrawModeComponentImpl::Private
 						HitProxyBuffer.Upload();
 					}
 				);
-
 			}
 		}
 
@@ -202,7 +201,8 @@ namespace UE::UsdDrawModeComponentImpl::Private
 							);
 							bOutputVelocity |= AlwaysHasVelocity();
 
-							FDynamicPrimitiveUniformBuffer& DynamicPrimitiveUniformBuffer = Collector.AllocateOneFrameResource<FDynamicPrimitiveUniformBuffer>();
+							FDynamicPrimitiveUniformBuffer&
+								DynamicPrimitiveUniformBuffer = Collector.AllocateOneFrameResource<FDynamicPrimitiveUniformBuffer>();
 
 							DynamicPrimitiveUniformBuffer.Set(
 								Collector.GetRHICommandList(),
@@ -1079,7 +1079,7 @@ HHitProxy* UUsdDrawModeComponent::CreateMeshHitProxy(int32 SectionIndex, int32 M
 {
 	return nullptr;
 }
-#endif // WITH_EDITOR
+#endif	  // WITH_EDITOR
 
 EUsdModelCardFace UUsdDrawModeComponent::GetAuthoredFaces() const
 {
diff --git a/Engine/Plugins/Importers/USDImporter/Source/USDClasses/Public/USDAssetCache.h b/Engine/Plugins/Importers/USDImporter/Source/USDClasses/Public/USDAssetCache.h
--- a/Engine/Plugins/Importers/USDImporter/Source/USDClasses/Public/USDAssetCache.h
+++ b/Engine/Plugins/Importers/USDImporter/Source/USDClasses/Public/USDAssetCache.h
@@ -3,14 +3,14 @@
 #pragma once
 
 #include "Containers/Map.h"
-#include "CoreMinimal.h"
 
 #include "USDAssetCache.generated.h"
 
 class UE_DEPRECATED(5.2, "Use the UUsdAssetCache2 class instead, which is an actual asset you can create on the Content Browser") UUsdAssetCache;
 
 /** Owns the assets generated and reused by the USD stage, allowing thread-safe retrieval/storage */
-UCLASS(meta=(ScriptName="UsdAssetCache_Deprecated"))
+UCLASS(meta = (ScriptName = "UsdAssetCache_Deprecated"))
+
 class USDCLASSES_API UUsdAssetCache : public UObject
 {
 	GENERATED_BODY()
@@ -19,38 +19,49 @@ public:
 	UUsdAssetCache();
 
 #if WITH_EDITOR
-	virtual void PostEditChangeProperty( struct FPropertyChangedEvent& PropertyChangedEvent ) override;
+	virtual void PostEditChangeProperty(struct FPropertyChangedEvent& PropertyChangedEvent) override;
 #endif
 
-	void CacheAsset( const FString& Hash, UObject* Asset, const FString& PrimPath = FString() );
-	void DiscardAsset( const FString& Hash );
-	UObject* GetCachedAsset( const FString& Hash ) const;
+	void CacheAsset(const FString& Hash, UObject* Asset, const FString& PrimPath = FString());
+	void DiscardAsset(const FString& Hash);
+	UObject* GetCachedAsset(const FString& Hash) const;
+
 	TMap<FString, TObjectPtr<UObject>> GetCachedAssets() const
 	{
-		TMap<FString, TObjectPtr<UObject>> CachedAssets( TransientStorage );
-		CachedAssets.Append( PersistentStorage );
+		TMap<FString, TObjectPtr<UObject>> CachedAssets(TransientStorage);
+		CachedAssets.Append(PersistentStorage);
 
 		return CachedAssets;
 	}
 
-	void LinkAssetToPrim( const FString& PrimPath, UObject* Asset );
-	void RemoveAssetPrimLink( const FString& PrimPath );
-	UObject* GetAssetForPrim( const FString& PrimPath ) const;
+	void LinkAssetToPrim(const FString& PrimPath, UObject* Asset);
+	void RemoveAssetPrimLink(const FString& PrimPath);
+	UObject* GetAssetForPrim(const FString& PrimPath) const;
 
 	/** Returns the first prim path associated with an asset (see TMap::FindKey) */
-	FString GetPrimForAsset( UObject* Asset ) const;
+	FString GetPrimForAsset(UObject* Asset) const;
 
 	/**
 	 * Returns the hash used as a key for an asset, in case we own it.
 	 * Returns the empty string otherwise.
 	 */
-	FString GetHashForAsset( UObject* Asset ) const;
+	FString GetHashForAsset(UObject* Asset) const;
 
-	TMap< FString, TWeakObjectPtr<UObject> > GetAssetPrimLinks() const { return PrimPathToAssets; }; // Can't return a reference as it wouldn't be thread-safe
+	TMap<FString, TWeakObjectPtr<UObject>> GetAssetPrimLinks() const
+	{
+		return PrimPathToAssets;
+	};	  // Can't return a reference as it wouldn't be thread-safe
 
-	bool IsAssetOwnedByCache( UObject* Asset ) const { return OwnedAssets.Contains( Asset ); }
+	bool IsAssetOwnedByCache(UObject* Asset) const
+	{
+		return OwnedAssets.Contains(Asset);
+	}
+
+	int32 GetNumAssets() const
+	{
+		return TransientStorage.Num() + PersistentStorage.Num();
+	}
 
-	int32 GetNumAssets() const { return TransientStorage.Num() + PersistentStorage.Num(); }
 	void Reset();
 
 	/**
@@ -67,28 +78,28 @@ public:
 	TSet<UObject*> GetActiveAssets() const;
 
 	// We need to be serializable so that AUsdStageActor can duplicate us for PIE
-	virtual void Serialize( FArchive& Ar ) override;
+	virtual void Serialize(FArchive& Ar) override;
 
 private:
 	// For now everything is NonPIEDuplicateTransient as this is mostly a subobject of AUsdStageActor. When the actor is duplicated
 	// it will need to reload the stage anyway to rebuild its prim links to components and assets, so there is no point in duplicating
 	// the properties here just yet. Obviously we want to duplicate these properties
-	UPROPERTY( NonPIEDuplicateTransient, Transient, VisibleAnywhere, Category = "Assets" )
-	TMap< FString, TObjectPtr<UObject> > TransientStorage;
+	UPROPERTY(NonPIEDuplicateTransient, Transient, VisibleAnywhere, Category = "Assets")
+	TMap<FString, TObjectPtr<UObject>> TransientStorage;
 
-	UPROPERTY( NonPIEDuplicateTransient, VisibleAnywhere, Category = "Assets" )
-	TMap< FString, TObjectPtr<UObject> > PersistentStorage;
+	UPROPERTY(NonPIEDuplicateTransient, VisibleAnywhere, Category = "Assets")
+	TMap<FString, TObjectPtr<UObject>> PersistentStorage;
 
-	UPROPERTY( EditAnywhere, Category = "Assets", AdvancedDisplay )
+	UPROPERTY(EditAnywhere, Category = "Assets", AdvancedDisplay)
 	bool bAllowPersistentStorage;
 
 	// Points to the assets in primary storage, used to quickly check if we own an asset
-	UPROPERTY( NonPIEDuplicateTransient )
-	TSet< TWeakObjectPtr<UObject> > OwnedAssets;
+	UPROPERTY(NonPIEDuplicateTransient)
+	TSet<TWeakObjectPtr<UObject>> OwnedAssets;
 
 	// Keeps associations from prim paths to assets that we own in primary storage
-	UPROPERTY( NonPIEDuplicateTransient )
-    TMap< FString, TWeakObjectPtr<UObject> > PrimPathToAssets;
+	UPROPERTY(NonPIEDuplicateTransient)
+	TMap<FString, TWeakObjectPtr<UObject>> PrimPathToAssets;
 
 	// Assets that were added/retrieved since the last call to MarkAssetsAsSlate();
 	mutable TSet<UObject*> ActiveAssets;
@@ -96,4 +107,3 @@ private:
 	// Used to ensure thread safety
 	mutable FCriticalSection CriticalSection;
 };
-
diff --git a/Engine/Plugins/Importers/USDImporter/Source/USDClasses/Public/USDAssetCache2.h b/Engine/Plugins/Importers/USDImporter/Source/USDClasses/Public/USDAssetCache2.h
--- a/Engine/Plugins/Importers/USDImporter/Source/USDClasses/Public/USDAssetCache2.h
+++ b/Engine/Plugins/Importers/USDImporter/Source/USDClasses/Public/USDAssetCache2.h
@@ -4,7 +4,6 @@
 
 #include "Containers/LruCache.h"
 #include "Containers/Map.h"
-#include "CoreMinimal.h"
 #include "Serialization/BulkData.h"
 #include "UObject/ObjectKey.h"
 
@@ -13,7 +12,7 @@
 /**
  * Owns the assets generated and reused by USD Stages, allowing thread-safe retrieval/storage.
  */
-UCLASS(BlueprintType, Blueprintable, meta=(DisplayName="USD Asset Cache", ScriptName="UsdAssetCache"))
+UCLASS(BlueprintType, Blueprintable, meta = (DisplayName = "USD Asset Cache", ScriptName = "UsdAssetCache"))
 class USDCLASSES_API UUsdAssetCache2 : public UObject
 {
 	GENERATED_BODY()
@@ -220,7 +219,7 @@ private:
 		TSet<FSoftObjectPath> Consumers;
 
 		// Transient stuff
-		ECacheStorageType CurrentStorageType = ECacheStorageType::None;  // Only used internally during RefreshStorage()
+		ECacheStorageType CurrentStorageType = ECacheStorageType::None;	   // Only used internally during RefreshStorage()
 		TSet<FObjectKey> Referencers;
 
 	public:
@@ -248,7 +247,7 @@ private:
 	// When this is set to something, we will track that it is referencing any new asset that we cache.
 	// See FUsdScopedAssetCacheReferencer just below
 	const UObject* CurrentScopedReferencer = nullptr;
-};
+};	  // UCLASS(BlueprintType)
 
 /**
  * The UUsdAssetCache2 can track all the UObjects that are referencing assets, so that it knows when to discard an
@@ -275,4 +274,4 @@ public:
 
 private:
 	TWeakObjectPtr<UUsdAssetCache2> AssetCache;
-};
\ No newline at end of file
+};
diff --git a/Engine/Plugins/Importers/USDImporter/Source/USDClasses/Public/USDAssetImportData.h b/Engine/Plugins/Importers/USDImporter/Source/USDClasses/Public/USDAssetImportData.h
--- a/Engine/Plugins/Importers/USDImporter/Source/USDClasses/Public/USDAssetImportData.h
+++ b/Engine/Plugins/Importers/USDImporter/Source/USDClasses/Public/USDAssetImportData.h
@@ -2,7 +2,6 @@
 
 #pragma once
 
-#include "CoreMinimal.h"
 #include "EditorFramework/AssetImportData.h"
 
 #include "USDAssetUserData.h"
@@ -27,35 +26,43 @@ public:
 
 PRAGMA_DISABLE_DEPRECATION_WARNINGS
 UCLASS()
-class UE_DEPRECATED(5.3, "Use USD AssetUserData instead of AssetImportData for USD-specific info") USDCLASSES_API UUsdAnimSequenceAssetImportData : public UUsdAssetImportData
+class UE_DEPRECATED(5.3, "Use USD AssetUserData instead of AssetImportData for USD-specific info") USDCLASSES_API UUsdAnimSequenceAssetImportData
+	: public UUsdAssetImportData
 {
 	GENERATED_BODY()
 
 public:
-	// Offset into the stage in seconds to when this AnimSequence should start playing to match the
-	// skeletal animation in that stage. This is required because UAnimSequences just range from the
-	// first skeletal timeSample to the last, and the first sample is not necessarily the startTimeCode
-	// for the stage. Note that this is wrt. the startTimeCode of the *layer*, and not the composed stage.
-	// This is only used for animating USkeletalMeshComponents via the AUsdStageActor, and it is in
-	// seconds since it will need to drive the UAnimSequence with a time value in the [0, LenghtSeconds] range.
-	// This should be applied *after* any offset/scale conversions on the time coordinate.
+	// Offset into the stage in seconds to when this AnimSequence should start
+	// playing to match the skeletal animation in that stage. This is required
+	// because UAnimSequences just range from the first skeletal timeSample to the
+	// last, and the first sample is not necessarily the startTimeCode for the
+	// stage. Note that this is wrt. the startTimeCode of the *layer*, and not the
+	// composed stage. This is only used for animating USkeletalMeshComponents via
+	// the AUsdStageActor, and it is in seconds since it will need to drive the
+	// UAnimSequence with a time value in the [0, LenghtSeconds] range. This
+	// should be applied *after* any offset/scale conversions on the time
+	// coordinate.
 	UPROPERTY()
 	float LayerStartOffsetSeconds = 0.0f;
 };
 
 /** We assign these to UStaticMeshes or USkeletalMeshes generated from USD */
 UCLASS()
-class UE_DEPRECATED(5.3, "Use USD AssetUserData instead of AssetImportData for USD-specific info") USDCLASSES_API UUsdMeshAssetImportData : public UUsdAssetImportData
+class UE_DEPRECATED(5.3, "Use USD AssetUserData instead of AssetImportData for USD-specific info") USDCLASSES_API UUsdMeshAssetImportData
+	: public UUsdAssetImportData
 {
 	GENERATED_BODY()
 
 public:
 	/**
-	 * Maps from a material slot index of an UStaticMesh or USkeletalMesh to a list of source prims that contain this
-	 * assignment. It can contain multiple prims in case we combine material slots and/or collapse prims
-	 * (e.g. {0: ['/Root/mesh', '/Root/othermesh/geomsubset0', '/Root/othermesh/geomsubset1'] }).
+	 * Maps from a material slot index of an UStaticMesh or USkeletalMesh to a
+	 * list of source prims that contain this assignment. It can contain multiple
+	 * prims in case we combine material slots and/or collapse prims (e.g. {0:
+	 * ['/Root/mesh', '/Root/othermesh/geomsubset0',
+	 * '/Root/othermesh/geomsubset1'] }).
 	 */
 	UPROPERTY()
-	TMap< int32, FUsdPrimPathList > MaterialSlotToPrimPaths;
+	TMap<int32, FUsdPrimPathList> MaterialSlotToPrimPaths;
 };
-PRAGMA_ENABLE_DEPRECATION_WARNINGS
\ No newline at end of file
+
+PRAGMA_ENABLE_DEPRECATION_WARNINGS
diff --git a/Engine/Plugins/Importers/USDImporter/Source/USDClasses/Public/USDProjectSettings.h b/Engine/Plugins/Importers/USDImporter/Source/USDClasses/Public/USDProjectSettings.h
--- a/Engine/Plugins/Importers/USDImporter/Source/USDClasses/Public/USDProjectSettings.h
+++ b/Engine/Plugins/Importers/USDImporter/Source/USDClasses/Public/USDProjectSettings.h
@@ -8,7 +8,7 @@
 
 #include "USDProjectSettings.generated.h"
 
-UENUM( BlueprintType )
+UENUM(BlueprintType)
 enum class EUsdSaveDialogBehavior : uint8
 {
 	NeverSave,
@@ -16,7 +16,7 @@ enum class EUsdSaveDialogBehavior : uint8
 	ShowPrompt
 };
 
-UENUM( BlueprintType )
+UENUM(BlueprintType)
 enum class EUsdEditInInstanceBehavior : uint8
 {
 	Ignore,
@@ -26,7 +26,7 @@ enum class EUsdEditInInstanceBehavior : uint8
 
 // USDImporter and defaultconfig here so this ends up at DefaultUSDImporter.ini in the editor, and is sent to the
 // packaged game as well
-UCLASS(config=USDImporter, defaultconfig, meta=(DisplayName=USDImporter), MinimalAPI)
+UCLASS(config = USDImporter, defaultconfig, meta = (DisplayName = USDImporter), MinimalAPI)
 class UUsdProjectSettings : public UDeveloperSettings
 {
 	GENERATED_BODY()
@@ -44,7 +44,7 @@ public:
 	//
 	// The packaging process cannot use an absolute path and will raise an error if given one
 	// when it tries to concatenate the game content directory path with an absolute path.
-	UPROPERTY( config, EditAnywhere, Category = USD, meta = (RelativeToGameContentDir) )
+	UPROPERTY(config, EditAnywhere, Category = USD, meta = (RelativeToGameContentDir))
 	TArray<FDirectoryPath> AdditionalPluginDirectories;
 
 	// The directories that will be used as the default search path by USD's default resolver
@@ -56,11 +56,11 @@ public:
 	// Note that the default search path must be set before the first invocation of USD's
 	// resolver system, so changing this setting will require a restart of the engine in order
 	// for the new setting to take effect.
-	UPROPERTY( config, EditAnywhere, Category = USD, meta = (ConfigRestartRequired = true) )
+	UPROPERTY(config, EditAnywhere, Category = USD, meta = (ConfigRestartRequired = true))
 	TArray<FDirectoryPath> DefaultResolverSearchPath;
 
 	// Material purposes to show on drop-downs in addition to the standard "preview" and "full"
-	UPROPERTY( config, EditAnywhere, Category = USD )
+	UPROPERTY(config, EditAnywhere, Category = USD)
 	TArray<FName> AdditionalMaterialPurposes;
 
 	/**
@@ -73,39 +73,39 @@ public:
 	UPROPERTY(config, EditAnywhere, Category = "USD|Dialogs")
 	bool bShowCreateDefaultAssetCacheDialog = true;
 
-	UPROPERTY( config, EditAnywhere, Category = "USD|Dialogs" )
+	UPROPERTY(config, EditAnywhere, Category = "USD|Dialogs")
 	bool bShowConfirmationWhenClearingLayers = true;
 
-	UPROPERTY( config, EditAnywhere, Category = "USD|Dialogs" )
+	UPROPERTY(config, EditAnywhere, Category = "USD|Dialogs")
 	bool bShowConfirmationWhenMutingDirtyLayers = true;
 
 	// Whether to show the warning dialog when authoring opinions that could have no effect on the composed stage
-	UPROPERTY( config, EditAnywhere, Category = "USD|Dialogs" )
+	UPROPERTY(config, EditAnywhere, Category = "USD|Dialogs")
 	bool bShowOverriddenOpinionsWarning = true;
 
 	// Whether to show the warning dialog when authoring opinions inside an instance or instance proxy
-	UPROPERTY( config, EditAnywhere, Category = "USD|Dialogs" )
+	UPROPERTY(config, EditAnywhere, Category = "USD|Dialogs")
 	EUsdEditInInstanceBehavior EditInInstanceableBehavior = EUsdEditInInstanceBehavior::ShowPrompt;
 
 	// Whether to show a warning whenever the "Duplicate All Local Layer Specs" option is picked, and the duplicated
 	// prim has some specs outside the local layer stack that will not be duplicated.
-	UPROPERTY( config, EditAnywhere, Category = "USD|Dialogs" )
+	UPROPERTY(config, EditAnywhere, Category = "USD|Dialogs")
 	bool bShowWarningOnIncompleteDuplication = true;
 
 	// Whether to show the warning dialog when authoring a transforms directly to a camera component
-	UPROPERTY( config, EditAnywhere, Category = "USD|Dialogs" )
+	UPROPERTY(config, EditAnywhere, Category = "USD|Dialogs")
 	bool bShowTransformOnCameraComponentWarning = true;
 
 	// Whether to show the warning dialog when authoring a transform track directly to a camera component
-	UPROPERTY( config, EditAnywhere, Category = "USD|Dialogs" )
+	UPROPERTY(config, EditAnywhere, Category = "USD|Dialogs")
 	bool bShowTransformTrackOnCameraComponentWarning = true;
 
 	// Whether to display the pop up dialog asking what to do about dirty USD layers when saving the UE level
-	UPROPERTY(config, EditAnywhere, Category = "USD|Dialogs" )
+	UPROPERTY(config, EditAnywhere, Category = "USD|Dialogs")
 	EUsdSaveDialogBehavior ShowSaveLayersDialogWhenSaving = EUsdSaveDialogBehavior::ShowPrompt;
 
 	// Whether to display the pop up dialog asking what to do about dirty USD layers when closing USD stages
-	UPROPERTY(config, EditAnywhere, Category = "USD|Dialogs" )
+	UPROPERTY(config, EditAnywhere, Category = "USD|Dialogs")
 	EUsdSaveDialogBehavior ShowSaveLayersDialogWhenClosing = EUsdSaveDialogBehavior::ShowPrompt;
 
 	// Note that the below properties being FSoftObjectPath ensure that these materials are cooked into packaged games
@@ -117,8 +117,8 @@ public:
 	 * as a material parameter named "Texture".
 	 * You can swap this with your own material, but make sure the replacement material has a "Texture" parameter
 	 */
-	UPROPERTY( config, EditAnywhere, Category = "USD|Reference Materials", meta = ( AllowedClasses = "/Script/Engine.MaterialInterface" ) )
-	FSoftObjectPath ReferenceModelCardTextureMaterial = FSoftObjectPath{ TEXT( "/USDImporter/Materials/CardTextureMaterial.CardTextureMaterial" ) };
+	UPROPERTY(config, EditAnywhere, Category = "USD|Reference Materials", meta = (AllowedClasses = "/Script/Engine.MaterialInterface"))
+	FSoftObjectPath ReferenceModelCardTextureMaterial = FSoftObjectPath{TEXT("/USDImporter/Materials/CardTextureMaterial.CardTextureMaterial")};
 
 	/**
 	 * What material to use as reference material when creating material instances from USD materials.
@@ -126,8 +126,13 @@ public:
 	 * the same names and types as the ones provided by the default material, otherwise the instances will not have
 	 * the parameters filled with values extracted from the USD material when parsing.
 	 */
-	UPROPERTY( config, EditAnywhere, Category = "USD|Reference Materials|UsdPreviewSurface", meta = ( AllowedClasses = "/Script/Engine.MaterialInterface" ) )
-	FSoftObjectPath ReferencePreviewSurfaceMaterial = FSoftObjectPath{ TEXT( "/USDImporter/Materials/UsdPreviewSurface.UsdPreviewSurface" ) };
+	UPROPERTY(
+		config,
+		EditAnywhere,
+		Category = "USD|Reference Materials|UsdPreviewSurface",
+		meta = (AllowedClasses = "/Script/Engine.MaterialInterface")
+	)
+	FSoftObjectPath ReferencePreviewSurfaceMaterial = FSoftObjectPath{TEXT("/USDImporter/Materials/UsdPreviewSurface.UsdPreviewSurface")};
 
 	/**
 	 * What material to use as reference material when creating material instances from USD materials.
@@ -135,8 +140,14 @@ public:
 	 * the same names and types as the ones provided by the default material, otherwise the instances will not have
 	 * the parameters filled with values extracted from the USD material when parsing.
 	 */
-	UPROPERTY( config, EditAnywhere, Category = "USD|Reference Materials|UsdPreviewSurface", meta = ( AllowedClasses = "/Script/Engine.MaterialInterface" ) )
-	FSoftObjectPath ReferencePreviewSurfaceTranslucentMaterial = FSoftObjectPath{ TEXT( "/USDImporter/Materials/UsdPreviewSurfaceTranslucent.UsdPreviewSurfaceTranslucent" ) };
+	UPROPERTY(
+		config,
+		EditAnywhere,
+		Category = "USD|Reference Materials|UsdPreviewSurface",
+		meta = (AllowedClasses = "/Script/Engine.MaterialInterface")
+	)
+	FSoftObjectPath ReferencePreviewSurfaceTranslucentMaterial = FSoftObjectPath{TEXT("/USDImporter/Materials/"
+																					  "UsdPreviewSurfaceTranslucent.UsdPreviewSurfaceTranslucent")};
 
 	/**
 	 * What material to use as reference material when creating material instances from USD materials.
@@ -144,8 +155,14 @@ public:
 	 * the same names and types as the ones provided by the default material, otherwise the instances will not have
 	 * the parameters filled with values extracted from the USD material when parsing.
 	 */
-	UPROPERTY( config, EditAnywhere, Category = "USD|Reference Materials|UsdPreviewSurface", meta = ( AllowedClasses = "/Script/Engine.MaterialInterface" ) )
-	FSoftObjectPath ReferencePreviewSurfaceTwoSidedMaterial = FSoftObjectPath{ TEXT( "/USDImporter/Materials/UsdPreviewSurfaceTwoSided.UsdPreviewSurfaceTwoSided" ) };
+	UPROPERTY(
+		config,
+		EditAnywhere,
+		Category = "USD|Reference Materials|UsdPreviewSurface",
+		meta = (AllowedClasses = "/Script/Engine.MaterialInterface")
+	)
+	FSoftObjectPath ReferencePreviewSurfaceTwoSidedMaterial = FSoftObjectPath{TEXT("/USDImporter/Materials/"
+																				   "UsdPreviewSurfaceTwoSided.UsdPreviewSurfaceTwoSided")};
 
 	/**
 	 * What material to use as reference material when creating material instances from USD materials.
@@ -153,8 +170,14 @@ public:
 	 * the same names and types as the ones provided by the default material, otherwise the instances will not have
 	 * the parameters filled with values extracted from the USD material when parsing.
 	 */
-	UPROPERTY( config, EditAnywhere, Category = "USD|Reference Materials|UsdPreviewSurface", meta = ( AllowedClasses = "/Script/Engine.MaterialInterface" ) )
-	FSoftObjectPath ReferencePreviewSurfaceTranslucentTwoSidedMaterial = FSoftObjectPath{ TEXT( "/USDImporter/Materials/UsdPreviewSurfaceTranslucentTwoSided.UsdPreviewSurfaceTranslucentTwoSided" ) };
+	UPROPERTY(
+		config,
+		EditAnywhere,
+		Category = "USD|Reference Materials|UsdPreviewSurface",
+		meta = (AllowedClasses = "/Script/Engine.MaterialInterface")
+	)
+	FSoftObjectPath ReferencePreviewSurfaceTranslucentTwoSidedMaterial = FSoftObjectPath{
+		TEXT("/USDImporter/Materials/UsdPreviewSurfaceTranslucentTwoSided.UsdPreviewSurfaceTranslucentTwoSided")};
 
 	/**
 	 * What material to use as reference material when creating material instances from USD materials.
@@ -162,8 +185,13 @@ public:
 	 * the same names and types as the ones provided by the default material, otherwise the instances will not have
 	 * the parameters filled with values extracted from the USD material when parsing.
 	 */
-	UPROPERTY( config, EditAnywhere, Category = "USD|Reference Materials|UsdPreviewSurface with Virtual Textures", meta = ( AllowedClasses = "/Script/Engine.MaterialInterface" ) )
-	FSoftObjectPath ReferencePreviewSurfaceVTMaterial = FSoftObjectPath{ TEXT( "/USDImporter/Materials/UsdPreviewSurfaceVT.UsdPreviewSurfaceVT" ) };
+	UPROPERTY(
+		config,
+		EditAnywhere,
+		Category = "USD|Reference Materials|UsdPreviewSurface with Virtual Textures",
+		meta = (AllowedClasses = "/Script/Engine.MaterialInterface")
+	)
+	FSoftObjectPath ReferencePreviewSurfaceVTMaterial = FSoftObjectPath{TEXT("/USDImporter/Materials/UsdPreviewSurfaceVT.UsdPreviewSurfaceVT")};
 
 	/**
 	 * What material to use as reference material when creating material instances from USD materials.
@@ -171,8 +199,14 @@ public:
 	 * the same names and types as the ones provided by the default material, otherwise the instances will not have
 	 * the parameters filled with values extracted from the USD material when parsing.
 	 */
-	UPROPERTY( config, EditAnywhere, Category = "USD|Reference Materials|UsdPreviewSurface with Virtual Textures", meta = ( AllowedClasses = "/Script/Engine.MaterialInterface" ) )
-	FSoftObjectPath ReferencePreviewSurfaceTranslucentVTMaterial = FSoftObjectPath{ TEXT( "/USDImporter/Materials/UsdPreviewSurfaceTranslucentVT.UsdPreviewSurfaceTranslucentVT" ) };
+	UPROPERTY(
+		config,
+		EditAnywhere,
+		Category = "USD|Reference Materials|UsdPreviewSurface with Virtual Textures",
+		meta = (AllowedClasses = "/Script/Engine.MaterialInterface")
+	)
+	FSoftObjectPath ReferencePreviewSurfaceTranslucentVTMaterial = FSoftObjectPath{
+		TEXT("/USDImporter/Materials/UsdPreviewSurfaceTranslucentVT.UsdPreviewSurfaceTranslucentVT")};
 
 	/**
 	 * What material to use as reference material when creating material instances from USD materials.
@@ -180,8 +214,14 @@ public:
 	 * the same names and types as the ones provided by the default material, otherwise the instances will not have
 	 * the parameters filled with values extracted from the USD material when parsing.
 	 */
-	UPROPERTY( config, EditAnywhere, Category = "USD|Reference Materials|UsdPreviewSurface with Virtual Textures", meta = ( AllowedClasses = "/Script/Engine.MaterialInterface" ) )
-	FSoftObjectPath ReferencePreviewSurfaceTwoSidedVTMaterial = FSoftObjectPath{ TEXT( "/USDImporter/Materials/UsdPreviewSurfaceTwoSidedVT.UsdPreviewSurfaceTwoSidedVT" ) };
+	UPROPERTY(
+		config,
+		EditAnywhere,
+		Category = "USD|Reference Materials|UsdPreviewSurface with Virtual Textures",
+		meta = (AllowedClasses = "/Script/Engine.MaterialInterface")
+	)
+	FSoftObjectPath ReferencePreviewSurfaceTwoSidedVTMaterial = FSoftObjectPath{TEXT("/USDImporter/Materials/"
+																					 "UsdPreviewSurfaceTwoSidedVT.UsdPreviewSurfaceTwoSidedVT")};
 
 	/**
 	 * What material to use as reference material when creating material instances from USD materials.
@@ -189,8 +229,14 @@ public:
 	 * the same names and types as the ones provided by the default material, otherwise the instances will not have
 	 * the parameters filled with values extracted from the USD material when parsing.
 	 */
-	UPROPERTY( config, EditAnywhere, Category = "USD|Reference Materials|UsdPreviewSurface with Virtual Textures", meta = ( AllowedClasses = "/Script/Engine.MaterialInterface" ) )
-	FSoftObjectPath ReferencePreviewSurfaceTranslucentTwoSidedVTMaterial = FSoftObjectPath{ TEXT( "/USDImporter/Materials/UsdPreviewSurfaceTranslucentTwoSidedVT.UsdPreviewSurfaceTranslucentTwoSidedVT" ) };
+	UPROPERTY(
+		config,
+		EditAnywhere,
+		Category = "USD|Reference Materials|UsdPreviewSurface with Virtual Textures",
+		meta = (AllowedClasses = "/Script/Engine.MaterialInterface")
+	)
+	FSoftObjectPath ReferencePreviewSurfaceTranslucentTwoSidedVTMaterial = FSoftObjectPath{
+		TEXT("/USDImporter/Materials/UsdPreviewSurfaceTranslucentTwoSidedVT.UsdPreviewSurfaceTranslucentTwoSidedVT")};
 
 	/**
 	 * What material to use as reference material when creating material instances from USD materials.
@@ -198,8 +244,8 @@ public:
 	 * the same names and types as the ones provided by the default material, otherwise the instances will not have
 	 * the parameters filled with values extracted from the USD material when parsing.
 	 */
-	UPROPERTY( config, EditAnywhere, Category = "USD|Reference Materials|DisplayColor", meta = ( AllowedClasses = "/Script/Engine.MaterialInterface" ) )
-	FSoftObjectPath ReferenceDisplayColorMaterial = FSoftObjectPath{ TEXT( "/USDImporter/Materials/DisplayColor.DisplayColor" ) };
+	UPROPERTY(config, EditAnywhere, Category = "USD|Reference Materials|DisplayColor", meta = (AllowedClasses = "/Script/Engine.MaterialInterface"))
+	FSoftObjectPath ReferenceDisplayColorMaterial = FSoftObjectPath{TEXT("/USDImporter/Materials/DisplayColor.DisplayColor")};
 
 	/**
 	 * What material to use as reference material when creating material instances from USD materials.
@@ -207,8 +253,9 @@ public:
 	 * the same names and types as the ones provided by the default material, otherwise the instances will not have
 	 * the parameters filled with values extracted from the USD material when parsing.
 	 */
-	UPROPERTY( config, EditAnywhere, Category = "USD|Reference Materials|DisplayColor", meta = ( AllowedClasses = "/Script/Engine.MaterialInterface" ) )
-	FSoftObjectPath ReferenceDisplayColorAndOpacityMaterial = FSoftObjectPath{ TEXT( "/USDImporter/Materials/DisplayColorAndOpacity.DisplayColorAndOpacity" ) };
+	UPROPERTY(config, EditAnywhere, Category = "USD|Reference Materials|DisplayColor", meta = (AllowedClasses = "/Script/Engine.MaterialInterface"))
+	FSoftObjectPath ReferenceDisplayColorAndOpacityMaterial = FSoftObjectPath{TEXT("/USDImporter/Materials/"
+																				   "DisplayColorAndOpacity.DisplayColorAndOpacity")};
 
 	/**
 	 * What material to use as reference material when creating material instances from USD materials.
@@ -216,8 +263,8 @@ public:
 	 * the same names and types as the ones provided by the default material, otherwise the instances will not have
 	 * the parameters filled with values extracted from the USD material when parsing.
 	 */
-	UPROPERTY( config, EditAnywhere, Category = "USD|Reference Materials|DisplayColor", meta = ( AllowedClasses = "/Script/Engine.MaterialInterface" ) )
-	FSoftObjectPath ReferenceDisplayColorTwoSidedMaterial = FSoftObjectPath{ TEXT( "/USDImporter/Materials/DisplayColorTwoSided.DisplayColorTwoSided" ) };
+	UPROPERTY(config, EditAnywhere, Category = "USD|Reference Materials|DisplayColor", meta = (AllowedClasses = "/Script/Engine.MaterialInterface"))
+	FSoftObjectPath ReferenceDisplayColorTwoSidedMaterial = FSoftObjectPath{TEXT("/USDImporter/Materials/DisplayColorTwoSided.DisplayColorTwoSided")};
 
 	/**
 	 * What material to use as reference material when creating material instances from USD materials.
@@ -225,6 +272,7 @@ public:
 	 * the same names and types as the ones provided by the default material, otherwise the instances will not have
 	 * the parameters filled with values extracted from the USD material when parsing.
 	 */
-	UPROPERTY( config, EditAnywhere, Category = "USD|Reference Materials|DisplayColor", meta = ( AllowedClasses = "/Script/Engine.MaterialInterface" ) )
-	FSoftObjectPath ReferenceDisplayColorAndOpacityTwoSidedMaterial = FSoftObjectPath{ TEXT( "/USDImporter/Materials/DisplayColorAndOpacityTwoSided.DisplayColorAndOpacityTwoSided" ) };
-};
\ No newline at end of file
+	UPROPERTY(config, EditAnywhere, Category = "USD|Reference Materials|DisplayColor", meta = (AllowedClasses = "/Script/Engine.MaterialInterface"))
+	FSoftObjectPath ReferenceDisplayColorAndOpacityTwoSidedMaterial = FSoftObjectPath{
+		TEXT("/USDImporter/Materials/DisplayColorAndOpacityTwoSided.DisplayColorAndOpacityTwoSided")};
+};
