USD: Reformat USDClassesEditor module.

#rnx

diff --git a/Engine/Plugins/Importers/USDImporter/Source/USDClassesEditor/Private/USDAssetCacheAssetActions.h b/Engine/Plugins/Importers/USDImporter/Source/USDClassesEditor/Private/USDAssetCacheAssetActions.h
--- a/Engine/Plugins/Importers/USDImporter/Source/USDClassesEditor/Private/USDAssetCacheAssetActions.h
+++ b/Engine/Plugins/Importers/USDImporter/Source/USDClassesEditor/Private/USDAssetCacheAssetActions.h
@@ -3,7 +3,6 @@
 #pragma once
 
 #include "AssetTypeActions_Base.h"
-#include "CoreMinimal.h"
 #include "Toolkits/IToolkitHost.h"
 
 class FUsdAssetCacheAssetActions : public FAssetTypeActions_Base
@@ -14,6 +13,7 @@ public:
 	virtual FText GetName() const override;
 	virtual UClass* GetSupportedClass() const override;
 	virtual FColor GetTypeColor() const override;
-	virtual void OpenAssetEditor(const TArray<UObject*>& InObjects, TSharedPtr<class IToolkitHost> EditWithinLevelEditor = TSharedPtr<IToolkitHost>()) override;
+	virtual void OpenAssetEditor(const TArray<UObject*>& InObjects, TSharedPtr<class IToolkitHost> EditWithinLevelEditor = TSharedPtr<IToolkitHost>())
+		override;
 	//~ End IAssetTypeActions interface
-};
\ No newline at end of file
+};
diff --git a/Engine/Plugins/Importers/USDImporter/Source/USDClassesEditor/Private/USDDefaultAssetCacheDialog.cpp b/Engine/Plugins/Importers/USDImporter/Source/USDClassesEditor/Private/USDDefaultAssetCacheDialog.cpp
--- a/Engine/Plugins/Importers/USDImporter/Source/USDClassesEditor/Private/USDDefaultAssetCacheDialog.cpp
+++ b/Engine/Plugins/Importers/USDImporter/Source/USDClassesEditor/Private/USDDefaultAssetCacheDialog.cpp
@@ -81,27 +81,23 @@ namespace UE::DefaultCacheDialog::Private
 					FEditorDirectories::Get().SetLastDirectory(ELastDirectory::NEW_ASSET, SavePackagePath);
 
 					const FName CallingContext = NAME_None;
-					Result = Cast<UUsdAssetCache2>(AssetTools.CreateAsset(
-						SaveAssetName,
-						SavePackagePath,
-						UUsdAssetCache2::StaticClass(),
-						Factory,
-						CallingContext
-					));
+					Result = Cast<UUsdAssetCache2>(
+						AssetTools.CreateAsset(SaveAssetName, SavePackagePath, UUsdAssetCache2::StaticClass(), Factory, CallingContext)
+					);
 				}
 			}
 		}
 
 		return Result;
 	}
-}
+}	 // namespace UE::DefaultCacheDialog::Private
 
 UUsdAssetCache2* SUsdDefaultAssetCacheDialog::GetCreatedCache()
 {
 	return ChosenCache;
 }
 
-void SUsdDefaultAssetCacheDialog::Construct( const FArguments& InArgs )
+void SUsdDefaultAssetCacheDialog::Construct(const FArguments& InArgs)
 {
 	ChosenCache = nullptr;
 	Window = InArgs._WidgetWindow;
@@ -109,6 +105,7 @@ void SUsdDefaultAssetCacheDialog::Construct( const FArguments& InArgs )
 
 	FSlateFontInfo MessageFont(FAppStyle::GetFontStyle("StandardDialog.LargeFont"));
 
+	// clang-format off
 	ChildSlot
 	[
 		SNew(SBorder)
@@ -150,7 +147,7 @@ void SUsdDefaultAssetCacheDialog::Construct( const FArguments& InArgs )
 			[
 				SNew(SHorizontalBox)
 
-				+ SHorizontalBox::Slot()
+				+SHorizontalBox::Slot()
 				.FillWidth(1.0f)
 				.HAlign(HAlign_Left)
 				.VAlign(VAlign_Center)
@@ -173,18 +170,18 @@ void SUsdDefaultAssetCacheDialog::Construct( const FArguments& InArgs )
 					]
 				]
 
-				+ SHorizontalBox::Slot()
+				+SHorizontalBox::Slot()
 				.AutoWidth()
 				.HAlign(HAlign_Right)
 				.VAlign(VAlign_Center)
 				.Padding(16.f, 0.f, 0.f, 0.f)
 				[
-					SNew( SUniformGridPanel )
+					SNew(SUniformGridPanel)
 					.SlotPadding(FAppStyle::Get().GetMargin("StandardDialog.SlotPadding"))
 					.MinDesiredSlotWidth(FAppStyle::Get().GetFloat("StandardDialog.MinDesiredSlotWidth"))
 					.MinDesiredSlotHeight(FAppStyle::Get().GetFloat("StandardDialog.MinDesiredSlotHeight"))
 
-					+ SUniformGridPanel::Slot( 0, 0 )
+					+SUniformGridPanel::Slot(0, 0)
 					[
 						SNew(SButton)
 						.HAlign(HAlign_Center)
@@ -195,7 +192,7 @@ void SUsdDefaultAssetCacheDialog::Construct( const FArguments& InArgs )
 						.ButtonStyle(&FAppStyle::Get(), "PrimaryButton")
 					]
 
-					+ SUniformGridPanel::Slot( 1, 0 )
+					+SUniformGridPanel::Slot(1, 0)
 					[
 						SNew(SButton)
 						.HAlign(HAlign_Center)
@@ -206,7 +203,7 @@ void SUsdDefaultAssetCacheDialog::Construct( const FArguments& InArgs )
 						.ButtonStyle(&FAppStyle::Get(), "Button")
 					]
 
-					+ SUniformGridPanel::Slot( 2, 0 )
+					+SUniformGridPanel::Slot(2, 0)
 					[
 						SNew(SButton)
 						.HAlign(HAlign_Center)
@@ -220,6 +217,7 @@ void SUsdDefaultAssetCacheDialog::Construct( const FArguments& InArgs )
 			]
 		]
 	];
+	// clang-format on
 }
 
 bool SUsdDefaultAssetCacheDialog::SupportsKeyboardFocus() const
@@ -272,9 +270,9 @@ FReply SUsdDefaultAssetCacheDialog::OnDontCreate()
 	return FReply::Handled();
 }
 
-FReply SUsdDefaultAssetCacheDialog::OnKeyDown( const FGeometry& MyGeometry, const FKeyEvent& InKeyEvent )
+FReply SUsdDefaultAssetCacheDialog::OnKeyDown(const FGeometry& MyGeometry, const FKeyEvent& InKeyEvent)
 {
-	if ( InKeyEvent.GetKey() == EKeys::Escape )
+	if (InKeyEvent.GetKey() == EKeys::Escape)
 	{
 		return OnDontCreate();
 	}
@@ -286,4 +284,4 @@ EDefaultAssetCacheDialogOption SUsdDefaultAssetCacheDialog::GetDialogOutcome() c
 	return DialogOutcome;
 }
 
-#undef LOCTEXT_NAMESPACE
\ No newline at end of file
+#undef LOCTEXT_NAMESPACE
