USD: Reformat USDExporter module.

#rnx

diff --git a/Engine/Plugins/Importers/USDImporter/Source/USDExporter/Private/AnimSequenceExporterUSD.cpp b/Engine/Plugins/Importers/USDImporter/Source/USDExporter/Private/AnimSequenceExporterUSD.cpp
--- a/Engine/Plugins/Importers/USDImporter/Source/USDExporter/Private/AnimSequenceExporterUSD.cpp
+++ b/Engine/Plugins/Importers/USDImporter/Source/USDExporter/Private/AnimSequenceExporterUSD.cpp
@@ -23,7 +23,6 @@
 #include "AssetExportTask.h"
 #include "Engine/SkeletalMesh.h"
 #include "EngineAnalytics.h"
-#include "MaterialExporterUSD.h"
 #include "Misc/EngineVersion.h"
 #include "UObject/GCObjectScopeGuard.h"
 
@@ -38,7 +37,7 @@ namespace UE::AnimSequenceExporterUSD::Private
 		const FString& Extension
 	)
 	{
-		if ( !Asset || !FEngineAnalytics::IsAvailable() )
+		if (!Asset || !FEngineAnalytics::IsAvailable())
 		{
 			return;
 		}
@@ -46,16 +45,16 @@ namespace UE::AnimSequenceExporterUSD::Private
 		FString ClassName = Asset->GetClass()->GetName();
 
 		TArray<FAnalyticsEventAttribute> EventAttributes;
-		EventAttributes.Emplace( TEXT( "AssetType" ), ClassName );
+		EventAttributes.Emplace(TEXT("AssetType"), ClassName);
 
-		if ( Options )
+		if (Options)
 		{
-			UsdUtils::AddAnalyticsAttributes( *Options, EventAttributes );
+			UsdUtils::AddAnalyticsAttributes(*Options, EventAttributes);
 		}
 
 		IUsdClassesModule::SendAnalytics(
-			MoveTemp( EventAttributes ),
-			FString::Printf( TEXT( "Export.%s" ), *ClassName ),
+			MoveTemp(EventAttributes),
+			FString::Printf(TEXT("Export.%s"), *ClassName),
 			bAutomated,
 			ElapsedSeconds,
 			NumberOfFrames,
@@ -67,10 +66,10 @@ namespace UE::AnimSequenceExporterUSD::Private
 UAnimSequenceExporterUSD::UAnimSequenceExporterUSD()
 {
 #if USE_USD_SDK
-	for ( const FString& Extension : UnrealUSDWrapper::GetNativeFileFormats() )
+	for (const FString& Extension : UnrealUSDWrapper::GetNativeFileFormats())
 	{
 		// USDZ is not supported for writing for now
-		if ( Extension.Equals( TEXT( "usdz" ) ) )
+		if (Extension.Equals(TEXT("usdz")))
 		{
 			continue;
 		}
@@ -80,14 +79,21 @@ UAnimSequenceExporterUSD::UAnimSequenceExporterUSD()
 	}
 	SupportedClass = UAnimSequence::StaticClass();
 	bText = false;
-#endif // #if USE_USD_SDK
+#endif	  // #if USE_USD_SDK
 }
 
-bool UAnimSequenceExporterUSD::ExportBinary( UObject* Object, const TCHAR* Type, FArchive& Ar, FFeedbackContext* Warn, int32 FileIndex, uint32 PortFlags )
+bool UAnimSequenceExporterUSD::ExportBinary(
+	UObject* Object,
+	const TCHAR* Type,
+	FArchive& Ar,
+	FFeedbackContext* Warn,
+	int32 FileIndex,
+	uint32 PortFlags
+)
 {
 #if USE_USD_SDK
-	UAnimSequence* AnimSequence = CastChecked< UAnimSequence >( Object );
-	if ( !AnimSequence )
+	UAnimSequence* AnimSequence = CastChecked<UAnimSequence>(Object);
+	if (!AnimSequence)
 	{
 		return false;
 	}
@@ -99,32 +105,32 @@ bool UAnimSequenceExporterUSD::ExportBinary( UObject* Object, const TCHAR* Type,
 	const FString AnimSequenceFile = UExporter::CurrentFilename;
 
 	UAnimSequenceExporterUSDOptions* Options = nullptr;
-	if ( ExportTask )
+	if (ExportTask)
 	{
-		Options = Cast<UAnimSequenceExporterUSDOptions>( ExportTask->Options );
+		Options = Cast<UAnimSequenceExporterUSDOptions>(ExportTask->Options);
 	}
-	if ( !Options )
+	if (!Options)
 	{
 		Options = GetMutableDefault<UAnimSequenceExporterUSDOptions>();
 
 		// Prompt with an options dialog if we can
-		if ( Options && ( !ExportTask || !ExportTask->bAutomated ) )
+		if (Options && (!ExportTask || !ExportTask->bAutomated))
 		{
-			Options->PreviewMeshOptions.MaterialBakingOptions.TexturesDir.Path = FPaths::Combine( FPaths::GetPath( AnimSequenceFile ), TEXT( "Textures" ) );
+			Options->PreviewMeshOptions.MaterialBakingOptions.TexturesDir.Path = FPaths::Combine(FPaths::GetPath(AnimSequenceFile), TEXT("Textures"));
 
-			const bool bContinue = SUsdOptionsWindow::ShowExportOptions( *Options );
-			if ( !bContinue )
+			const bool bContinue = SUsdOptionsWindow::ShowExportOptions(*Options);
+			if (!bContinue)
 			{
 				return false;
 			}
 		}
 	}
-	if ( !Options )
+	if (!Options)
 	{
 		return false;
 	}
 
-	if ( !IUsdExporterModule::CanExportToLayer( UExporter::CurrentFilename ) )
+	if (!IUsdExporterModule::CanExportToLayer(UExporter::CurrentFilename))
 	{
 		return false;
 	}
@@ -135,29 +141,29 @@ bool UAnimSequenceExporterUSD::ExportBinary( UObject* Object, const TCHAR* Type,
 	// Export preview mesh if needed
 	USkeletalMesh* SkeletalMesh = nullptr;
 	FString MeshAssetFile;
-	if ( Options && Options->bExportPreviewMesh )
+	if (Options && Options->bExportPreviewMesh)
 	{
 		SkeletalMesh = AnimSequence->GetPreviewMesh();
 		USkeleton* AnimSkeleton = SkeletalMesh ? SkeletalMesh->GetSkeleton() : nullptr;
 
-		if ( !AnimSkeleton && !SkeletalMesh )
+		if (!AnimSkeleton && !SkeletalMesh)
 		{
 			AnimSkeleton = AnimSequence->GetSkeleton();
-			SkeletalMesh = AnimSkeleton ? AnimSkeleton->GetAssetPreviewMesh( AnimSequence ) : nullptr;
+			SkeletalMesh = AnimSkeleton ? AnimSkeleton->GetAssetPreviewMesh(AnimSequence) : nullptr;
 		}
 
-		if ( AnimSkeleton && !SkeletalMesh )
+		if (AnimSkeleton && !SkeletalMesh)
 		{
 			SkeletalMesh = AnimSkeleton->FindCompatibleMesh();
 		}
 
-		if ( SkeletalMesh )
+		if (SkeletalMesh)
 		{
 			FString PathPart;
 			FString FilenamePart;
 			FString ExtensionPart;
-			FPaths::Split( AnimSequenceFile, PathPart, FilenamePart, ExtensionPart );
-			MeshAssetFile = FPaths::Combine( PathPart, FilenamePart + TEXT( "_SkeletalMesh." ) + ExtensionPart );
+			FPaths::Split(AnimSequenceFile, PathPart, FilenamePart, ExtensionPart);
+			MeshAssetFile = FPaths::Combine(PathPart, FilenamePart + TEXT("_SkeletalMesh.") + ExtensionPart);
 
 			USkeletalMeshExporterUSDOptions* SkeletalMeshOptions = GetMutableDefault<USkeletalMeshExporterUSDOptions>();
 			SkeletalMeshOptions->StageOptions = Options->StageOptions;
@@ -166,7 +172,7 @@ bool UAnimSequenceExporterUSD::ExportBinary( UObject* Object, const TCHAR* Type,
 			SkeletalMeshOptions->bReExportIdenticalAssets = Options->bReExportIdenticalAssets;
 
 			UAssetExportTask* LevelExportTask = NewObject<UAssetExportTask>();
-			FGCObjectScopeGuard ExportTaskGuard( LevelExportTask );
+			FGCObjectScopeGuard ExportTaskGuard(LevelExportTask);
 			LevelExportTask->Object = SkeletalMesh;
 			LevelExportTask->Options = SkeletalMeshOptions;
 			LevelExportTask->Exporter = nullptr;
@@ -176,10 +182,10 @@ bool UAnimSequenceExporterUSD::ExportBinary( UObject* Object, const TCHAR* Type,
 			LevelExportTask->bPrompt = false;
 			LevelExportTask->bUseFileArchive = false;
 			LevelExportTask->bWriteEmptyFiles = false;
-			LevelExportTask->bAutomated = true; // Pretend this is an automated task so it doesn't pop the options dialog
+			LevelExportTask->bAutomated = true;	   // Pretend this is an automated task so it doesn't pop the options dialog
 
-			const bool bSucceeded = UExporter::RunAssetExportTask( LevelExportTask );
-			if ( !bSucceeded )
+			const bool bSucceeded = UExporter::RunAssetExportTask(LevelExportTask);
+			if (!bSucceeded)
 			{
 				MeshAssetFile = {};
 			}
@@ -188,80 +194,81 @@ bool UAnimSequenceExporterUSD::ExportBinary( UObject* Object, const TCHAR* Type,
 		{
 			FUsdLogManager::LogMessage(
 				EMessageSeverity::Warning,
-				FText::Format( NSLOCTEXT( "AnimSequenceExporterUSD", "InvalidSkelMesh", "Couldn't find the skeletal mesh to export for anim sequence {0}." ), FText::FromName( AnimSequence->GetFName() ) ) );
+				FText::Format(
+					NSLOCTEXT("AnimSequenceExporterUSD", "InvalidSkelMesh", "Couldn't find the skeletal mesh to export for anim sequence {0}."),
+					FText::FromName(AnimSequence->GetFName())
+				)
+			);
 		}
 	}
 
 	// Collect the target paths for our SkelAnimation prim and its SkelRoot, if any
 	UE::FSdfPath SkelRootPath;
 	UE::FSdfPath SkelAnimPath;
-	if ( MeshAssetFile.IsEmpty() )
+	if (MeshAssetFile.IsEmpty())
 	{
-		SkelAnimPath = UE::FSdfPath::AbsoluteRootPath().AppendChild(
-			*UsdUtils::SanitizeUsdIdentifier( *AnimSequence->GetName() )
-		);
+		SkelAnimPath = UE::FSdfPath::AbsoluteRootPath().AppendChild(*UsdUtils::SanitizeUsdIdentifier(*AnimSequence->GetName()));
 	}
 	else
 	{
-		SkelRootPath = UE::FSdfPath::AbsoluteRootPath().AppendChild(
-			*UsdUtils::SanitizeUsdIdentifier( *SkeletalMesh->GetName() )
-		);
+		SkelRootPath = UE::FSdfPath::AbsoluteRootPath().AppendChild(*UsdUtils::SanitizeUsdIdentifier(*SkeletalMesh->GetName()));
 
-		SkelAnimPath = SkelRootPath.AppendChild(
-			*UsdUtils::SanitizeUsdIdentifier( *AnimSequence->GetName() )
-		);
+		SkelAnimPath = SkelRootPath.AppendChild(*UsdUtils::SanitizeUsdIdentifier(*AnimSequence->GetName()));
 	}
 
 	FString AnimSequenceVersion;
-	if ( IAnimationDataModel* DataModel = AnimSequence->GetDataModel() )
+	if (IAnimationDataModel* DataModel = AnimSequence->GetDataModel())
 	{
 		FSHA1 SHA1;
 
 		FGuid DataModelGuid = DataModel->GenerateGuid();
-		SHA1.Update( reinterpret_cast< uint8* >( &DataModelGuid ), sizeof( DataModelGuid ) );
+		SHA1.Update(reinterpret_cast<uint8*>(&DataModelGuid), sizeof(DataModelGuid));
 
-		UsdUtils::HashForAnimSequenceExport( *Options, SHA1 );
+		UsdUtils::HashForAnimSequenceExport(*Options, SHA1);
 
 		SHA1.Final();
 		FSHAHash Hash;
-		SHA1.GetHash( &Hash.Hash[ 0 ] );
+		SHA1.GetHash(&Hash.Hash[0]);
 		AnimSequenceVersion = Hash.ToString();
 	}
 
 	// Check if we already have exported what we plan on exporting anyway
-	if ( FPaths::FileExists( AnimSequenceFile ) && !AnimSequenceVersion.IsEmpty() )
+	if (FPaths::FileExists(AnimSequenceFile) && !AnimSequenceVersion.IsEmpty())
 	{
-		if ( !ExportTask->bReplaceIdentical )
+		if (!ExportTask->bReplaceIdentical)
 		{
-			UE_LOG( LogUsd, Log,
-				TEXT( "Skipping export of asset '%s' as the target file '%s' already exists." ),
+			UE_LOG(
+				LogUsd,
+				Log,
+				TEXT("Skipping export of asset '%s' as the target file '%s' already exists."),
 				*Object->GetPathName(),
 				*AnimSequenceFile
 			);
 			return false;
 		}
 		// If we don't want to re-export this asset we need to check if its the same version
-		else if ( !Options->bReExportIdenticalAssets )
+		else if (!Options->bReExportIdenticalAssets)
 		{
 			// Don't use the stage cache here as we want this stage to close within this scope in case
 			// we have to overwrite its files due to e.g. missing payload or anything like that
 			const bool bUseStageCache = false;
 			const EUsdInitialLoadSet InitialLoadSet = EUsdInitialLoadSet::LoadNone;
-			if ( UE::FUsdStage TempStage = UnrealUSDWrapper::OpenStage( *AnimSequenceFile, InitialLoadSet, bUseStageCache ) )
+			if (UE::FUsdStage TempStage = UnrealUSDWrapper::OpenStage(*AnimSequenceFile, InitialLoadSet, bUseStageCache))
 			{
-				if ( UE::FUsdPrim SkelAnimPrim = TempStage.GetPrimAtPath( SkelAnimPath ) )
+				if (UE::FUsdPrim SkelAnimPrim = TempStage.GetPrimAtPath(SkelAnimPath))
 				{
-					FUsdUnrealAssetInfo Info = UsdUtils::GetPrimAssetInfo( SkelAnimPrim );
+					FUsdUnrealAssetInfo Info = UsdUtils::GetPrimAssetInfo(SkelAnimPrim);
 
 					const bool bVersionMatches = !Info.Version.IsEmpty() && Info.Version == AnimSequenceVersion;
 
-					const bool bAssetTypeMatches = !Info.UnrealAssetType.IsEmpty()
-						&& Info.UnrealAssetType == Object->GetClass()->GetName();
+					const bool bAssetTypeMatches = !Info.UnrealAssetType.IsEmpty() && Info.UnrealAssetType == Object->GetClass()->GetName();
 
-					if ( bVersionMatches && bAssetTypeMatches )
+					if (bVersionMatches && bAssetTypeMatches)
 					{
-						UE_LOG( LogUsd, Log,
-							TEXT( "Skipping export of asset '%s' as the target file '%s' already contains up-to-date exported data." ),
+						UE_LOG(
+							LogUsd,
+							Log,
+							TEXT("Skipping export of asset '%s' as the target file '%s' already contains up-to-date exported data."),
 							*Object->GetPathName(),
 							*AnimSequenceFile
 						);
@@ -274,8 +281,8 @@ bool UAnimSequenceExporterUSD::ExportBinary( UObject* Object, const TCHAR* Type,
 
 	double StartTime = FPlatformTime::Cycles64();
 
-	UE::FUsdStage AnimationStage = UnrealUSDWrapper::NewStage( *AnimSequenceFile );
-	if ( !AnimationStage )
+	UE::FUsdStage AnimationStage = UnrealUSDWrapper::NewStage(*AnimSequenceFile);
+	if (!AnimationStage)
 	{
 		return false;
 	}
@@ -284,27 +291,27 @@ bool UAnimSequenceExporterUSD::ExportBinary( UObject* Object, const TCHAR* Type,
 	UE::FUsdPrim SkelAnimPrim;
 
 	// Haven't exported the SkeletalMesh, just make a stage with a SkelAnimation prim
-	if( MeshAssetFile.IsEmpty() )
+	if (MeshAssetFile.IsEmpty())
 	{
-		SkelAnimPrim = AnimationStage.DefinePrim( SkelAnimPath, TEXT( "SkelAnimation" ) );
-		if ( !SkelAnimPrim )
+		SkelAnimPrim = AnimationStage.DefinePrim(SkelAnimPath, TEXT("SkelAnimation"));
+		if (!SkelAnimPrim)
 		{
 			return false;
 		}
 
-		AnimationStage.SetDefaultPrim( SkelAnimPrim );
+		AnimationStage.SetDefaultPrim(SkelAnimPrim);
 	}
 	// Exported a SkeletalMesh prim elsewhere, create a SkelRoot containing this SkelAnimation prim
 	else
 	{
-		SkelRootPrim = AnimationStage.DefinePrim( SkelRootPath, TEXT( "SkelRoot" ) );
-		if ( !SkelRootPrim )
+		SkelRootPrim = AnimationStage.DefinePrim(SkelRootPath, TEXT("SkelRoot"));
+		if (!SkelRootPrim)
 		{
 			return false;
 		}
 
-		SkelAnimPrim = AnimationStage.DefinePrim( SkelAnimPath, TEXT( "SkelAnimation" ) );
-		if ( !SkelAnimPrim )
+		SkelAnimPrim = AnimationStage.DefinePrim(SkelAnimPath, TEXT("SkelAnimation"));
+		if (!SkelAnimPrim)
 		{
 			return false;
 		}
@@ -329,20 +336,20 @@ bool UAnimSequenceExporterUSD::ExportBinary( UObject* Object, const TCHAR* Type,
 
 	// Configure stage metadata
 	{
-		if ( Options )
+		if (Options)
 		{
-			UsdUtils::SetUsdStageMetersPerUnit( AnimationStage, Options->StageOptions.MetersPerUnit );
-			UsdUtils::SetUsdStageUpAxis( AnimationStage, Options->StageOptions.UpAxis );
+			UsdUtils::SetUsdStageMetersPerUnit(AnimationStage, Options->StageOptions.MetersPerUnit);
+			UsdUtils::SetUsdStageUpAxis(AnimationStage, Options->StageOptions.UpAxis);
 		}
 
 		const double StartTimeCode = 0.0;
 		const double EndTimeCode = AnimSequence->GetNumberOfSampledKeys() - 1;
-		UsdUtils::AddTimeCodeRangeToLayer( AnimationStage.GetRootLayer(), StartTimeCode, EndTimeCode );
+		UsdUtils::AddTimeCodeRangeToLayer(AnimationStage.GetRootLayer(), StartTimeCode, EndTimeCode);
 
-		AnimationStage.SetTimeCodesPerSecond( AnimSequence->GetSamplingFrameRate().AsDecimal() );
+		AnimationStage.SetTimeCodesPerSecond(AnimSequence->GetSamplingFrameRate().AsDecimal());
 	}
 
-	UnrealToUsd::ConvertAnimSequence( AnimSequence, SkelAnimPrim );
+	UnrealToUsd::ConvertAnimSequence(AnimSequence, SkelAnimPrim);
 
 	if (Options->MetadataOptions.bExportAssetInfo)
 	{
@@ -355,7 +362,7 @@ bool UAnimSequenceExporterUSD::ExportBinary( UObject* Object, const TCHAR* Type,
 		Info.UnrealExportTime = FDateTime::Now().ToString();
 		Info.UnrealEngineVersion = FEngineVersion::Current().ToString();
 
-		UsdUtils::SetPrimAssetInfo( SkelAnimPrim, Info );
+		UsdUtils::SetPrimAssetInfo(SkelAnimPrim, Info);
 	}
 
 	if (Options->MetadataOptions.bExportAssetMetadata)
@@ -376,15 +383,15 @@ bool UAnimSequenceExporterUSD::ExportBinary( UObject* Object, const TCHAR* Type,
 	// Analytics
 	{
 		bool bAutomated = ExportTask ? ExportTask->bAutomated : false;
-		double ElapsedSeconds = FPlatformTime::ToSeconds64( FPlatformTime::Cycles64() - StartTime );
-		FString Extension = FPaths::GetExtension( AnimSequenceFile );
+		double ElapsedSeconds = FPlatformTime::ToSeconds64(FPlatformTime::Cycles64() - StartTime);
+		FString Extension = FPaths::GetExtension(AnimSequenceFile);
 
 		UE::AnimSequenceExporterUSD::Private::SendAnalytics(
 			Object,
 			Options,
 			bAutomated,
 			ElapsedSeconds,
-			UsdUtils::GetUsdStageNumFrames( AnimationStage ),
+			UsdUtils::GetUsdStageNumFrames(AnimationStage),
 			Extension
 		);
 	}
@@ -392,5 +399,5 @@ bool UAnimSequenceExporterUSD::ExportBinary( UObject* Object, const TCHAR* Type,
 	return true;
 #else
 	return false;
-#endif // #if USE_USD_SDK
+#endif	  // #if USE_USD_SDK
 }
diff --git a/Engine/Plugins/Importers/USDImporter/Source/USDExporter/Private/LevelExporterUSD.cpp b/Engine/Plugins/Importers/USDImporter/Source/USDExporter/Private/LevelExporterUSD.cpp
--- a/Engine/Plugins/Importers/USDImporter/Source/USDExporter/Private/LevelExporterUSD.cpp
+++ b/Engine/Plugins/Importers/USDImporter/Source/USDExporter/Private/LevelExporterUSD.cpp
@@ -4,29 +4,19 @@
 
 #include "LevelExporterUSDOptions.h"
 #include "UnrealUSDWrapper.h"
-#include "USDErrorUtils.h"
-#include "USDLayerUtils.h"
 #include "USDOptionsWindow.h"
-#include "USDSkeletalDataConversion.h"
-
-#include "UsdWrappers/SdfLayer.h"
-#include "UsdWrappers/SdfPath.h"
-#include "UsdWrappers/UsdPrim.h"
-#include "UsdWrappers/UsdStage.h"
 
 #include "AssetExportTask.h"
-#include "EditorLevelUtils.h"
-#include "Engine/LevelStreaming.h"
 #include "Engine/World.h"
 #include "IPythonScriptPlugin.h"
 
 ULevelExporterUSD::ULevelExporterUSD()
 {
 #if USE_USD_SDK
-	for ( const FString& Extension : UnrealUSDWrapper::GetNativeFileFormats() )
+	for (const FString& Extension : UnrealUSDWrapper::GetNativeFileFormats())
 	{
 		// USDZ is not supported for writing for now
-		if ( Extension.Equals( TEXT( "usdz" ) ) )
+		if (Extension.Equals(TEXT("usdz")))
 		{
 			continue;
 		}
@@ -36,34 +26,34 @@ ULevelExporterUSD::ULevelExporterUSD()
 	}
 	SupportedClass = UWorld::StaticClass();
 	bText = false;
-#endif // #if USE_USD_SDK
+#endif	  // #if USE_USD_SDK
 }
 
-bool ULevelExporterUSD::ExportBinary( UObject* Object, const TCHAR* Type, FArchive& Ar, FFeedbackContext* Warn, int32 FileIndex, uint32 PortFlags )
+bool ULevelExporterUSD::ExportBinary(UObject* Object, const TCHAR* Type, FArchive& Ar, FFeedbackContext* Warn, int32 FileIndex, uint32 PortFlags)
 {
 #if USE_USD_SDK
-	UWorld* World = Cast< UWorld >( Object );
-	if ( !World )
+	UWorld* World = Cast<UWorld>(Object);
+	if (!World)
 	{
 		return false;
 	}
 
 	ULevelExporterUSDOptions* Options = nullptr;
-	if ( ExportTask )
+	if (ExportTask)
 	{
-		Options = Cast<ULevelExporterUSDOptions>( ExportTask->Options );
+		Options = Cast<ULevelExporterUSDOptions>(ExportTask->Options);
 	}
 
-	if ( !Options )
+	if (!Options)
 	{
 		Options = GetMutableDefault<ULevelExporterUSDOptions>();
 	}
-	if ( !Options )
+	if (!Options)
 	{
 		return false;
 	}
 
-	if ( ExportTask )
+	if (ExportTask)
 	{
 		// Set this early so that we can use the Object from it (the world/level to export) when showing
 		// the level export options dialog.
@@ -76,12 +66,12 @@ bool ULevelExporterUSD::ExportBinary( UObject* Object, const TCHAR* Type, FArchi
 	}
 
 	// Manual export: Show dialog options
-	if ( !ExportTask || !ExportTask->bAutomated )
+	if (!ExportTask || !ExportTask->bAutomated)
 	{
-		Options->Inner.AssetFolder.Path = FPaths::Combine( FPaths::GetPath( UExporter::CurrentFilename ), TEXT( "Assets" ) );
+		Options->Inner.AssetFolder.Path = FPaths::Combine(FPaths::GetPath(UExporter::CurrentFilename), TEXT("Assets"));
 
-		const bool bContinue = SUsdOptionsWindow::ShowExportOptions( *Options );
-		if ( !bContinue )
+		const bool bContinue = SUsdOptionsWindow::ShowExportOptions(*Options);
+		if (!bContinue)
 		{
 			return false;
 		}
@@ -89,14 +79,14 @@ bool ULevelExporterUSD::ExportBinary( UObject* Object, const TCHAR* Type, FArchi
 
 	// Note how we don't explicitly pass the Options down to Python here: We stash our desired export options on the CDO, and
 	// those are read from Python by executing export_with_cdo_options().
-	if ( IPythonScriptPlugin::Get()->IsPythonAvailable() )
+	if (IPythonScriptPlugin::Get()->IsPythonAvailable())
 	{
-		IPythonScriptPlugin::Get()->ExecPythonCommand( TEXT( "import usd_unreal.level_exporter; usd_unreal.level_exporter.export_with_cdo_options()" ) );
+		IPythonScriptPlugin::Get()->ExecPythonCommand(TEXT("import usd_unreal.level_exporter; usd_unreal.level_exporter.export_with_cdo_options()"));
 	}
 	Options->CurrentTask = nullptr;
 
 	return true;
 #else
 	return false;
-#endif // #if USE_USD_SDK
+#endif	  // #if USE_USD_SDK
 }
diff --git a/Engine/Plugins/Importers/USDImporter/Source/USDExporter/Private/LevelExporterUSDOptionsCustomization.cpp b/Engine/Plugins/Importers/USDImporter/Source/USDExporter/Private/LevelExporterUSDOptionsCustomization.cpp
--- a/Engine/Plugins/Importers/USDImporter/Source/USDExporter/Private/LevelExporterUSDOptionsCustomization.cpp
+++ b/Engine/Plugins/Importers/USDImporter/Source/USDExporter/Private/LevelExporterUSDOptionsCustomization.cpp
@@ -2,9 +2,9 @@
 
 #include "LevelExporterUSDOptionsCustomization.h"
 
-#include "USDClassesModule.h"
-#include "LevelSequenceExporterUSDOptions.h"
 #include "LevelExporterUSDOptions.h"
+#include "LevelSequenceExporterUSDOptions.h"
+#include "USDClassesModule.h"
 
 #include "Brushes/SlateColorBrush.h"
 #include "CoreMinimal.h"
@@ -21,136 +21,140 @@
 
 #define LOCTEXT_NAMESPACE "LevelExporterUSDOptionsCustomization"
 
-namespace UE
+namespace UE::LevelExporterOptions::Private
 {
-	namespace LevelExporterOptions
+	class SLevelPickerRow : public STableRow<TSharedRef<FString>>
 	{
-		namespace Private
+	public:
+		SLATE_BEGIN_ARGS(SLevelPickerRow)
 		{
-			class SLevelPickerRow : public STableRow< TSharedRef<FString> >
-			{
-			public:
-				SLATE_BEGIN_ARGS(SLevelPickerRow) {}
-				SLATE_END_ARGS()
-
-				void Construct( const FArguments& InArgs, const TSharedRef<STableViewBase>& OwnerTableView, TWeakPtr<FString> InEntry, FLevelExporterUSDOptionsInner* Inner )
-				{
-					STableRow::Construct( STableRow::FArguments(), OwnerTableView );
+		}
+		SLATE_END_ARGS()
+
+		void Construct(
+			const FArguments& InArgs,
+			const TSharedRef<STableViewBase>& OwnerTableView,
+			TWeakPtr<FString> InEntry,
+			FLevelExporterUSDOptionsInner* Inner
+		)
+		{
+			STableRow::Construct(STableRow::FArguments(), OwnerTableView);
 
-					SetBorderBackgroundColor( FLinearColor::Transparent );
+			SetBorderBackgroundColor(FLinearColor::Transparent);
 
-					FString LevelName;
-					if ( TSharedPtr<FString> PinnedEntry = InEntry.Pin() )
-					{
-						LevelName = *PinnedEntry;
-					}
+			FString LevelName;
+			if (TSharedPtr<FString> PinnedEntry = InEntry.Pin())
+			{
+				LevelName = *PinnedEntry;
+			}
 
+			// clang-format off
 					SetRowContent(
-						SNew( SHorizontalBox )
+						SNew(SHorizontalBox)
 
-						+ SHorizontalBox::Slot()
-						.HAlign( HAlign_Left )
-						.VAlign( VAlign_Center )
+						+SHorizontalBox::Slot()
+						.HAlign(HAlign_Left)
+						.VAlign(VAlign_Center)
 						.MaxWidth(20)
 						[
-							SNew( SCheckBox )
-							.IsChecked_Lambda( [ LevelName, Inner ]()
+							SNew(SCheckBox)
+							.IsChecked_Lambda([LevelName, Inner]()
 							{
-								if ( Inner )
+								if (Inner)
 								{
-									return Inner->LevelsToIgnore.Contains( LevelName )
+									return Inner->LevelsToIgnore.Contains(LevelName)
 										? ECheckBoxState::Unchecked
 										: ECheckBoxState::Checked;
 								}
 
 								return ECheckBoxState::Undetermined;
 							})
-							.OnCheckStateChanged_Lambda( [ LevelName, Inner ]( ECheckBoxState State )
+							.OnCheckStateChanged_Lambda([LevelName, Inner](ECheckBoxState State)
 							{
-								if ( Inner )
+								if (Inner)
 								{
-									if ( State == ECheckBoxState::Checked )
+									if (State == ECheckBoxState::Checked)
 									{
-										Inner->LevelsToIgnore.Remove( LevelName );
+										Inner->LevelsToIgnore.Remove(LevelName);
 									}
-									else if ( State == ECheckBoxState::Unchecked )
+									else if (State == ECheckBoxState::Unchecked)
 									{
-										Inner->LevelsToIgnore.Add( LevelName );
+										Inner->LevelsToIgnore.Add(LevelName);
 									}
 								}
 							})
 						]
 
-						+ SHorizontalBox::Slot()
-						.HAlign( HAlign_Left )
-						.VAlign( VAlign_Center )
+						+SHorizontalBox::Slot()
+						.HAlign(HAlign_Left)
+						.VAlign(VAlign_Center)
 						[
-							SNew( STextBlock )
-							.Text( FText::FromString( LevelName ) )
-							.Font( FAppStyle::GetFontStyle( TEXT( "PropertyWindow.NormalFont" ) ) )
+							SNew(STextBlock)
+							.Text(FText::FromString(LevelName))
+							.Font(FAppStyle::GetFontStyle(TEXT("PropertyWindow.NormalFont")))
 						]
 					);
+			// clang-format on
+		}
+	};
+
+	class SLevelPickerList : public SListView<TSharedRef<FString>>
+	{
+	public:
+		void Construct(const FArguments& InArgs, FLevelExporterUSDOptionsInner* Inner, UWorld* WorldToExport)
+		{
+			if (!Inner || !WorldToExport)
+			{
+				return;
+			}
+
+			if (ULevel* PersistentLevel = WorldToExport->PersistentLevel)
+			{
+				const FString LevelName = TEXT("Persistent Level");
+				RootItems.Add(MakeShared<FString>(LevelName));
+				if (!PersistentLevel->bIsVisible)
+				{
+					Inner->LevelsToIgnore.Add(LevelName);
 				}
-			};
+			}
 
-			class SLevelPickerList : public SListView<TSharedRef<FString>>
+			for (ULevelStreaming* StreamingLevel : WorldToExport->GetStreamingLevels())
 			{
-			public:
-				void Construct( const FArguments& InArgs, FLevelExporterUSDOptionsInner* Inner, UWorld* WorldToExport )
+				if (StreamingLevel)
 				{
-					if ( !Inner || !WorldToExport )
-					{
-						return;
-					}
+					const FString LevelName = FPaths::GetBaseFilename(StreamingLevel->GetWorldAssetPackageName());
 
-					if ( ULevel* PersistentLevel = WorldToExport->PersistentLevel )
-					{
-						const FString LevelName = TEXT( "Persistent Level" );
-						RootItems.Add( MakeShared< FString >( LevelName ) );
-						if ( !PersistentLevel->bIsVisible )
-						{
-							Inner->LevelsToIgnore.Add( LevelName );
-						}
-					}
+					RootItems.Add(MakeShared<FString>(*LevelName));
 
-					for ( ULevelStreaming* StreamingLevel : WorldToExport->GetStreamingLevels() )
+					if (!StreamingLevel->GetShouldBeVisibleInEditor())
 					{
-						if ( StreamingLevel )
-						{
-							const FString LevelName = FPaths::GetBaseFilename( StreamingLevel->GetWorldAssetPackageName() );
-
-							RootItems.Add( MakeShared< FString >( *LevelName ) );
-
-							if ( !StreamingLevel->GetShouldBeVisibleInEditor() )
-							{
-								Inner->LevelsToIgnore.Add( LevelName );
-							}
-						}
+						Inner->LevelsToIgnore.Add(LevelName);
 					}
-
-					SListView::Construct
-					(
-						SListView::FArguments()
-						.ListItemsSource( &RootItems )
-						.SelectionMode( ESelectionMode::None )
-						.OnGenerateRow( this, &SLevelPickerList::OnGenerateRow, Inner )
-					);
-
-					static FSlateColorBrush TransparentBrush{ FLinearColor::Transparent };
-					SetBackgroundBrush( &TransparentBrush );
 				}
+			}
 
-			private:
-				TSharedRef< ITableRow > OnGenerateRow( TSharedRef<FString> InEntry, const TSharedRef<STableViewBase>& OwnerTable, FLevelExporterUSDOptionsInner* Inner ) const
-				{
-					return SNew( SLevelPickerRow, OwnerTable, InEntry, Inner );
-				}
+			SListView::Construct(SListView::FArguments()
+									 .ListItemsSource(&RootItems)
+									 .SelectionMode(ESelectionMode::None)
+									 .OnGenerateRow(this, &SLevelPickerList::OnGenerateRow, Inner));
 
-				TArray< TSharedRef<FString> > RootItems;
-			};
+			static FSlateColorBrush TransparentBrush{FLinearColor::Transparent};
+			SetBackgroundBrush(&TransparentBrush);
 		}
-	}
-}
+
+	private:
+		TSharedRef<ITableRow> OnGenerateRow(
+			TSharedRef<FString> InEntry,
+			const TSharedRef<STableViewBase>& OwnerTable,
+			FLevelExporterUSDOptionsInner* Inner
+		) const
+		{
+			return SNew(SLevelPickerRow, OwnerTable, InEntry, Inner);
+		}
+
+		TArray<TSharedRef<FString>> RootItems;
+	};
+}	 // namespace UE::LevelExporterOptions::Private
 namespace LevelExporterUSDImpl = UE::LevelExporterOptions::Private;
 
 TSharedRef<IDetailCustomization> FLevelExporterUSDOptionsCustomization::MakeInstance()
@@ -161,12 +165,12 @@ TSharedRef<IDetailCustomization> FLevelExporterUSDOptionsCustomization::MakeInst
 void FLevelExporterUSDOptionsCustomization::CustomizeDetails(IDetailLayoutBuilder& DetailLayoutBuilder)
 {
 	TArray<TWeakObjectPtr<UObject>> SelectedObjects = DetailLayoutBuilder.GetSelectedObjects();
-	if ( SelectedObjects.Num() != 1 )
+	if (SelectedObjects.Num() != 1)
 	{
 		return;
 	}
 
-	TSharedPtr< LevelExporterUSDImpl::SLevelPickerList > PickerTree = nullptr;
+	TSharedPtr<LevelExporterUSDImpl::SLevelPickerList> PickerTree = nullptr;
 	TStrongObjectPtr<UObject> OptionsPtr;
 	FName LevelFilterPropName;
 	FName ExportSublayersPropName;
@@ -174,86 +178,92 @@ void FLevelExporterUSDOptionsCustomization::CustomizeDetails(IDetailLayoutBuilde
 	FName SublayersCategoryName;
 	TAttribute<bool> SublayersEditCondition;
 
-	if ( ULevelExporterUSDOptions* Options = Cast< ULevelExporterUSDOptions>( SelectedObjects[ 0 ].Get() ) )
+	if (ULevelExporterUSDOptions* Options = Cast<ULevelExporterUSDOptions>(SelectedObjects[0].Get()))
 	{
-		OptionsPtr.Reset( Options );
+		OptionsPtr.Reset(Options);
 
 		UAssetExportTask* Task = Options->CurrentTask.Get();
 
-		UWorld* World = Task
-			? Cast<UWorld>( Task->Object )
-			: GEditor->GetEditorWorldContext().World();
+		UWorld* World = Task ? Cast<UWorld>(Task->Object) : GEditor->GetEditorWorldContext().World();
 
-		PickerTree = SNew( LevelExporterUSDImpl::SLevelPickerList, &Options->Inner, World );
-		LevelFilterPropName = TEXT( "Inner.LevelsToIgnore" );
-		ExportSublayersPropName = TEXT( "Inner.bExportSublayers" );
-		TexturesDirPropName = TEXT( "Inner.AssetOptions.MaterialBakingOptions.TexturesDir" );
-		SublayersCategoryName = TEXT( "Sublayers" );
+		PickerTree = SNew(LevelExporterUSDImpl::SLevelPickerList, &Options->Inner, World);
+		LevelFilterPropName = TEXT("Inner.LevelsToIgnore");
+		ExportSublayersPropName = TEXT("Inner.bExportSublayers");
+		TexturesDirPropName = TEXT("Inner.AssetOptions.MaterialBakingOptions.TexturesDir");
+		SublayersCategoryName = TEXT("Sublayers");
 
 		SublayersEditCondition = true;
 	}
-	else if ( ULevelSequenceExporterUsdOptions* LevelSequenceOptions = Cast< ULevelSequenceExporterUsdOptions>( SelectedObjects[ 0 ].Get() ) )
+	else if (ULevelSequenceExporterUsdOptions* LevelSequenceOptions = Cast<ULevelSequenceExporterUsdOptions>(SelectedObjects[0].Get()))
 	{
-		OptionsPtr.Reset( LevelSequenceOptions );
+		OptionsPtr.Reset(LevelSequenceOptions);
 
 		// For now there is no easy way of fetching the level to export from a ULevelSequence... we could potentially try to guess what it is
 		// by looking at the soft object paths, but even those aren't exposed, so here we just default to using the current level as the export level.
-		if ( LevelSequenceOptions->Level == nullptr )
+		if (LevelSequenceOptions->Level == nullptr)
 		{
 			const bool bEditorWorldsOnly = true;
-			LevelSequenceOptions->Level = IUsdClassesModule::GetCurrentWorld( bEditorWorldsOnly );
+			LevelSequenceOptions->Level = IUsdClassesModule::GetCurrentWorld(bEditorWorldsOnly);
 		}
 
-		PickerTree = SNew( LevelExporterUSDImpl::SLevelPickerList, &LevelSequenceOptions->LevelExportOptions, LevelSequenceOptions->Level->GetWorld() );
-		LevelFilterPropName = TEXT( "LevelExportOptions.LevelsToIgnore" );
-		ExportSublayersPropName = TEXT( "LevelExportOptions.bExportSublayers" );
-		TexturesDirPropName = TEXT( "LevelExportOptions.AssetOptions.MaterialBakingOptions.TexturesDir" );
-		SublayersCategoryName = TEXT( "Level Export" );
+		PickerTree = SNew(LevelExporterUSDImpl::SLevelPickerList, &LevelSequenceOptions->LevelExportOptions, LevelSequenceOptions->Level->GetWorld());
+		LevelFilterPropName = TEXT("LevelExportOptions.LevelsToIgnore");
+		ExportSublayersPropName = TEXT("LevelExportOptions.bExportSublayers");
+		TexturesDirPropName = TEXT("LevelExportOptions.AssetOptions.MaterialBakingOptions.TexturesDir");
+		SublayersCategoryName = TEXT("Level Export");
 
 		// Refresh the dialog whenever we pick a new world to export, so that we can show this world's sublevels in the sublevel picker
-		FSimpleDelegate RebuildDisplayDelegate = FSimpleDelegate::CreateLambda( [&DetailLayoutBuilder]()
-		{
-			DetailLayoutBuilder.ForceRefreshDetails();
-		});
-		TSharedRef<IPropertyHandle> LevelToExportProp = DetailLayoutBuilder.GetProperty( GET_MEMBER_NAME_CHECKED( ULevelSequenceExporterUsdOptions, Level ) );
-		LevelToExportProp->SetOnPropertyValueChanged( RebuildDisplayDelegate );
+		FSimpleDelegate RebuildDisplayDelegate = FSimpleDelegate::CreateLambda(
+			[&DetailLayoutBuilder]()
+			{
+				DetailLayoutBuilder.ForceRefreshDetails();
+			}
+		);
+		TSharedRef<IPropertyHandle> LevelToExportProp = DetailLayoutBuilder.GetProperty(
+			GET_MEMBER_NAME_CHECKED(ULevelSequenceExporterUsdOptions, Level)
+		);
+		LevelToExportProp->SetOnPropertyValueChanged(RebuildDisplayDelegate);
 
 		// Only let us pick the sublayer options if we're exporting a level with the level sequence
-		TSharedRef<IPropertyHandle> ExportLevelProp = DetailLayoutBuilder.GetProperty( GET_MEMBER_NAME_CHECKED( ULevelSequenceExporterUsdOptions, bExportLevel ) );
-		SublayersEditCondition = TAttribute<bool>::Create( [this, ExportLevelProp]()
-		{
-			bool bExportingLevel = true;
-			ExportLevelProp->GetValue( bExportingLevel );
-			return bExportingLevel;
-		});
+		TSharedRef<IPropertyHandle> ExportLevelProp = DetailLayoutBuilder.GetProperty(
+			GET_MEMBER_NAME_CHECKED(ULevelSequenceExporterUsdOptions, bExportLevel)
+		);
+		SublayersEditCondition = TAttribute<bool>::Create(
+			[ExportLevelProp]()
+			{
+				bool bExportingLevel = true;
+				ExportLevelProp->GetValue(bExportingLevel);
+				return bExportingLevel;
+			}
+		);
 	}
 	else
 	{
 		return;
 	}
 
-	TSharedRef<IPropertyHandle> LevelFilterProp = DetailLayoutBuilder.GetProperty( LevelFilterPropName );
-	TSharedRef<IPropertyHandle> ExportSublayersProp = DetailLayoutBuilder.GetProperty( ExportSublayersPropName );
+	TSharedRef<IPropertyHandle> LevelFilterProp = DetailLayoutBuilder.GetProperty(LevelFilterPropName);
+	TSharedRef<IPropertyHandle> ExportSublayersProp = DetailLayoutBuilder.GetProperty(ExportSublayersPropName);
 
 	// Touch these properties and categories to enforce this ordering
-	DetailLayoutBuilder.EditCategory( TEXT( "Stage options" ) );
-	DetailLayoutBuilder.EditCategory( TEXT( "Export settings" ) );
-	IDetailCategoryBuilder& AssetOptionsCategory = DetailLayoutBuilder.EditCategory( TEXT( "Asset options" ) );
-	IDetailCategoryBuilder& MetadataOptionsCategory = DetailLayoutBuilder.EditCategory( TEXT( "Metadata options" ) );
-
-	// Promote all AssetOptions up a level on LevelExportUsdOptions or else we'll end up with a property named AssetOptions inside the AssetOptions category
-	// This is the same effect as ShowOnlyInnerProperties, but in this case we need to do it manually as it doesn't work recursively
-	if ( TSharedPtr<IPropertyHandle> AssetOptionsProperty = DetailLayoutBuilder.GetProperty( TEXT( "Inner.AssetOptions" ) ) )
+	DetailLayoutBuilder.EditCategory(TEXT("Stage options"));
+	DetailLayoutBuilder.EditCategory(TEXT("Export settings"));
+	IDetailCategoryBuilder& AssetOptionsCategory = DetailLayoutBuilder.EditCategory(TEXT("Asset options"));
+	IDetailCategoryBuilder& MetadataOptionsCategory = DetailLayoutBuilder.EditCategory(TEXT("Metadata options"));
+
+	// Promote all AssetOptions up a level on LevelExportUsdOptions or else we'll end up with a property named AssetOptions inside the AssetOptions
+	// category This is the same effect as ShowOnlyInnerProperties, but in this case we need to do it manually as it doesn't work recursively
+	if (TSharedPtr<IPropertyHandle> AssetOptionsProperty = DetailLayoutBuilder.GetProperty(TEXT("Inner.AssetOptions")))
 	{
-		DetailLayoutBuilder.HideProperty( AssetOptionsProperty );
+		DetailLayoutBuilder.HideProperty(AssetOptionsProperty);
 
 		uint32 NumChildren = 0;
-		if ( AssetOptionsProperty->GetNumChildren( NumChildren ) == FPropertyAccess::Result::Success )
+		if (AssetOptionsProperty->GetNumChildren(NumChildren) == FPropertyAccess::Result::Success)
 		{
-			for ( uint32 Index = 0; Index < NumChildren; ++Index )
+			for (uint32 Index = 0; Index < NumChildren; ++Index)
 			{
-				TSharedPtr<IPropertyHandle> ChildProperty = AssetOptionsProperty->GetChildHandle( Index );
-				AssetOptionsCategory.AddProperty( ChildProperty );
+				TSharedPtr<IPropertyHandle> ChildProperty = AssetOptionsProperty->GetChildHandle(Index);
+				AssetOptionsCategory.AddProperty(ChildProperty);
 			}
 		}
 	}
@@ -274,37 +284,38 @@ void FLevelExporterUSDOptionsCustomization::CustomizeDetails(IDetailLayoutBuilde
 	}
 
 	// Hide the textures dir property because we'll add multiple textures folders (one next to each exported material)
-	TSharedPtr<IPropertyHandle> TexturesDirProperty = DetailLayoutBuilder.GetProperty( TexturesDirPropName );
-	if ( TexturesDirProperty->IsValidHandle() )
+	TSharedPtr<IPropertyHandle> TexturesDirProperty = DetailLayoutBuilder.GetProperty(TexturesDirPropName);
+	if (TexturesDirProperty->IsValidHandle())
 	{
-		DetailLayoutBuilder.HideProperty( TexturesDirProperty );
+		DetailLayoutBuilder.HideProperty(TexturesDirProperty);
 	}
 
-	DetailLayoutBuilder.EditCategory( TEXT( "Landscape options" ) );
+	DetailLayoutBuilder.EditCategory(TEXT("Landscape options"));
 
 	// Customize the level filter picker widget
-	DetailLayoutBuilder.HideProperty( LevelFilterProp );
-	IDetailCategoryBuilder& CatBuilder = DetailLayoutBuilder.EditCategory( SublayersCategoryName );
-	CatBuilder.AddProperty( ExportSublayersProp ).EditCondition( SublayersEditCondition, nullptr );
-	CatBuilder.AddCustomRow( LevelFilterProp->GetPropertyDisplayName() )
-	.NameContent()
-	[
-		SNew( STextBlock )
-		.Text( FText::FromString( TEXT( "Levels To Export" ) ) )
-		.Font( FAppStyle::GetFontStyle( TEXT( "PropertyWindow.NormalFont" ) ) )
-	]
-	.ValueContent()
-	.MinDesiredWidth(300)
-	[
-		SNew(SScrollBorder, PickerTree.ToSharedRef())
+	DetailLayoutBuilder.HideProperty(LevelFilterProp);
+	IDetailCategoryBuilder& CatBuilder = DetailLayoutBuilder.EditCategory(SublayersCategoryName);
+	CatBuilder.AddProperty(ExportSublayersProp).EditCondition(SublayersEditCondition, nullptr);
+
+	// clang-format off
+	CatBuilder.AddCustomRow(LevelFilterProp->GetPropertyDisplayName())
+		.NameContent()
+		[
+			SNew(STextBlock)
+			.Text(FText::FromString(TEXT("Levels To Export")))
+			.Font(FAppStyle::GetFontStyle(TEXT("PropertyWindow.NormalFont")))
+		]
+		.ValueContent()
+		.MinDesiredWidth(300)
 		[
-			SNew(SBox) // Prevent the list from expanding freely
-			.MaxDesiredHeight(200)
+			SNew(SScrollBorder, PickerTree.ToSharedRef())
 			[
-				PickerTree.ToSharedRef()
+				SNew(SBox)	  // Prevent the list from expanding freely
+				.MaxDesiredHeight(200)[PickerTree.ToSharedRef()]
 			]
 		]
-	].EditCondition( SublayersEditCondition, nullptr );
+		.EditCondition(SublayersEditCondition, nullptr);
+	// clang-format on
 }
 
 #undef LOCTEXT_NAMESPACE
