USD: Reformat USDExporter module.

diff --git a/Engine/Plugins/Importers/USDImporter/Source/USDExporter/Private/LevelSequenceExporterUSD.cpp b/Engine/Plugins/Importers/USDImporter/Source/USDExporter/Private/LevelSequenceExporterUSD.cpp
--- a/Engine/Plugins/Importers/USDImporter/Source/USDExporter/Private/LevelSequenceExporterUSD.cpp
+++ b/Engine/Plugins/Importers/USDImporter/Source/USDExporter/Private/LevelSequenceExporterUSD.cpp
@@ -18,7 +18,6 @@
 
 #include "UsdWrappers/SdfLayer.h"
 #include "UsdWrappers/SdfPath.h"
-#include "UsdWrappers/UsdAttribute.h"
 #include "UsdWrappers/UsdPrim.h"
 #include "UsdWrappers/UsdStage.h"
 
@@ -28,7 +27,6 @@
 #include "Editor.h"
 #include "Engine/SphereReflectionCapture.h"
 #include "EngineAnalytics.h"
-#include "EngineUtils.h"
 #include "Evaluation/MovieSceneSequenceHierarchy.h"
 #include "ISequencer.h"
 #include "ISequencerModule.h"
@@ -52,1470 +50,1503 @@
 
 #define LOCTEXT_NAMESPACE "LevelSequenceExporterUSD"
 
-namespace UE
+namespace UE::LevelSequenceExporterUSD::Private
 {
-	namespace LevelSequenceExporterUSD
+	void SendAnalytics(
+		const ULevelSequence* LevelSequence,
+		const ULevelSequenceExporterUsdOptions* Options,
+		const TArray<UE::FUsdStage>& ExportedStages,
+		bool bAutomated,
+		double ElapsedSeconds,
+		const FString& Extension
+	)
 	{
-		namespace Private
+		if (!LevelSequence || ExportedStages.Num() == 0 || !FEngineAnalytics::IsAvailable())
 		{
-			void SendAnalytics(
-				const ULevelSequence* LevelSequence,
-				const ULevelSequenceExporterUsdOptions* Options,
-				const TArray<UE::FUsdStage>& ExportedStages,
-				bool bAutomated,
-				double ElapsedSeconds,
-				const FString& Extension
-			)
-			{
-				if ( !LevelSequence || ExportedStages.Num() == 0 || !FEngineAnalytics::IsAvailable() )
-				{
-					return;
-				}
+			return;
+		}
 
-				FString ClassName = LevelSequence->GetClass()->GetName();
+		FString ClassName = LevelSequence->GetClass()->GetName();
 
-				TArray<FAnalyticsEventAttribute> EventAttributes;
-				EventAttributes.Emplace( TEXT( "AssetType" ), ClassName );
-				EventAttributes.Emplace( TEXT( "NumExportedLevelSequenceLayers" ), ExportedStages.Num() );
+		TArray<FAnalyticsEventAttribute> EventAttributes;
+		EventAttributes.Emplace(TEXT("AssetType"), ClassName);
+		EventAttributes.Emplace(TEXT("NumExportedLevelSequenceLayers"), ExportedStages.Num());
 
-				int32 NumFrames = 0;
-				if ( Options )
-				{
-					NumFrames = Options->EndFrame - Options->StartFrame;
-					UsdUtils::AddAnalyticsAttributes( *Options, EventAttributes );
-				}
+		int32 NumFrames = 0;
+		if (Options)
+		{
+			NumFrames = Options->EndFrame - Options->StartFrame;
+			UsdUtils::AddAnalyticsAttributes(*Options, EventAttributes);
+		}
 
-				IUsdClassesModule::SendAnalytics(
-					MoveTemp( EventAttributes ),
-					FString::Printf( TEXT( "Export.%s" ), *ClassName ),
-					bAutomated,
-					ElapsedSeconds,
-					NumFrames,
-					Extension
-				);
-			}
+		IUsdClassesModule::SendAnalytics(
+			MoveTemp(EventAttributes),
+			FString::Printf(TEXT("Export.%s"), *ClassName),
+			bAutomated,
+			ElapsedSeconds,
+			NumFrames,
+			Extension
+		);
+	}
 
 #if USE_USD_SDK
-			// Custom spawn register so that when DestroySpawnedObject is called while bDestroyingJustHides is true we
-			// actually just hide the objects, so that we can keep a live reference to components within the bakers.
-			// We're going to convert the spawnable tracks into visibility tracks when exporting to USD, which
-			// also works well with this approach
-			class FLevelSequenceHidingSpawnRegister : public FLevelSequenceEditorSpawnRegister
+	// Custom spawn register so that when DestroySpawnedObject is called while bDestroyingJustHides is true we
+	// actually just hide the objects, so that we can keep a live reference to components within the bakers.
+	// We're going to convert the spawnable tracks into visibility tracks when exporting to USD, which
+	// also works well with this approach
+	class FLevelSequenceHidingSpawnRegister : public FLevelSequenceEditorSpawnRegister
+	{
+	public:
+		bool bDestroyingJustHides = true;
+
+		virtual UObject* SpawnObject(
+			FMovieSceneSpawnable& Spawnable,
+			FMovieSceneSequenceIDRef TemplateID,
+			TSharedRef<const FSharedPlaybackState> SharedPlaybackState
+		) override
+		{
+			// Never spawn ASphereReflectionCapture actors. These are useless in USD anyway, and we run into
+			// trouble after we're done exporting them because on the tick where they're destroyed the editor
+			// will still attempt to update their captures and some downstream code doesn't like that their
+			// components are pending kill (check UE-167593 for more info)
+			if (const ASphereReflectionCapture* ReflectionCapture = Cast<const ASphereReflectionCapture>(Spawnable.GetObjectTemplate()))
 			{
-			public:
-				bool bDestroyingJustHides = true;
+				return nullptr;
+			}
 
-				virtual UObject* SpawnObject( FMovieSceneSpawnable& Spawnable, FMovieSceneSequenceIDRef TemplateID, TSharedRef<const FSharedPlaybackState> SharedPlaybackState ) override
-				{
-					// Never spawn ASphereReflectionCapture actors. These are useless in USD anyway, and we run into
-					// trouble after we're done exporting them because on the tick where they're destroyed the editor
-					// will still attempt to update their captures and some downstream code doesn't like that their
-					// components are pending kill (check UE-167593 for more info)
-					if ( const ASphereReflectionCapture* ReflectionCapture = Cast< const ASphereReflectionCapture >( Spawnable.GetObjectTemplate() ) )
-					{
-						return nullptr;
-					}
+			UMovieSceneSequence* RootSequence = SharedPlaybackState->GetRootSequence();
+			if (!RootSequence)
+			{
+				return nullptr;
+			}
 
-					UMovieSceneSequence* RootSequence = SharedPlaybackState->GetRootSequence();
-					if ( !RootSequence )
-					{
-						return nullptr;
-					}
+			UObject* Object = nullptr;
 
-					UObject* Object = nullptr;
+			const FGuid& Guid = Spawnable.GetGuid();
 
-					const FGuid& Guid = Spawnable.GetGuid();
+			TArray<UObject*>& ExistingInstancesForGuid = SpawnableInstances.FindOrAdd(Guid);
+
+			TMap<FMovieSceneSequenceID, TMap<FGuid, int32>>& SequenceInstanceToSpawnableIndices = RootSequenceToSpawnableInstanceIndices.FindOrAdd(
+				RootSequence
+			);
+			TMap<FGuid, int32>& SpawnableIndices = SequenceInstanceToSpawnableIndices.FindOrAdd(TemplateID);
 
-					TArray<UObject*>& ExistingInstancesForGuid = SpawnableInstances.FindOrAdd( Guid );
+			// Already have an instance of this spawnable for this movie scene sequence instance
+			if (int32* ExistingIndex = SpawnableIndices.Find(Guid))
+			{
+				Object = ExistingInstancesForGuid[*ExistingIndex];
+
+				UE_LOG(
+					LogUsd,
+					VeryVerbose,
+					TEXT("Already spawned '%s' (%0x) for RootSequence '%s', TemplateID '%u', Guid '%s' (index %d)"),
+					*Object->GetPathName(),
+					Object,
+					*RootSequence->GetPathName(),
+					TemplateID.GetInternalValue(),
+					*Guid.ToString(),
+					*ExistingIndex
+				);
+			}
 
-					TMap<FMovieSceneSequenceID, TMap<FGuid, int32>>& SequenceInstanceToSpawnableIndices = RootSequenceToSpawnableInstanceIndices.FindOrAdd( RootSequence );
-					TMap<FGuid, int32>& SpawnableIndices = SequenceInstanceToSpawnableIndices.FindOrAdd( TemplateID );
+			// We don't have an instance of the spawnable spawned for this exact movie sequence ID, but try to see if we can
+			// reuse any of the existing spawns for it.
+			// We have to reuse these because the exported level will just contain prims that correspond to the spawnables spawned for
+			// the top-level level sequence. If we don't reuse these, recursive exports of the subsequences would attempt to spawn
+			// their own spawnables (as we always need an instance for each movie scene sequence ID), which wouldn't correspond to
+			// any prim that was exported on the level as the actor names would get sanitized away from any name used by the top-level export
+			if (!Object && ExistingInstancesForGuid.Num() > 0)
+			{
+				TArray<bool> UsedIndices;
+				UsedIndices.SetNumZeroed(ExistingInstancesForGuid.Num());
 
-					// Already have an instance of this spawnable for this movie scene sequence instance
-					if ( int32* ExistingIndex = SpawnableIndices.Find( Guid ) )
+				for (const TPair<FMovieSceneSequenceID, TMap<FGuid, int32>>& Pair : SequenceInstanceToSpawnableIndices)
+				{
+					if (const int32* UsedIndex = Pair.Value.Find(Guid))
 					{
-						Object = ExistingInstancesForGuid[ *ExistingIndex ];
-
-						UE_LOG( LogUsd, VeryVerbose, TEXT( "Already spawned '%s' (%0x) for RootSequence '%s', TemplateID '%u', Guid '%s' (index %d)" ),
-							*Object->GetPathName(),
-							Object,
-							*RootSequence->GetPathName(),
-							TemplateID.GetInternalValue(),
-							*Guid.ToString(),
-							*ExistingIndex
-						);
+						UsedIndices[*UsedIndex] = true;
 					}
+				}
 
-					// We don't have an instance of the spawnable spawned for this exact movie sequence ID, but try to see if we can
-					// reuse any of the existing spawns for it.
-					// We have to reuse these because the exported level will just contain prims that correspond to the spawnables spawned for
-					// the top-level level sequence. If we don't reuse these, recursive exports of the subsequences would attempt to spawn
-					// their own spawnables (as we always need an instance for each movie scene sequence ID), which wouldn't correspond to
-					// any prim that was exported on the level as the actor names would get sanitized away from any name used by the top-level export
-					if ( !Object && ExistingInstancesForGuid.Num() > 0 )
+				int32 IndexToReuse = INDEX_NONE;
+				for (int32 Index = 0; Index < UsedIndices.Num(); ++Index)
+				{
+					if (!UsedIndices[Index])
 					{
-						TArray<bool> UsedIndices;
-						UsedIndices.SetNumZeroed( ExistingInstancesForGuid.Num() );
-
-						for ( const TPair<FMovieSceneSequenceID, TMap<FGuid, int32>>& Pair : SequenceInstanceToSpawnableIndices )
-						{
-							if ( const int32* UsedIndex = Pair.Value.Find( Guid ) )
-							{
-								UsedIndices[ *UsedIndex ] = true;
-							}
-						}
-
-						int32 IndexToReuse = INDEX_NONE;
-						for ( int32 Index = 0; Index < UsedIndices.Num(); ++Index )
-						{
-							if ( !UsedIndices[ Index ] )
-							{
-								IndexToReuse = Index;
-								break;
-							}
-						}
-
-						if ( IndexToReuse != INDEX_NONE )
-						{
-							Object = ExistingInstancesForGuid[ IndexToReuse ];
-							SpawnableIndices.Add( Guid, IndexToReuse );
-
-							UE_LOG( LogUsd, VeryVerbose, TEXT( "Reusing '%s' (%0x) for RootSequence '%s', TemplateID '%u', Guid '%s' (index %d)" ),
-								*Object->GetPathName(),
-								Object,
-								*RootSequence->GetPathName(),
-								TemplateID.GetInternalValue(),
-								*Guid.ToString(),
-								IndexToReuse
-							);
-						}
+						IndexToReuse = Index;
+						break;
 					}
+				}
 
-					// Don't even have anything we can reuse: We need to spawn a brand new instance of this spawnable
-					if ( !Object )
-					{
-						Object = FLevelSequenceEditorSpawnRegister::SpawnObject( Spawnable, TemplateID, SharedPlaybackState );
-						UE_LOG( LogUsd, VeryVerbose, TEXT( "Spawning '%s' (%0x) for RootSequence '%s', TemplateID '%u', Guid '%s'" ),
-							*Object->GetPathName(),
-							Object,
-							*RootSequence->GetPathName(),
-							TemplateID.GetInternalValue(),
-							*Guid.ToString()
-						);
-
-						if ( AActor* SpawnedActor = Cast<AActor>( Object ) )
-						{
-							// Rename the spawn to a unique name or else in case of name collisions they will overwrite each other when writing animation data.
-							// The level exporter will rename actors to unique prims by itself though.
-							FString NewLabel = UsdUtils::GetUniqueName( SpawnedActor->GetActorLabel(), UsedActorLabels );
-							if ( NewLabel != SpawnedActor->GetActorLabel() )
-							{
-								const bool bMarkDirty = false;
-								SpawnedActor->SetActorLabel( NewLabel, bMarkDirty );
-							}
-							UsedActorLabels.Add( NewLabel );
-						}
+				if (IndexToReuse != INDEX_NONE)
+				{
+					Object = ExistingInstancesForGuid[IndexToReuse];
+					SpawnableIndices.Add(Guid, IndexToReuse);
+
+					UE_LOG(
+						LogUsd,
+						VeryVerbose,
+						TEXT("Reusing '%s' (%0x) for RootSequence '%s', TemplateID '%u', Guid '%s' (index %d)"),
+						*Object->GetPathName(),
+						Object,
+						*RootSequence->GetPathName(),
+						TemplateID.GetInternalValue(),
+						*Guid.ToString(),
+						IndexToReuse
+					);
+				}
+			}
 
-						ExistingInstancesForGuid.Add( Object );
-						SpawnableIndices.Add( Guid, ExistingInstancesForGuid.Num() - 1 );
-					}
+			// Don't even have anything we can reuse: We need to spawn a brand new instance of this spawnable
+			if (!Object)
+			{
+				Object = FLevelSequenceEditorSpawnRegister::SpawnObject(Spawnable, TemplateID, SharedPlaybackState);
+				UE_LOG(
+					LogUsd,
+					VeryVerbose,
+					TEXT("Spawning '%s' (%0x) for RootSequence '%s', TemplateID '%u', Guid '%s'"),
+					*Object->GetPathName(),
+					Object,
+					*RootSequence->GetPathName(),
+					TemplateID.GetInternalValue(),
+					*Guid.ToString()
+				);
 
-					if ( Object )
+				if (AActor* SpawnedActor = Cast<AActor>(Object))
+				{
+					// Rename the spawn to a unique name or else in case of name collisions they will overwrite each other when writing
+					// animation data. The level exporter will rename actors to unique prims by itself though.
+					FString NewLabel = UsdUtils::GetUniqueName(SpawnedActor->GetActorLabel(), UsedActorLabels);
+					if (NewLabel != SpawnedActor->GetActorLabel())
 					{
-						USceneComponent* Component = nullptr;
+						const bool bMarkDirty = false;
+						SpawnedActor->SetActorLabel(NewLabel, bMarkDirty);
+					}
+					UsedActorLabels.Add(NewLabel);
+				}
 
-						if ( AActor* Actor = Cast<AActor>( Object ) )
-						{
-							const bool bIsHidden = false;
-							Actor->SetActorHiddenInGame( bIsHidden );
+				ExistingInstancesForGuid.Add(Object);
+				SpawnableIndices.Add(Guid, ExistingInstancesForGuid.Num() - 1);
+			}
 
-							Component = Actor->GetRootComponent();
-						}
-						else
-						{
-							Component = Cast<USceneComponent>( Object );
-						}
+			if (Object)
+			{
+				USceneComponent* Component = nullptr;
 
-						if ( Component )
-						{
-							const bool bNewHidden = false;
-							Component->SetHiddenInGame( bNewHidden );
-						}
-					}
+				if (AActor* Actor = Cast<AActor>(Object))
+				{
+					const bool bIsHidden = false;
+					Actor->SetActorHiddenInGame(bIsHidden);
 
-					return Object;
+					Component = Actor->GetRootComponent();
 				}
-
-				virtual void PreDestroyObject(UObject& Object, const FGuid& BindingId, FMovieSceneSequenceIDRef TemplateID) override
+				else
 				{
-					// Don't let the FLevelSequenceEditorSpawnRegister's overload run as it will mess with our editor selection
-					if (bDestroyingJustHides)
-					{
-						FLevelSequenceSpawnRegister::PreDestroyObject(Object, BindingId, TemplateID);
-						return;
-					}
-
-					FLevelSequenceEditorSpawnRegister::PreDestroyObject(Object, BindingId, TemplateID);
+					Component = Cast<USceneComponent>(Object);
 				}
 
-				virtual void DestroySpawnedObject( UObject& Object ) override
+				if (Component)
 				{
-					if ( bDestroyingJustHides )
-					{
-						USceneComponent* Component = nullptr;
-
-						if ( AActor* Actor = Cast<AActor>( &Object ) )
-						{
-							const bool bIsHidden = true;
-							Actor->SetActorHiddenInGame( bIsHidden );
+					const bool bNewHidden = false;
+					Component->SetHiddenInGame(bNewHidden);
+				}
+			}
 
-							Component = Actor->GetRootComponent();
-						}
-						else
-						{
-							Component = Cast<USceneComponent>( &Object );
-						}
+			return Object;
+		}
 
-						// Make sure we toggle both actor and component as the component is the one
-						// that is monitored by the level exporter, and we want our spawnables to be hidden
-						// by default
-						if ( Component )
-						{
-							const bool bNewHidden = true;
-							Component->SetHiddenInGame( bNewHidden );
-						}
-					}
-					else
-					{
-						// We shouldn't need to do this because we only ever fully delete when we're cleaning up,
-						// and by then we'll delete all of these maps anyway
-						for ( TPair<FGuid, TArray<UObject*>>& Pair : SpawnableInstances )
-						{
-							const FGuid& Guid = Pair.Key;
-							TArray<UObject*>& InstancesForGuid = Pair.Value;
+		virtual void PreDestroyObject(UObject& Object, const FGuid& BindingId, FMovieSceneSequenceIDRef TemplateID) override
+		{
+			// Don't let the FLevelSequenceEditorSpawnRegister's overload run as it will mess with our editor selection
+			if (bDestroyingJustHides)
+			{
+				FLevelSequenceSpawnRegister::PreDestroyObject(Object, BindingId, TemplateID);
+				return;
+			}
 
-							int32 IndexToDelete = INDEX_NONE;
-							for ( int32 Index = 0; Index < InstancesForGuid.Num(); ++Index )
-							{
-								if ( InstancesForGuid[ Index ] == &Object )
-								{
-									IndexToDelete = Index;
-									break;
-								}
-							}
+			FLevelSequenceEditorSpawnRegister::PreDestroyObject(Object, BindingId, TemplateID);
+		}
 
-							if ( IndexToDelete != INDEX_NONE )
-							{
-								for ( TPair<const UMovieSceneSequence*, TMap<FMovieSceneSequenceID, TMap<FGuid, int32>>>& RootSequencePair : RootSequenceToSpawnableInstanceIndices )
-								{
-									for ( TPair<FMovieSceneSequenceID, TMap<FGuid, int32>>& SequenceIDPair : RootSequencePair.Value )
-									{
-										TMap<FGuid, int32>& GuidToInstance = SequenceIDPair.Value;
-										if ( int32* InstanceIndex = GuidToInstance.Find( Guid ) )
-										{
-											if ( *InstanceIndex == IndexToDelete )
-											{
-												GuidToInstance.Remove( Guid );
-											}
-										}
-									}
-								}
+		virtual void DestroySpawnedObject(UObject& Object) override
+		{
+			if (bDestroyingJustHides)
+			{
+				USceneComponent* Component = nullptr;
 
-								InstancesForGuid.RemoveAt( IndexToDelete );
-								break;
-							}
-						}
+				if (AActor* Actor = Cast<AActor>(&Object))
+				{
+					const bool bIsHidden = true;
+					Actor->SetActorHiddenInGame(bIsHidden);
 
-						FLevelSequenceEditorSpawnRegister::DestroySpawnedObject( Object );
-					}
+					Component = Actor->GetRootComponent();
 				}
-
-				bool HasSpawnedObject(const FGuid& BindingGuid) const
+				else
 				{
-					return SpawnableInstances.Contains(BindingGuid);
+					Component = Cast<USceneComponent>(&Object);
 				}
 
-				void DeleteSpawns( TSharedRef<const FSharedPlaybackState> SharedPlaybackState )
+				// Make sure we toggle both actor and component as the component is the one
+				// that is monitored by the level exporter, and we want our spawnables to be hidden
+				// by default
+				if (Component)
 				{
-					bDestroyingJustHides = false;
-					CleanUp( SharedPlaybackState );
-
-					// If we still have existing spawns it may be because our base class' Register member didn't contain an
-					// entry for a spawnable before we called CleanUp (check its implementation: It just iterates over that Register).
-					// This is expected in some scenarios because we're sort of abusing this inheritance: The base class expects
-					// CleanUp to delete the object already and so removes its Register entry, but we'll keep the instances
-					// alive when bDestroyingJustHides=true. Because of this we must explicitly clean up these "abandoned" spawns here,
-					// which resynchronizes us with Register:
-					TArray<UObject*> ObjectsToDelete;
-					for ( TPair<FGuid, TArray<UObject*>>& Pair : SpawnableInstances )
-					{
-						ObjectsToDelete.Append( Pair.Value );
-					}
-					for ( UObject* Object : ObjectsToDelete )
+					const bool bNewHidden = true;
+					Component->SetHiddenInGame(bNewHidden);
+				}
+			}
+			else
+			{
+				// We shouldn't need to do this because we only ever fully delete when we're cleaning up,
+				// and by then we'll delete all of these maps anyway
+				for (TPair<FGuid, TArray<UObject*>>& Pair : SpawnableInstances)
+				{
+					const FGuid& Guid = Pair.Key;
+					TArray<UObject*>& InstancesForGuid = Pair.Value;
+
+					int32 IndexToDelete = INDEX_NONE;
+					for (int32 Index = 0; Index < InstancesForGuid.Num(); ++Index)
 					{
-						if ( Object )
+						if (InstancesForGuid[Index] == &Object)
 						{
-							DestroySpawnedObject( *Object );
+							IndexToDelete = Index;
+							break;
 						}
 					}
-				}
 
-				// FLevelSequenceHidingSpawnRegister is a bit of a hack and just hides it's spawned actors instead
-				// of deleting them (when we want it to do so). Unfortunately, the base FMovieSceneSpawnRegister part
-				// will still nevertheless clear it's Register entry for the spawnable when deleting (even if just hiding),
-				// and there's nothing we can do to prevent it. This means we can't call FindSpawnedObject and must use our
-				// own GetExistingSpawn and data members
-				UObject* GetExistingSpawn( const UMovieSceneSequence& RootSequence, FMovieSceneSequenceID SequenceID, const FGuid& SpawnableGuid )
-				{
-					int32 SpawnableIndex = INDEX_NONE;
-					if ( TMap<FMovieSceneSequenceID, TMap<FGuid, int32>>* SequenceIDsToSpawns = RootSequenceToSpawnableInstanceIndices.Find( &RootSequence ) )
+					if (IndexToDelete != INDEX_NONE)
 					{
-						if ( TMap<FGuid, int32>* Spawns = SequenceIDsToSpawns->Find( SequenceID ) )
+						for (TPair<const UMovieSceneSequence*, TMap<FMovieSceneSequenceID, TMap<FGuid, int32>>>& RootSequencePair :
+							 RootSequenceToSpawnableInstanceIndices)
 						{
-							if ( int32* Index = Spawns->Find( SpawnableGuid ) )
+							for (TPair<FMovieSceneSequenceID, TMap<FGuid, int32>>& SequenceIDPair : RootSequencePair.Value)
 							{
-								SpawnableIndex = *Index;
+								TMap<FGuid, int32>& GuidToInstance = SequenceIDPair.Value;
+								if (int32* InstanceIndex = GuidToInstance.Find(Guid))
+								{
+									if (*InstanceIndex == IndexToDelete)
+									{
+										GuidToInstance.Remove(Guid);
+									}
+								}
 							}
 						}
-					}
 
-					if ( SpawnableIndex != INDEX_NONE )
-					{
-						if ( TArray<UObject*>* Instances = SpawnableInstances.Find( SpawnableGuid ) )
-						{
-							return ( *Instances )[ SpawnableIndex ];
-						}
+						InstancesForGuid.RemoveAt(IndexToDelete);
+						break;
 					}
-
-					return nullptr;
 				}
 
-			private:
-				// Ensures all of our new spawns have unique names
-				TSet<FString> UsedActorLabels;
-
-				// Tracks all instances we created for a given spawnable guid
-				TMap<FGuid, TArray<UObject*>> SpawnableInstances;
+				FLevelSequenceEditorSpawnRegister::DestroySpawnedObject(Object);
+			}
+		}
 
-				// Tracks the indices into SpawnableInstances for each spawnable guid, used by each sequence ID, in the hierarchy of each root sequence
-				TMap<const UMovieSceneSequence*, TMap<FMovieSceneSequenceID, TMap<FGuid, int32>>> RootSequenceToSpawnableInstanceIndices;
-			};
+		bool HasSpawnedObject(const FGuid& BindingGuid) const
+		{
+			return SpawnableInstances.Contains(BindingGuid);
+		}
 
-			// Contain all of the baker lambda functions for a given component. Only one baker per baking type is allowed.
-			struct FCombinedComponentBakers
+		void DeleteSpawns(TSharedRef<const FSharedPlaybackState> SharedPlaybackState)
+		{
+			bDestroyingJustHides = false;
+			CleanUp(SharedPlaybackState);
+
+			// If we still have existing spawns it may be because our base class' Register member didn't contain an
+			// entry for a spawnable before we called CleanUp (check its implementation: It just iterates over that Register).
+			// This is expected in some scenarios because we're sort of abusing this inheritance: The base class expects
+			// CleanUp to delete the object already and so removes its Register entry, but we'll keep the instances
+			// alive when bDestroyingJustHides=true. Because of this we must explicitly clean up these "abandoned" spawns here,
+			// which resynchronizes us with Register:
+			TArray<UObject*> ObjectsToDelete;
+			for (TPair<FGuid, TArray<UObject*>>& Pair : SpawnableInstances)
 			{
-				UnrealToUsd::EBakingType CombinedBakingType = UnrealToUsd::EBakingType::None;
-				TArray<UnrealToUsd::FComponentBaker> Bakers;
-			};
-
-			class FLevelSequenceExportContext
+				ObjectsToDelete.Append(Pair.Value);
+			}
+			for (UObject* Object : ObjectsToDelete)
 			{
-			public:
-				FLevelSequenceExportContext( ULevelSequence& InSequence, TSharedRef<ISequencer> InSequencer, TSharedRef<FLevelSequenceHidingSpawnRegister> InSpawnRegister )
-					: RootSequence( InSequence )
-					, Sequencer( InSequencer )
-					, SpawnRegister( InSpawnRegister )
+				if (Object)
 				{
+					DestroySpawnedObject(*Object);
 				}
+			}
+		}
 
-				FLevelSequenceExportContext() = delete;
-				FLevelSequenceExportContext( const FLevelSequenceExportContext& Other ) = delete;
-				FLevelSequenceExportContext( const FLevelSequenceExportContext&& Other ) = delete;
-				FLevelSequenceExportContext& operator=( const FLevelSequenceExportContext& Other ) = delete;
-				FLevelSequenceExportContext& operator=( const FLevelSequenceExportContext&& Other ) = delete;
-
-			public:
-				// The actual content asset that is being exported
-				ULevelSequence& RootSequence;
-
-				ULevelSequenceExporterUsdOptions* ExportOptions;
-
-				// If ExportOptions->bSelectionOnly is true, this specifies the actors whose bindings we should export
-				TSet<AActor*> SelectedActors;
-
-				// Where we store our ExportTask's bReplaceIdentical, which indicates if we should overwrite files or not
-				bool bReplaceIdentical;
-
-				// Our own read-only sequencer that we use to play the level sequences while we bake them out one frame at a time
-				TSharedRef<ISequencer> Sequencer;
-
-				// Object that manages spawned instances for FMovieScenePossessables
-				TSharedRef<FLevelSequenceHidingSpawnRegister> SpawnRegister;
-
-				// Used to keep track of which sequences we already baked
-				TMap<UMovieSceneSequence*, FString> ExportedMovieScenes;
-
-				// File paths that we already used during this export.
-				// Used so that we can prevent conflicts between files emitted for this export, but can still
-				// overwrite other files on disk
-				TSet<FString> UsedFilePaths;
-
-				// File path of the exported USD root layer, in case we also exported the level along with the level sequence
-				FString LevelFilePath;
-
-				// World that is being animated for the LevelSequence bake
-				UWorld* World = nullptr;
-
-			public:
-				~FLevelSequenceExportContext()
-				{
-					SpawnRegister->DeleteSpawns( Sequencer->GetSharedPlaybackState() );
-				}
-			};
-
-			bool IsTrackAnimated( const UMovieSceneTrack* Track )
+		// FLevelSequenceHidingSpawnRegister is a bit of a hack and just hides it's spawned actors instead
+		// of deleting them (when we want it to do so). Unfortunately, the base FMovieSceneSpawnRegister part
+		// will still nevertheless clear it's Register entry for the spawnable when deleting (even if just hiding),
+		// and there's nothing we can do to prevent it. This means we can't call FindSpawnedObject and must use our
+		// own GetExistingSpawn and data members
+		UObject* GetExistingSpawn(const UMovieSceneSequence& RootSequence, FMovieSceneSequenceID SequenceID, const FGuid& SpawnableGuid)
+		{
+			int32 SpawnableIndex = INDEX_NONE;
+			if (TMap<FMovieSceneSequenceID, TMap<FGuid, int32>>* SequenceIDsToSpawns = RootSequenceToSpawnableInstanceIndices.Find(&RootSequence))
 			{
-				for ( const UMovieSceneSection* Section : Track->GetAllSections() )
+				if (TMap<FGuid, int32>* Spawns = SequenceIDsToSpawns->Find(SequenceID))
 				{
-					// We can't just check whether a section has range here because it may just have a bunch
-					// of channels with modified default values instead (and so no range or even keys, but can still affect the level).
-					// Sadly there's no way of telling if the default value in the channel has been modified or not (that is, whether it
-					// matches the unanimated value or not), so we'll just have to export any track with an active section for now.
-					// This will emit a bunch of unwanted tracks, but later on we may want to do a post-processing pass on our exported
-					// data to make sure we don't emit intermediary keys on linearly interpolated sections, and that would naturally fix
-					// this too.
-					if ( Section->IsActive() )
+					if (int32* Index = Spawns->Find(SpawnableGuid))
 					{
-						return true;
+						SpawnableIndex = *Index;
 					}
 				}
-
-				return false;
 			}
 
-			TMap< UMovieSceneSequence*, TArray<FMovieSceneSequenceID> > GetSequenceHierarchyInstances( UMovieSceneSequence& Sequence, ISequencer& Sequencer )
+			if (SpawnableIndex != INDEX_NONE)
 			{
-				TMap< UMovieSceneSequence*, TArray<FMovieSceneSequenceID> > SequenceInstances;
-
-				FMovieSceneSequenceHierarchy SequenceHierarchyCache;
-				UMovieSceneCompiledDataManager::CompileHierarchy( &Sequence, &SequenceHierarchyCache, EMovieSceneServerClientMask::All );
-
-				SequenceInstances.FindOrAdd( &Sequence ).Add( Sequencer.GetRootTemplateID() );
-				for ( const TTuple< FMovieSceneSequenceID, FMovieSceneSubSequenceData >& Pair : SequenceHierarchyCache.AllSubSequenceData() )
+				if (TArray<UObject*>* Instances = SpawnableInstances.Find(SpawnableGuid))
 				{
-					if ( ULevelSequence* SubSequence = Cast<ULevelSequence>( Pair.Value.GetSequence() ) )
-					{
-						SequenceInstances.FindOrAdd( SubSequence ).Add( Pair.Key );
-					}
+					return (*Instances)[SpawnableIndex];
 				}
-
-				return SequenceInstances;
 			}
 
-			// Spawn and hide all spawnables for Sequence hierarchy
-			TMap< UMovieSceneSequence*, TArray<FMovieSceneSequenceID> > PreSpawnSpawnables( FLevelSequenceExportContext& Context, UMovieSceneSequence& RootSequence )
-			{
-				TMap< UMovieSceneSequence*, TArray<FMovieSceneSequenceID> > SequenceInstances = GetSequenceHierarchyInstances( RootSequence, Context.Sequencer.Get() );
+			return nullptr;
+		}
 
-				UMovieSceneSequence* OrigRootSequence = Context.Sequencer->GetRootMovieSceneSequence();
-				Context.Sequencer->ResetToNewRootSequence( RootSequence );
+	private:
+		// Ensures all of our new spawns have unique names
+		TSet<FString> UsedActorLabels;
 
-				// Spawn everything for this Sequence hierarchy
-				for ( TPair< UMovieSceneSequence*, TArray<FMovieSceneSequenceID> >& SequenceInstancePair : SequenceInstances )
-				{
-					const UMovieSceneSequence* Sequence = SequenceInstancePair.Key;
-					if ( !Sequence )
-					{
-						continue;
-					}
+		// Tracks all instances we created for a given spawnable guid
+		TMap<FGuid, TArray<UObject*>> SpawnableInstances;
 
-					UMovieScene* MovieScene = Sequence->GetMovieScene();
-					if ( !MovieScene )
-					{
-						continue;
-					}
+		// Tracks the indices into SpawnableInstances for each spawnable guid, used by each sequence ID, in the hierarchy of each root
+		// sequence
+		TMap<const UMovieSceneSequence*, TMap<FMovieSceneSequenceID, TMap<FGuid, int32>>> RootSequenceToSpawnableInstanceIndices;
+	};
 
-					// Spawn everything for this instance
-					for ( FMovieSceneSequenceID SequenceInstance : SequenceInstancePair.Value )
-					{
-						int32 NumSpawnables = MovieScene->GetSpawnableCount();
-						for ( int32 Index = 0; Index < NumSpawnables; ++Index )
-						{
-							const FMovieSceneSpawnable& Spawnable = MovieScene->GetSpawnable( Index );
-							const FGuid& Guid = Spawnable.GetGuid();
+	// Contain all of the baker lambda functions for a given component. Only one baker per baking type is allowed.
+	struct FCombinedComponentBakers
+	{
+		UnrealToUsd::EBakingType CombinedBakingType = UnrealToUsd::EBakingType::None;
+		TArray<UnrealToUsd::FComponentBaker> Bakers;
+	};
 
-							StaticCastSharedRef<FMovieSceneSpawnRegister>( Context.SpawnRegister )->SpawnObject( Guid, *MovieScene, SequenceInstance, *Context.Sequencer );
-						}
-					}
-				}
+	class FLevelSequenceExportContext
+	{
+	public:
+		FLevelSequenceExportContext(
+			ULevelSequence& InSequence,
+			TSharedRef<ISequencer> InSequencer,
+			TSharedRef<FLevelSequenceHidingSpawnRegister> InSpawnRegister
+		)
+			: RootSequence(InSequence)
+			, Sequencer(InSequencer)
+			, SpawnRegister(InSpawnRegister)
+		{
+		}
 
-				// Recurse into subsequences if we'll export them as sublayers. This because when exporting each, we'll reset our sequencer
-				// to have each subsequence as the root level sequence again, which means that the sequence ids will all change. We need to
-				// prepare our SpawnRegister to be able to reuse the previous spawns for these new instance ids
-				if ( Context.ExportOptions->bExportSubsequencesAsLayers )
-				{
-					for ( TPair< UMovieSceneSequence*, TArray<FMovieSceneSequenceID> >& SequenceInstancePair : SequenceInstances )
-					{
-						UMovieSceneSequence* Sequence = SequenceInstancePair.Key;
-						if ( !Sequence || Sequence == &RootSequence )
-						{
-							continue;
-						}
+		FLevelSequenceExportContext() = delete;
+		FLevelSequenceExportContext(const FLevelSequenceExportContext& Other) = delete;
+		FLevelSequenceExportContext(const FLevelSequenceExportContext&& Other) = delete;
+		FLevelSequenceExportContext& operator=(const FLevelSequenceExportContext& Other) = delete;
+		FLevelSequenceExportContext& operator=(const FLevelSequenceExportContext&& Other) = delete;
 
-						PreSpawnSpawnables( Context, *Sequence );
-					}
-				}
+	public:
+		// The actual content asset that is being exported
+		ULevelSequence& RootSequence;
 
-				// Put this back to what it was before this call
-				if (OrigRootSequence)
-				{
-					Context.Sequencer->ResetToNewRootSequence(*OrigRootSequence);
-				}
+		ULevelSequenceExporterUsdOptions* ExportOptions;
 
-				return SequenceInstances;
-			}
+		// If ExportOptions->bSelectionOnly is true, this specifies the actors whose bindings we should export
+		TSet<AActor*> SelectedActors;
 
-			// Collect selected actors and components before we do anything.
-			// Note that this is especially needed for spawnables since we're going to have to replicate the selection on
-			// the analogous spawnables that belong to our TempSequencer
-			void CollectEditorSelection(
-				const FLevelSequenceHidingSpawnRegister& SpawnRegister,
-				TArray<AActor*>& OutSelectedSpawnableActors,
-				TArray<AActor*>& OutSelectedNonSpawnableActors,
-				TArray<UE::MovieScene::FFixedObjectBindingID>& OutSelectedSpawnableBindings
-			)
-			{
-				OutSelectedSpawnableActors.Reset();
-				OutSelectedNonSpawnableActors.Reset();
-				OutSelectedSpawnableBindings.Reset();
+		// Where we store our ExportTask's bReplaceIdentical, which indicates if we should overwrite files or not
+		bool bReplaceIdentical;
 
-				// We'll start off with all actors inside OutSelectedNonSpawnableActors, and then remove the
-				// spawnables and put them on OutSelectedSpawnableActors
-				USelection* ActorSelection = GEditor->GetSelectedActors();
-				ActorSelection->GetSelectedObjects(OutSelectedNonSpawnableActors);
+		// Our own read-only sequencer that we use to play the level sequences while we bake them out one frame at a time
+		TSharedRef<ISequencer> Sequencer;
 
-				OutSelectedSpawnableActors.Reserve(OutSelectedNonSpawnableActors.Num());
+		// Object that manages spawned instances for FMovieScenePossessables
+		TSharedRef<FLevelSequenceHidingSpawnRegister> SpawnRegister;
 
-				TArray<UE::MovieScene::FFixedObjectBindingID> SpawnableBindings;
-				for (const TWeakPtr<ISequencer>& Sequencer : FLevelEditorSequencerIntegration::Get().GetSequencers())
-				{
-					if (TSharedPtr<ISequencer> PinnedSequencer = Sequencer.Pin())
-					{
-						// Components are never "spawnable", so we only need to get these bindings for actors
-						for(int32 Index = OutSelectedNonSpawnableActors.Num() - 1; Index >= 0; --Index)
-						{
-							AActor* Actor = OutSelectedNonSpawnableActors[Index];
+		// Used to keep track of which sequences we already baked
+		TMap<UMovieSceneSequence*, FString> ExportedMovieScenes;
 
-							TArray<FMovieSceneObjectBindingID> Bindings;
-							PinnedSequencer->State.FilterObjectBindings(Actor, *PinnedSequencer, &Bindings);
+		// File paths that we already used during this export.
+		// Used so that we can prevent conflicts between files emitted for this export, but can still
+		// overwrite other files on disk
+		TSet<FString> UsedFilePaths;
 
-							OutSelectedSpawnableBindings.Reserve(Bindings.Num() + OutSelectedSpawnableBindings.Num());
-							for (const FMovieSceneObjectBindingID& Binding : Bindings)
-							{
-								// Using the spawn register is an easy way of telling if a Guid is a spawnable or not,
-								// but it's more appropriate because we really only ever care about the spawnables that
-								// we have spawned on our TempSequencer
-								const bool bIsSpawnable = SpawnRegister.HasSpawnedObject(Binding.GetGuid());
-								if (!bIsSpawnable)
-								{
-									continue;
-								}
+		// File path of the exported USD root layer, in case we also exported the level along with the level sequence
+		FString LevelFilePath;
 
-								// We must keep fixed bindings only as those are kind of like "absolute file paths" and can be used
-								// to identify any of the bindings within this LevelSequence's template hierarchy. The alternative
-								// would be a relative binding ID, which would only be useful within the context of a particular
-								// subsequence of the root sequence
-								UE::MovieScene::FFixedObjectBindingID NewBinding;
-								if (Binding.IsFixedBinding())
-								{
-									NewBinding = Binding.ReinterpretAsFixed();
-								}
-								else
-								{
-									NewBinding = Binding.ResolveToFixed(Binding.GetRelativeSequenceID(), *PinnedSequencer);
-								}
-								OutSelectedSpawnableBindings.Add(NewBinding);
-								OutSelectedSpawnableActors.Add(Actor);
+		// World that is being animated for the LevelSequence bake
+		UWorld* World = nullptr;
 
-								const int32 Count = 1;
-								const bool bAllowShrinking = false;
-								OutSelectedNonSpawnableActors.RemoveAt(Index, Count, bAllowShrinking);
-							}
-						}
-					}
-				}
+	public:
+		~FLevelSequenceExportContext()
+		{
+			SpawnRegister->DeleteSpawns(Sequencer->GetSharedPlaybackState());
+		}
+	};
+
+	bool IsTrackAnimated(const UMovieSceneTrack* Track)
+	{
+		for (const UMovieSceneSection* Section : Track->GetAllSections())
+		{
+			// We can't just check whether a section has range here because it may just have a bunch
+			// of channels with modified default values instead (and so no range or even keys, but can still affect the level).
+			// Sadly there's no way of telling if the default value in the channel has been modified or not (that is, whether it
+			// matches the unanimated value or not), so we'll just have to export any track with an active section for now.
+			// This will emit a bunch of unwanted tracks, but later on we may want to do a post-processing pass on our exported
+			// data to make sure we don't emit intermediary keys on linearly interpolated sections, and that would naturally fix
+			// this too.
+			if (Section->IsActive())
+			{
+				return true;
 			}
+		}
+
+		return false;
+	}
+
+	TMap<UMovieSceneSequence*, TArray<FMovieSceneSequenceID>> GetSequenceHierarchyInstances(UMovieSceneSequence& Sequence, ISequencer& Sequencer)
+	{
+		TMap<UMovieSceneSequence*, TArray<FMovieSceneSequenceID>> SequenceInstances;
+
+		FMovieSceneSequenceHierarchy SequenceHierarchyCache;
+		UMovieSceneCompiledDataManager::CompileHierarchy(&Sequence, &SequenceHierarchyCache, EMovieSceneServerClientMask::All);
 
-			void TransferSpawnableSelection(
-				ISequencer& Sequencer,
-				const TArray<UE::MovieScene::FFixedObjectBindingID>& SelectedSpawnableBindings,
-				TArray<AActor*>& OutSelectedActors
-			)
+		SequenceInstances.FindOrAdd(&Sequence).Add(Sequencer.GetRootTemplateID());
+		for (const TTuple<FMovieSceneSequenceID, FMovieSceneSubSequenceData>& Pair : SequenceHierarchyCache.AllSubSequenceData())
+		{
+			if (ULevelSequence* SubSequence = Cast<ULevelSequence>(Pair.Value.GetSequence()))
 			{
-				USelection* ActorSelection = GEditor->GetSelectedActors();
-				if (!ensure(ActorSelection))
-				{
-					return;
-				}
-				ActorSelection->Modify();
-				ActorSelection->BeginBatchSelectOperation();
+				SequenceInstances.FindOrAdd(SubSequence).Add(Pair.Key);
+			}
+		}
 
-				OutSelectedActors.Reset();
+		return SequenceInstances;
+	}
 
-				for (const UE::MovieScene::FFixedObjectBindingID& FixedBinding : SelectedSpawnableBindings)
-				{
-					FMovieSceneObjectBindingID BindingID{FixedBinding};
-					TArrayView<TWeakObjectPtr<UObject>> BoundObjects = BindingID.ResolveBoundObjects(FixedBinding.SequenceID, Sequencer);
-					for (const TWeakObjectPtr<UObject>& BoundObject : BoundObjects)
-					{
-						if (AActor* Actor = Cast<AActor>(BoundObject.Get()))
-						{
-							ActorSelection->Select(Actor);
+	// Spawn and hide all spawnables for Sequence hierarchy
+	TMap<UMovieSceneSequence*, TArray<FMovieSceneSequenceID>> PreSpawnSpawnables(
+		FLevelSequenceExportContext& Context,
+		UMovieSceneSequence& RootSequence
+	)
+	{
+		TMap<UMovieSceneSequence*, TArray<FMovieSceneSequenceID>> SequenceInstances = GetSequenceHierarchyInstances(
+			RootSequence,
+			Context.Sequencer.Get()
+		);
 
-							OutSelectedActors.Add(Actor);
-						}
-					}
-				}
+		UMovieSceneSequence* OrigRootSequence = Context.Sequencer->GetRootMovieSceneSequence();
+		Context.Sequencer->ResetToNewRootSequence(RootSequence);
 
-				const bool bNotify = true;
-				ActorSelection->EndBatchSelectOperation(bNotify);
+		// Spawn everything for this Sequence hierarchy
+		for (TPair<UMovieSceneSequence*, TArray<FMovieSceneSequenceID>>& SequenceInstancePair : SequenceInstances)
+		{
+			const UMovieSceneSequence* Sequence = SequenceInstancePair.Key;
+			if (!Sequence)
+			{
+				continue;
 			}
 
-			void RestoreEditorSelection(
-				const TArray<AActor*>& SelectedNonSpawnableActors,
-				const TArray<UE::MovieScene::FFixedObjectBindingID>& SelectedSpawnableBindings
-			)
+			UMovieScene* MovieScene = Sequence->GetMovieScene();
+			if (!MovieScene)
 			{
-				USelection* ActorSelection = GEditor->GetSelectedActors();
-				if (!ensure(ActorSelection))
-				{
-					return;
-				}
-				ActorSelection->Modify();
-				ActorSelection->BeginBatchSelectOperation();
-				ActorSelection->DeselectAll();
+				continue;
+			}
 
-				for (AActor* Actor : SelectedNonSpawnableActors)
+			// Spawn everything for this instance
+			for (FMovieSceneSequenceID SequenceInstance : SequenceInstancePair.Value)
+			{
+				int32 NumSpawnables = MovieScene->GetSpawnableCount();
+				for (int32 Index = 0; Index < NumSpawnables; ++Index)
 				{
-					ActorSelection->Select(Actor);
-				}
+					const FMovieSceneSpawnable& Spawnable = MovieScene->GetSpawnable(Index);
+					const FGuid& Guid = Spawnable.GetGuid();
 
-				// Note that we're not tracking *which* sequencer produced which spawnable originally, because realistically the
-				// user will only ever have a single sequencer editing a particular LevelSequence at a time, if that
-				for (const TWeakPtr<ISequencer>& Sequencer : FLevelEditorSequencerIntegration::Get().GetSequencers())
-				{
-					if (TSharedPtr<ISequencer> PinnedSequencer = Sequencer.Pin())
-					{
-						for (const UE::MovieScene::FFixedObjectBindingID& FixedBinding : SelectedSpawnableBindings)
-						{
-							FMovieSceneObjectBindingID BindingID{FixedBinding};
-							TArrayView<TWeakObjectPtr<UObject>> BoundObjects = BindingID.ResolveBoundObjects(FixedBinding.SequenceID, *PinnedSequencer);
-							for (const TWeakObjectPtr<UObject>& BoundObject : BoundObjects)
-							{
-								if (AActor* Actor = Cast<AActor>(BoundObject.Get()))
-								{
-									ActorSelection->Select(Actor);
-								}
-							}
-						}
-					}
+					StaticCastSharedRef<FMovieSceneSpawnRegister>(Context.SpawnRegister)
+						->SpawnObject(Guid, *MovieScene, SequenceInstance, *Context.Sequencer);
 				}
-
-				const bool bNotify = true;
-				ActorSelection->EndBatchSelectOperation(bNotify);
 			}
+		}
 
-			// Appends to InOutComponentBakers all of the component bakers for all components bound to MovieSceneSequence.
-			// In the process it will generate the output prims for each of these components, and keep track of them
-			// within the bakers themselves
-			void GenerateBakersForMovieScene(
-				FLevelSequenceExportContext& Context,
-				UMovieSceneSequence& MovieSceneSequence,
-				const TMap< UMovieSceneSequence*, TArray<FMovieSceneSequenceID> >& SequenceInstances,
-				UE::FUsdStage& UsdStage,
-				TMap<USceneComponent*, FCombinedComponentBakers>& InOutComponentBakers
-			)
+		// Recurse into subsequences if we'll export them as sublayers. This because when exporting each, we'll reset our sequencer
+		// to have each subsequence as the root level sequence again, which means that the sequence ids will all change. We need to
+		// prepare our SpawnRegister to be able to reuse the previous spawns for these new instance ids
+		if (Context.ExportOptions->bExportSubsequencesAsLayers)
+		{
+			for (TPair<UMovieSceneSequence*, TArray<FMovieSceneSequenceID>>& SequenceInstancePair : SequenceInstances)
 			{
-				UMovieScene* MovieScene = MovieSceneSequence.GetMovieScene();
-				if ( !MovieScene )
+				UMovieSceneSequence* Sequence = SequenceInstancePair.Key;
+				if (!Sequence || Sequence == &RootSequence)
 				{
-					return;
+					continue;
 				}
 
-				// Find all objects bound to the movie scene. We'll track UObjects here because we may have tracks
-				// bound to root components but also separate tracks bound directly to the actors, and we want to
-				// capture both.
-				// Index from UObject to FGuid because we may have multiple spawned objects for a given spawnable Guid
-				TMap<UObject*, FGuid> BoundObjects;
-
-				// Collect any USD-related DynamicBinding. The idea being that if we find any, we're likely looking at a
-				// loaded USD Stage that's going to be exported, and the possessable is one of the transient actors and
-				// components. It that's the case, we don't want to just come up with a random name for the prim based on
-				// the actor/component path, but instead want to use the prim path that it has been given on the dynamic binding,
-				// if any
-				TMap<FGuid, const FMovieSceneDynamicBinding*> DynamicBindings;
-
-				const TArray<FMovieSceneSequenceID>* InstancesOfThisSequence = SequenceInstances.Find( &MovieSceneSequence );
-				if ( !InstancesOfThisSequence )
-				{
-					return;
-				}
+				PreSpawnSpawnables(Context, *Sequence);
+			}
+		}
 
-				UMovieSceneSequence* RootSequence = Context.Sequencer->GetRootMovieSceneSequence();
-				if ( !RootSequence )
-				{
-					return;
-				}
+		// Put this back to what it was before this call
+		if (OrigRootSequence)
+		{
+			Context.Sequencer->ResetToNewRootSequence(*OrigRootSequence);
+		}
 
-				// Force spawn spawnables again here so that they exist on the Register map when we rely on FindBoundObjects
-				// to resolve bindings
-				PreSpawnSpawnables(Context, MovieSceneSequence);
+		return SequenceInstances;
+	}
 
-				for ( FMovieSceneSequenceID SequenceInstance : *InstancesOfThisSequence )
-				{
-					FMovieSceneObjectCache& ObjectCache = Context.Sequencer->State.GetObjectCache(SequenceInstance);
+	// Collect selected actors and components before we do anything.
+	// Note that this is especially needed for spawnables since we're going to have to replicate the selection on
+	// the analogous spawnables that belong to our TempSequencer
+	void CollectEditorSelection(
+		const FLevelSequenceHidingSpawnRegister& SpawnRegister,
+		TArray<AActor*>& OutSelectedSpawnableActors,
+		TArray<AActor*>& OutSelectedNonSpawnableActors,
+		TArray<UE::MovieScene::FFixedObjectBindingID>& OutSelectedSpawnableBindings
+	)
+	{
+		OutSelectedSpawnableActors.Reset();
+		OutSelectedNonSpawnableActors.Reset();
+		OutSelectedSpawnableBindings.Reset();
 
-					// Possessables
-					int32 NumPossessables = MovieScene->GetPossessableCount();
-					for ( int32 Index = 0; Index < NumPossessables; ++Index )
-					{
-						const FMovieScenePossessable& Possessable = MovieScene->GetPossessable( Index );
-						const FGuid& Guid = Possessable.GetGuid();
-
-						UObject* BoundObject = nullptr;
-
-						// Go through FMovieSceneObjectCache and FindBoundObjects because that will also evaluate DynamicBindings.
-						// Note that we need to make sure that PreSpawnSpawnables has been called above this (at all, but also at
-						// least once *after* Context.SpawnRegister->CleanUp(), if that has been called). The idea here is that FindBoundObjects
-						// will manage to find the binding even it their "parent context" is a spawnable (e.g. if it's a possessable component of
-						// a spawnable) and also uses DynamicBindings, which is great! For it to be able to find our spawns however, the
-						// spawnables must be *currently* spawned.
-						// It doesn't help at all that our custom spawn register only hides stuff instead of destroying them, because even if the
-						// UObject itself still exists and is just hidden, having been "despawned" means the spawnable has been removed from the
-						// "Register" member of FMovieSceneSpawnRegister, and so the base part of the spawn register "doesn't know about it".
-						// For reference, check how FMovieSceneObjectCache::UpdateBindings (called by FindBoundObjects) will end up calling
-						// "Player.GetSpawnRegister().FindSpawnedObject", and observe how that in turn just checks the "Register" member...
-						// Ideally we could tweak a bit how the "Register" member is used, but that is part of the base FMovieSceneSpawnRegister.
-						TArrayView<TWeakObjectPtr<UObject>> ObjectWeakPtrs = ObjectCache.FindBoundObjects(Guid, *Context.Sequencer);
-						if (ObjectWeakPtrs.Num() > 0)
-						{
-							BoundObject = ObjectWeakPtrs[0].Get();
-						}
+		// We'll start off with all actors inside OutSelectedNonSpawnableActors, and then remove the
+		// spawnables and put them on OutSelectedSpawnableActors
+		USelection* ActorSelection = GEditor->GetSelectedActors();
+		ActorSelection->GetSelectedObjects(OutSelectedNonSpawnableActors);
 
-						if ( BoundObject && ( BoundObject->IsA<USceneComponent>() || BoundObject->IsA<AActor>() ) )
-						{
-							BoundObjects.Add( BoundObject, Guid );
+		OutSelectedSpawnableActors.Reserve(OutSelectedNonSpawnableActors.Num());
 
-							if (Possessable.DynamicBinding.Function)
-							{
-								DynamicBindings.Add(Guid, &Possessable.DynamicBinding);
-							}
-						}
-					}
+		TArray<UE::MovieScene::FFixedObjectBindingID> SpawnableBindings;
+		for (const TWeakPtr<ISequencer>& Sequencer : FLevelEditorSequencerIntegration::Get().GetSequencers())
+		{
+			if (TSharedPtr<ISequencer> PinnedSequencer = Sequencer.Pin())
+			{
+				// Components are never "spawnable", so we only need to get these bindings for actors
+				for (int32 Index = OutSelectedNonSpawnableActors.Num() - 1; Index >= 0; --Index)
+				{
+					AActor* Actor = OutSelectedNonSpawnableActors[Index];
 
-					// Spawnables
-					int32 NumSpawnables = MovieScene->GetSpawnableCount();
-					for ( int32 Index = 0; Index < NumSpawnables; ++Index )
-					{
-						FMovieSceneSpawnable& Spawnable = MovieScene->GetSpawnable( Index );
-						const FGuid& Guid = Spawnable.GetGuid();
+					TArray<FMovieSceneObjectBindingID> Bindings;
+					PinnedSequencer->State.FilterObjectBindings(Actor, *PinnedSequencer, &Bindings);
 
-						// We won't have spawned ASphereReflectionCapture here.
-						// See the comment inside FLevelSequenceHidingSpawnRegister::SpawnObject and UE-167593 for more info
-						if ( ASphereReflectionCapture* ReflectionCapture = Cast< ASphereReflectionCapture>( Spawnable.GetObjectTemplate() ) )
+					OutSelectedSpawnableBindings.Reserve(Bindings.Num() + OutSelectedSpawnableBindings.Num());
+					for (const FMovieSceneObjectBindingID& Binding : Bindings)
+					{
+						// Using the spawn register is an easy way of telling if a Guid is a spawnable or not,
+						// but it's more appropriate because we really only ever care about the spawnables that
+						// we have spawned on our TempSequencer
+						const bool bIsSpawnable = SpawnRegister.HasSpawnedObject(Binding.GetGuid());
+						if (!bIsSpawnable)
 						{
 							continue;
 						}
 
-						UObject* BoundObject = Context.SpawnRegister->GetExistingSpawn( *RootSequence, SequenceInstance, Guid );
-						if ( !BoundObject )
+						// We must keep fixed bindings only as those are kind of like "absolute file paths" and can be used
+						// to identify any of the bindings within this LevelSequence's template hierarchy. The alternative
+						// would be a relative binding ID, which would only be useful within the context of a particular
+						// subsequence of the root sequence
+						UE::MovieScene::FFixedObjectBindingID NewBinding;
+						if (Binding.IsFixedBinding())
 						{
-							// This should never happen as we preemptively spawn everything:
-							// At this point all our spawns should be spawned, but invisible
-							UE_LOG( LogUsd, Warning, TEXT( "Failed to find spawned object for spawnable with Guid '%s'" ), *Guid.ToString() );
-							continue;
+							NewBinding = Binding.ReinterpretAsFixed();
 						}
-
-						if ( BoundObject && ( BoundObject->IsA<USceneComponent>() || BoundObject->IsA<AActor>() ) )
+						else
 						{
-							BoundObjects.Add( BoundObject, Guid );
-
-							if (Spawnable.DynamicBinding.Function)
-							{
-								DynamicBindings.Add(Guid, &Spawnable.DynamicBinding);
-							}
+							NewBinding = Binding.ResolveToFixed(Binding.GetRelativeSequenceID(), *PinnedSequencer);
 						}
+						OutSelectedSpawnableBindings.Add(NewBinding);
+						OutSelectedSpawnableActors.Add(Actor);
+
+						const int32 Count = 1;
+						const bool bAllowShrinking = false;
+						OutSelectedNonSpawnableActors.RemoveAt(Index, Count, bAllowShrinking);
 					}
 				}
+			}
+		}
+	}
 
-				// Generate bakers
-				for ( const TPair<UObject*, FGuid>& Pair : BoundObjects )
-				{
-					UObject* BoundObject = Pair.Key;
-					const FGuid& Guid = Pair.Value;
-
-					// We always use components here because when exporting actors and components to USD we basically
-					// just ignore actors altogether and export the component attachment hierarchy instead
-					USceneComponent* BoundComponent = nullptr;
-					AActor* BoundActor = nullptr;
-					if (USceneComponent* Component = Cast<USceneComponent>(BoundObject))
-					{
-						BoundComponent = Component;
-						BoundActor = BoundComponent->GetOwner();
-					}
-					else if (AActor* Actor = Cast<AActor>(BoundObject))
-					{
-						BoundActor = Actor;
-						BoundComponent = Actor->GetRootComponent();
-					}
+	void TransferSpawnableSelection(
+		ISequencer& Sequencer,
+		const TArray<UE::MovieScene::FFixedObjectBindingID>& SelectedSpawnableBindings,
+		TArray<AActor*>& OutSelectedActors
+	)
+	{
+		USelection* ActorSelection = GEditor->GetSelectedActors();
+		if (!ensure(ActorSelection))
+		{
+			return;
+		}
+		ActorSelection->Modify();
+		ActorSelection->BeginBatchSelectOperation();
 
-					if (!BoundComponent || !BoundActor)
-					{
-						continue;
-					}
-					if (Context.ExportOptions->bSelectionOnly && !Context.SelectedActors.Contains(BoundActor))
-					{
-						continue;
-					}
+		OutSelectedActors.Reset();
 
-					FString PrimPath = UsdUtils::GetPrimPathForObject(
-						BoundComponent,
-						TEXT(""),
-						Context.ExportOptions && Context.ExportOptions->LevelExportOptions.bExportActorFolders
-					);
-					if ( PrimPath.IsEmpty() )
-					{
-						continue;
-					}
+		for (const UE::MovieScene::FFixedObjectBindingID& FixedBinding : SelectedSpawnableBindings)
+		{
+			FMovieSceneObjectBindingID BindingID{FixedBinding};
+			TArrayView<TWeakObjectPtr<UObject>> BoundObjects = BindingID.ResolveBoundObjects(FixedBinding.SequenceID, Sequencer);
+			for (const TWeakObjectPtr<UObject>& BoundObject : BoundObjects)
+			{
+				if (AActor* Actor = Cast<AActor>(BoundObject.Get()))
+				{
+					ActorSelection->Select(Actor);
 
-					// If this binding has one of our dynamic bindings set up pointing to a valid prim path, let's use that path
-					// instead of using our generated PrimPath, as that one will better match the prim paths that we'll get when
-					// opening a referenced stage via an exported UsdStageActor
-					if (const FMovieSceneDynamicBinding* DynamicBinding = DynamicBindings.FindRef(Guid))
-					{
-						if (const FMovieSceneDynamicBindingPayloadVariable* FoundPrimPathPayload = DynamicBinding->PayloadVariables.Find(TEXT("PrimPath")))
-						{
-							FString PrimPathInSourceStage = FoundPrimPathPayload->Value;
-							if (!PrimPathInSourceStage.IsEmpty())
-							{
-								AActor* PossibleParentStageActor = BoundActor;
-								while (PossibleParentStageActor && !PossibleParentStageActor->IsA<AUsdStageActor>())
-								{
-									PossibleParentStageActor = PossibleParentStageActor->GetAttachParentActor();
-								}
+					OutSelectedActors.Add(Actor);
+				}
+			}
+		}
 
-								// Our possessable has a dynamic binding with a "PrimPath" payload variable and is a child of a stage actor,
-								// for now let's consider this enough to consider this is one of our dynamic bindings
-								if (AUsdStageActor* StageActor = Cast<AUsdStageActor>(PossibleParentStageActor))
-								{
-									FString ParentStageActorPrimPathOnExport = UsdUtils::GetPrimPathForObject(
-										StageActor,
-										TEXT(""),
-										Context.ExportOptions && Context.ExportOptions->LevelExportOptions.bExportActorFolders
-									);
-									if (!ParentStageActorPrimPathOnExport.IsEmpty())
-									{
-										UE::FSdfPath SdfPrimPathInSourceStage{*PrimPathInSourceStage};
-
-										UE::FUsdStage LoadedStage = StageActor->GetBaseUsdStage();
-										UE::FSdfPath DefaultPrimPath = LoadedStage.GetDefaultPrim().GetPrimPath();
-
-										if (!LoadedStage.GetRootLayer().IsAnonymous() &&
-											!DefaultPrimPath.IsEmpty() &&
-											SdfPrimPathInSourceStage.HasPrefix(DefaultPrimPath))
-										{
-											// Note that it's perfectly fine if this ends up being just "."
-											UE::FSdfPath RelativePrimPathInSourceStage = SdfPrimPathInSourceStage.MakeRelativePath(DefaultPrimPath);
-
-											FString PrimPathRelativeToParentStageActor = UE::FSdfPath{*ParentStageActorPrimPathOnExport}
-												.AppendPath(RelativePrimPathInSourceStage)
-												.GetString();
-
-											if (!PrimPathRelativeToParentStageActor.IsEmpty())
-											{
-												PrimPath = PrimPathRelativeToParentStageActor;
-											}
-										}
-										else
-										{
-											FUsdLogManager::LogMessage(
-												EMessageSeverity::Warning,
-												FText::Format(
-													LOCTEXT(
-														"NonIdealComposition",
-														"Exported animation for prim '{0}' may not compose correctly with the prims from referenced "
-														"layer '{1}' on the exported stage for the LevelSequence '{2}'. For best results, make sure "
-														"the referenced layer is saved to disk (i.e. not anonymous), has a defaultPrim setup, and "
-														"that the animation tracks are only bound to prims that are descendents of the defaultPrim."
-													),
-													FText::FromString(PrimPathInSourceStage),
-													FText::FromString(LoadedStage.GetRootLayer().GetIdentifier()),
-													FText::FromString(Context.RootSequence.GetPathName())
-												)
-											);
-										}
-									}
-								}
-							}
-						}
-					}
+		const bool bNotify = true;
+		ActorSelection->EndBatchSelectOperation(bNotify);
+	}
 
-					FString SchemaName = UsdUtils::GetSchemaNameForComponent( *BoundComponent );
-					if ( SchemaName.IsEmpty() )
-					{
-						continue;
-					}
+	void RestoreEditorSelection(
+		const TArray<AActor*>& SelectedNonSpawnableActors,
+		const TArray<UE::MovieScene::FFixedObjectBindingID>& SelectedSpawnableBindings
+	)
+	{
+		USelection* ActorSelection = GEditor->GetSelectedActors();
+		if (!ensure(ActorSelection))
+		{
+			return;
+		}
+		ActorSelection->Modify();
+		ActorSelection->BeginBatchSelectOperation();
+		ActorSelection->DeselectAll();
 
-					// We will define a prim here so that we can apply schemas and use the shortcut CreateXAttribute functions,
-					// and not have to worry about attribute names and types. Later on we will convert these prims back into just 'overs' though
-					UE::FUsdPrim Prim = UsdStage.DefinePrim( UE::FSdfPath{ *PrimPath }, *SchemaName );
-					if ( !Prim )
-					{
-						continue;
-					}
+		for (AActor* Actor : SelectedNonSpawnableActors)
+		{
+			ActorSelection->Select(Actor);
+		}
 
-					bool bHasTransformBaker = false;
-					if ( const FMovieSceneBinding* Binding = MovieScene->FindBinding( Guid ) )
+		// Note that we're not tracking *which* sequencer produced which spawnable originally, because realistically the
+		// user will only ever have a single sequencer editing a particular LevelSequence at a time, if that
+		for (const TWeakPtr<ISequencer>& Sequencer : FLevelEditorSequencerIntegration::Get().GetSequencers())
+		{
+			if (TSharedPtr<ISequencer> PinnedSequencer = Sequencer.Pin())
+			{
+				for (const UE::MovieScene::FFixedObjectBindingID& FixedBinding : SelectedSpawnableBindings)
+				{
+					FMovieSceneObjectBindingID BindingID{FixedBinding};
+					TArrayView<TWeakObjectPtr<UObject>> BoundObjects = BindingID.ResolveBoundObjects(FixedBinding.SequenceID, *PinnedSequencer);
+					for (const TWeakObjectPtr<UObject>& BoundObject : BoundObjects)
 					{
-						for ( const UMovieSceneTrack* Track : Binding->GetTracks() )
+						if (AActor* Actor = Cast<AActor>(BoundObject.Get()))
 						{
-							// Let even non-animated spawn tracks through because even if they have zero keyframes we'll still need
-							// to bake the visibility of the spawnables despawning when e.g. a subsequence shot is over
-							if ( !Track || ( !IsTrackAnimated( Track ) && !Track->IsA<UMovieSceneSpawnTrack>() ) )
-							{
-								continue;
-							}
-
-							UnrealToUsd::FComponentBaker Baker;
+							ActorSelection->Select(Actor);
+						}
+					}
+				}
+			}
+		}
 
-							if ( const UMovieScenePropertyTrack* PropertyTrack = Cast<UMovieScenePropertyTrack>( Track ) )
-							{
-								const FString PropertyPath = PropertyTrack->GetPropertyPath().ToString();
-								UnrealToUsd::CreateComponentPropertyBaker( Prim, *BoundComponent, PropertyPath, Baker );
-							}
-							else if ( Track->IsA< UMovieSceneSpawnTrack >() )
-							{
-								// Just handle spawnable tracks as if they're visibility tracks, and hide the prim when not "spawned"
-								// Remember that our spawn register just hides the spawnables when they're not spawned anyway, so this
-								// is essentially the same
-								const FString PropertyPath = TEXT( "bHidden" );
-								UnrealToUsd::CreateComponentPropertyBaker( Prim, *BoundComponent, PropertyPath, Baker);
-							}
-							// Check for the control rig tracks too, because if the user did "Bake to Control Rig" the controlrig code will silently set the
-							// original skeletal animation track sections as disabled, so they'd fail the "IsTrackAnimated" check above
-							else if ( Track->IsA< UMovieSceneSkeletalAnimationTrack >() || Track->IsA< UMovieSceneControlRigParameterTrack >() )
-							{
-								if (USkeletalMeshComponent* SkeletalBoundComponent = Cast<USkeletalMeshComponent>(BoundComponent))
-								{
-									UE::FUsdPrim SkelAnimPrim = UsdStage.DefinePrim(
-										UE::FSdfPath{*PrimPath}.AppendChild(TEXT("Anim")),
-										TEXT("SkelAnimation")
-									);
+		const bool bNotify = true;
+		ActorSelection->EndBatchSelectOperation(bNotify);
+	}
 
-									UE::FUsdPrim SkeletonPrim = UsdStage.DefinePrim(
-										UE::FSdfPath{*PrimPath}.AppendChild(UnrealIdentifiers::ExportedSkeletonPrimName),
-										TEXT("Skeleton")
-									);
+	// Appends to InOutComponentBakers all of the component bakers for all components bound to MovieSceneSequence.
+	// In the process it will generate the output prims for each of these components, and keep track of them
+	// within the bakers themselves
+	void GenerateBakersForMovieScene(
+		FLevelSequenceExportContext& Context,
+		UMovieSceneSequence& MovieSceneSequence,
+		const TMap<UMovieSceneSequence*, TArray<FMovieSceneSequenceID>>& SequenceInstances,
+		UE::FUsdStage& UsdStage,
+		TMap<USceneComponent*, FCombinedComponentBakers>& InOutComponentBakers
+	)
+	{
+		UMovieScene* MovieScene = MovieSceneSequence.GetMovieScene();
+		if (!MovieScene)
+		{
+			return;
+		}
 
-									if (!SkelAnimPrim || !SkeletonPrim)
-									{
-										UE_LOG(
-											LogUsd,
-											Warning,
-											TEXT("Failed to generate Skeleton or SkelAnimation prim when baking out SkelRoot '%s'"),
-											*PrimPath
-										);
-										continue;
-									}
-									UnrealToUsd::CreateSkeletalAnimationBaker(SkeletonPrim, SkelAnimPrim, *SkeletalBoundComponent, Baker);
-								}
-							}
+		// Find all objects bound to the movie scene. We'll track UObjects here because we may have tracks
+		// bound to root components but also separate tracks bound directly to the actors, and we want to
+		// capture both.
+		// Index from UObject to FGuid because we may have multiple spawned objects for a given spawnable Guid
+		TMap<UObject*, FGuid> BoundObjects;
+
+		// Collect any USD-related DynamicBinding. The idea being that if we find any, we're likely looking at a
+		// loaded USD Stage that's going to be exported, and the possessable is one of the transient actors and
+		// components. It that's the case, we don't want to just come up with a random name for the prim based on
+		// the actor/component path, but instead want to use the prim path that it has been given on the dynamic binding,
+		// if any
+		TMap<FGuid, const FMovieSceneDynamicBinding*> DynamicBindings;
+
+		const TArray<FMovieSceneSequenceID>* InstancesOfThisSequence = SequenceInstances.Find(&MovieSceneSequence);
+		if (!InstancesOfThisSequence)
+		{
+			return;
+		}
 
-							// If we made a baker and we don't have one of this type for this component yet, add its lambda to the array
-							FCombinedComponentBakers& ExistingBakers = InOutComponentBakers.FindOrAdd( BoundComponent );
-							if ( Baker.BakerType != UnrealToUsd::EBakingType::None && !EnumHasAnyFlags( ExistingBakers.CombinedBakingType, Baker.BakerType ) )
-							{
-								ExistingBakers.Bakers.Add(Baker);
-								ExistingBakers.CombinedBakingType |= Baker.BakerType;
-							}
+		UMovieSceneSequence* RootSequence = Context.Sequencer->GetRootMovieSceneSequence();
+		if (!RootSequence)
+		{
+			return;
+		}
 
-							if (Baker.BakerType == UnrealToUsd::EBakingType::Transform)
-							{
-								bHasTransformBaker = true;
-							}
-						}
-					}
+		// Force spawn spawnables again here so that they exist on the Register map when we rely on FindBoundObjects
+		// to resolve bindings
+		PreSpawnSpawnables(Context, MovieSceneSequence);
 
-					// If our component is attached to a specific socket of its parent, make sure its transform is baked.
-					// This because if this parent has any AnimSequence animating it, BoundComponent's world transform may change
-					// without it ever having an animated transform, and we don't have any form of rigging/socket attachment on USD that
-					// would cause the parent prim's skeletal animation to also affect its child prims.
-					// Ideally we'd actually search through the tracks to know for sure whether our parent has a SkeletalAnimation section,
-					// but it's probably safer to just do this in case it is hidden behind N subsequences or some obscure feature
-					if (!bHasTransformBaker && BoundComponent->GetAttachSocketName() != NAME_None)
-					{
-						UnrealToUsd::FComponentBaker Baker;
-						const FString PropertyPath = TEXT("Transform");
-						UnrealToUsd::CreateComponentPropertyBaker(Prim, *BoundComponent, PropertyPath, Baker);
+		for (FMovieSceneSequenceID SequenceInstance : *InstancesOfThisSequence)
+		{
+			FMovieSceneObjectCache& ObjectCache = Context.Sequencer->State.GetObjectCache(SequenceInstance);
 
-						FCombinedComponentBakers& ExistingBakers = InOutComponentBakers.FindOrAdd(BoundComponent);
-						if (Baker.BakerType != UnrealToUsd::EBakingType::None && !EnumHasAnyFlags(ExistingBakers.CombinedBakingType, Baker.BakerType))
-						{
-							ExistingBakers.Bakers.Add(Baker);
-							ExistingBakers.CombinedBakingType |= Baker.BakerType;
-						}
-					}
+			// Possessables
+			int32 NumPossessables = MovieScene->GetPossessableCount();
+			for (int32 Index = 0; Index < NumPossessables; ++Index)
+			{
+				const FMovieScenePossessable& Possessable = MovieScene->GetPossessable(Index);
+				const FGuid& Guid = Possessable.GetGuid();
+
+				UObject* BoundObject = nullptr;
+
+				// Go through FMovieSceneObjectCache and FindBoundObjects because that will also evaluate DynamicBindings.
+				// Note that we need to make sure that PreSpawnSpawnables has been called above this (at all, but also at
+				// least once *after* Context.SpawnRegister->CleanUp(), if that has been called). The idea here is that FindBoundObjects
+				// will manage to find the binding even it their "parent context" is a spawnable (e.g. if it's a possessable component of
+				// a spawnable) and also uses DynamicBindings, which is great! For it to be able to find our spawns however, the
+				// spawnables must be *currently* spawned.
+				// It doesn't help at all that our custom spawn register only hides stuff instead of destroying them, because even if the
+				// UObject itself still exists and is just hidden, having been "despawned" means the spawnable has been removed from the
+				// "Register" member of FMovieSceneSpawnRegister, and so the base part of the spawn register "doesn't know about it".
+				// For reference, check how FMovieSceneObjectCache::UpdateBindings (called by FindBoundObjects) will end up calling
+				// "Player.GetSpawnRegister().FindSpawnedObject", and observe how that in turn just checks the "Register" member...
+				// Ideally we could tweak a bit how the "Register" member is used, but that is part of the base FMovieSceneSpawnRegister.
+				TArrayView<TWeakObjectPtr<UObject>> ObjectWeakPtrs = ObjectCache.FindBoundObjects(Guid, *Context.Sequencer);
+				if (ObjectWeakPtrs.Num() > 0)
+				{
+					BoundObject = ObjectWeakPtrs[0].Get();
 				}
 
-				// We always have to generate bakers for subsequences (even if we're not exporting separate files for subsequences) since we
-				// will bake the full combined, composed level sequence as the main output USD layer. This because USD doesn't allow any form
-				// of animation blending, and so composing individual USD layers exported for each subsequence with the same result as the sequencer
-				// is impossible... we have to settle for having each USD layer represent the full effect of it's level sequence.
-				// Also note that we can't share these bakers with our parent movie scenes in case we're a subsequence, unfortunately, because our
-				// bakers will contain lambdas that write directly to a given prim, and those prims are specific to each layer that we're exporting
-				// (e.g. our parent level sequence will export to different prims than this movie scene will)
-				for ( const UMovieSceneTrack* Track : MovieScene->GetTracks() )
+				if (BoundObject && (BoundObject->IsA<USceneComponent>() || BoundObject->IsA<AActor>()))
 				{
-					const UMovieSceneSubTrack* SubTrack = Cast<const UMovieSceneSubTrack>( Track );
-					if ( !SubTrack )
-					{
-						continue;
-					}
+					BoundObjects.Add(BoundObject, Guid);
 
-					for ( const UMovieSceneSection* Section : SubTrack->GetAllSections() )
+					if (Possessable.DynamicBinding.Function)
 					{
-						const UMovieSceneSubSection* SubSection = Cast<const UMovieSceneSubSection>( Section );
-						if ( !SubSection )
-						{
-							continue;
-						}
-
-						UMovieSceneSequence* SubSequence = SubSection->GetSequence();
-						if ( !SubSequence )
-						{
-							continue;
-						}
-
-						GenerateBakersForMovieScene( Context, *SubSequence, SequenceInstances, UsdStage, InOutComponentBakers );
+						DynamicBindings.Add(Guid, &Possessable.DynamicBinding);
 					}
 				}
 			}
 
-			// Steps through MovieSceneSequence frame by frame, invoking all baker lambdas every frame
-			void BakeMovieSceneSequence(
-				FLevelSequenceExportContext& Context,
-				UMovieSceneSequence& MovieSceneSequence,
-				UE::FUsdStage& UsdStage,
-				const TMap<USceneComponent*, FCombinedComponentBakers>& ComponentBakers
-			)
+			// Spawnables
+			int32 NumSpawnables = MovieScene->GetSpawnableCount();
+			for (int32 Index = 0; Index < NumSpawnables; ++Index)
 			{
-				UMovieScene* MovieScene = MovieSceneSequence.GetMovieScene();
-				if ( !MovieScene )
+				FMovieSceneSpawnable& Spawnable = MovieScene->GetSpawnable(Index);
+				const FGuid& Guid = Spawnable.GetGuid();
+
+				// We won't have spawned ASphereReflectionCapture here.
+				// See the comment inside FLevelSequenceHidingSpawnRegister::SpawnObject and UE-167593 for more info
+				if (ASphereReflectionCapture* ReflectionCapture = Cast<ASphereReflectionCapture>(Spawnable.GetObjectTemplate()))
 				{
-					return;
+					continue;
 				}
 
-				// Hide all our spawns again so we can pretend they haven't actually spawned and let the sequence spawn them
-				// as it plays
-				Context.SpawnRegister->bDestroyingJustHides = true;
-				Context.SpawnRegister->CleanUp( *Context.Sequencer );
-
-				const TRange< FFrameNumber > PlaybackRange = MovieScene->GetPlaybackRange();
-				const FFrameRate Resolution = MovieScene->GetTickResolution();
-				const FFrameRate DisplayRate = MovieScene->GetDisplayRate();
-				const double StageTimeCodesPerSecond = UsdStage.GetTimeCodesPerSecond();
-				const FFrameRate StageFrameRate( StageTimeCodesPerSecond, 1 );
-
-				const FFrameTime Interval = FFrameRate::TransformTime( 1, DisplayRate, Resolution );
-				FFrameNumber StartFrame = UE::MovieScene::DiscreteInclusiveLower( PlaybackRange );
-				FFrameNumber EndFrame = UE::MovieScene::DiscreteExclusiveUpper( PlaybackRange );
+				UObject* BoundObject = Context.SpawnRegister->GetExistingSpawn(*RootSequence, SequenceInstance, Guid);
+				if (!BoundObject)
+				{
+					// This should never happen as we preemptively spawn everything:
+					// At this point all our spawns should be spawned, but invisible
+					UE_LOG(LogUsd, Warning, TEXT("Failed to find spawned object for spawnable with Guid '%s'"), *Guid.ToString());
+					continue;
+				}
 
-				if ( Context.ExportOptions && Context.ExportOptions->bOverrideExportRange )
+				if (BoundObject && (BoundObject->IsA<USceneComponent>() || BoundObject->IsA<AActor>()))
 				{
-					StartFrame = FFrameRate::TransformTime( Context.ExportOptions->StartFrame, DisplayRate, Resolution ).FloorToFrame();
-					EndFrame = FFrameRate::TransformTime( Context.ExportOptions->EndFrame, DisplayRate, Resolution ).CeilToFrame();
+					BoundObjects.Add(BoundObject, Guid);
+
+					if (Spawnable.DynamicBinding.Function)
+					{
+						DynamicBindings.Add(Guid, &Spawnable.DynamicBinding);
+					}
 				}
+			}
+		}
 
-				FFrameTime StartFrameUETime = FFrameRate::Snap( StartFrame, Resolution, DisplayRate ).FloorToFrame();
-				double StartTimeCode = FFrameRate::TransformTime( StartFrameUETime, Resolution, StageFrameRate ).AsDecimal();
-				UsdStage.SetStartTimeCode( StartTimeCode );
+		// Generate bakers
+		for (const TPair<UObject*, FGuid>& Pair : BoundObjects)
+		{
+			UObject* BoundObject = Pair.Key;
+			const FGuid& Guid = Pair.Value;
+
+			// We always use components here because when exporting actors and components to USD we basically
+			// just ignore actors altogether and export the component attachment hierarchy instead
+			USceneComponent* BoundComponent = nullptr;
+			AActor* BoundActor = nullptr;
+			if (USceneComponent* Component = Cast<USceneComponent>(BoundObject))
+			{
+				BoundComponent = Component;
+				BoundActor = BoundComponent->GetOwner();
+			}
+			else if (AActor* Actor = Cast<AActor>(BoundObject))
+			{
+				BoundActor = Actor;
+				BoundComponent = Actor->GetRootComponent();
+			}
 
-				FFrameTime EndFrameUETime = FFrameRate::Snap( EndFrame, Resolution, DisplayRate ).FloorToFrame();
-				double EndTimeCode = FFrameRate::TransformTime( EndFrameUETime, Resolution, StageFrameRate ).AsDecimal();
-				UsdStage.SetEndTimeCode( EndTimeCode );
+			if (!BoundComponent || !BoundActor)
+			{
+				continue;
+			}
+			if (Context.ExportOptions->bSelectionOnly && !Context.SelectedActors.Contains(BoundActor))
+			{
+				continue;
+			}
 
-				const FConstraintsManagerController& Controller = FConstraintsManagerController::Get(Context.World);
-				static constexpr bool bSorted = true;
-				const TArray<TWeakObjectPtr<UTickableConstraint>> AllConstraints = Controller.GetAllConstraints(bSorted);
+			FString PrimPath = UsdUtils::GetPrimPathForObject(
+				BoundComponent,
+				TEXT(""),
+				Context.ExportOptions && Context.ExportOptions->LevelExportOptions.bExportActorFolders
+			);
+			if (PrimPath.IsEmpty())
+			{
+				continue;
+			}
 
-				// Collect and sort the bakers: We need all skeletal animations evaluated first, as we need to manually force
-				// the component to update and attached components will only get the correct values if they are evaluated
-				// after that
-				TArray<UnrealToUsd::FComponentBaker> SortedBakers;
-				SortedBakers.Reserve(ComponentBakers.Num() * UnrealToUsd::NumBakingTypes);
-				for (const TPair<USceneComponent*, FCombinedComponentBakers>& Pair : ComponentBakers)
+			// If this binding has one of our dynamic bindings set up pointing to a valid prim path, let's use that path
+			// instead of using our generated PrimPath, as that one will better match the prim paths that we'll get when
+			// opening a referenced stage via an exported UsdStageActor
+			if (const FMovieSceneDynamicBinding* DynamicBinding = DynamicBindings.FindRef(Guid))
+			{
+				if (const FMovieSceneDynamicBindingPayloadVariable* FoundPrimPathPayload = DynamicBinding->PayloadVariables.Find(TEXT("PrimPa"
+																																	  "th")))
 				{
-					for (const UnrealToUsd::FComponentBaker& Baker : Pair.Value.Bakers)
+					FString PrimPathInSourceStage = FoundPrimPathPayload->Value;
+					if (!PrimPathInSourceStage.IsEmpty())
 					{
-						if (Baker.BakerFunction)
+						AActor* PossibleParentStageActor = BoundActor;
+						while (PossibleParentStageActor && !PossibleParentStageActor->IsA<AUsdStageActor>())
 						{
-							SortedBakers.Add(Baker);
+							PossibleParentStageActor = PossibleParentStageActor->GetAttachParentActor();
 						}
-					}
-				}
-				SortedBakers.Sort(
-					[](const UnrealToUsd::FComponentBaker& LHS, const UnrealToUsd::FComponentBaker& RHS) -> bool
-					{
-						if (LHS.BakerType != RHS.BakerType)
+
+						// Our possessable has a dynamic binding with a "PrimPath" payload variable and is a child of a stage actor,
+						// for now let's consider this enough to consider this is one of our dynamic bindings
+						if (AUsdStageActor* StageActor = Cast<AUsdStageActor>(PossibleParentStageActor))
 						{
-							if (LHS.BakerType == UnrealToUsd::EBakingType::Skeletal)
-							{
-								// We want all skeletal bakers first so that the joints are updated before we refresh transforms attached
-								// to joints and sockets
-								return true;
-							}
-							else if (RHS.BakerType == UnrealToUsd::EBakingType::Skeletal)
+							FString ParentStageActorPrimPathOnExport = UsdUtils::GetPrimPathForObject(
+								StageActor,
+								TEXT(""),
+								Context.ExportOptions && Context.ExportOptions->LevelExportOptions.bExportActorFolders
+							);
+							if (!ParentStageActorPrimPathOnExport.IsEmpty())
 							{
-								return false;
+								UE::FSdfPath SdfPrimPathInSourceStage{*PrimPathInSourceStage};
+
+								UE::FUsdStage LoadedStage = StageActor->GetBaseUsdStage();
+								UE::FSdfPath DefaultPrimPath = LoadedStage.GetDefaultPrim().GetPrimPath();
+
+								if (!LoadedStage.GetRootLayer().IsAnonymous() && !DefaultPrimPath.IsEmpty()
+									&& SdfPrimPathInSourceStage.HasPrefix(DefaultPrimPath))
+								{
+									// Note that it's perfectly fine if this ends up being just "."
+									UE::FSdfPath RelativePrimPathInSourceStage = SdfPrimPathInSourceStage.MakeRelativePath(DefaultPrimPath);
+
+									FString PrimPathRelativeToParentStageActor = UE::FSdfPath{*ParentStageActorPrimPathOnExport}
+																					 .AppendPath(RelativePrimPathInSourceStage)
+																					 .GetString();
+
+									if (!PrimPathRelativeToParentStageActor.IsEmpty())
+									{
+										PrimPath = PrimPathRelativeToParentStageActor;
+									}
+								}
+								else
+								{
+									FUsdLogManager::LogMessage(
+										EMessageSeverity::Warning,
+										FText::Format(
+											LOCTEXT(
+												"NonIdealComposition",
+												"Exported animation for prim '{0}' may not compose correctly with the prims from referenced "
+												"layer '{1}' on the exported stage for the LevelSequence '{2}'. For best results, make sure "
+												"the referenced layer is saved to disk (i.e. not anonymous), has a defaultPrim setup, and "
+												"that the animation tracks are only bound to prims that are descendents of the defaultPrim."
+											),
+											FText::FromString(PrimPathInSourceStage),
+											FText::FromString(LoadedStage.GetRootLayer().GetIdentifier()),
+											FText::FromString(Context.RootSequence.GetPathName())
+										)
+									);
+								}
 							}
 						}
-
-						// Parents should go first (so that if we have two nested skeletal mesh components attached to each other, we refresh the
-						// parent first). To be honest this is likely not necessary as child joint transforms shouldn't depend on parent joint
-						// transforms, but we should enforce some consistent ordering anyway so might as well use this
-						return LHS.ComponentPath < RHS.ComponentPath;
 					}
-				);
+				}
+			}
+
+			FString SchemaName = UsdUtils::GetSchemaNameForComponent(*BoundComponent);
+			if (SchemaName.IsEmpty())
+			{
+				continue;
+			}
+
+			// We will define a prim here so that we can apply schemas and use the shortcut CreateXAttribute functions,
+			// and not have to worry about attribute names and types. Later on we will convert these prims back into just 'overs' though
+			UE::FUsdPrim Prim = UsdStage.DefinePrim(UE::FSdfPath{*PrimPath}, *SchemaName);
+			if (!Prim)
+			{
+				continue;
+			}
 
-				for ( FFrameTime EvalTime = StartFrame; EvalTime <= EndFrame; EvalTime += Interval )
+			bool bHasTransformBaker = false;
+			if (const FMovieSceneBinding* Binding = MovieScene->FindBinding(Guid))
+			{
+				for (const UMovieSceneTrack* Track : Binding->GetTracks())
 				{
-					Context.Sequencer->SetLocalTimeDirectly( EvalTime );
-					Context.Sequencer->ForceEvaluate();
+					// Let even non-animated spawn tracks through because even if they have zero keyframes we'll still need
+					// to bake the visibility of the spawnables despawning when e.g. a subsequence shot is over
+					if (!Track || (!IsTrackAnimated(Track) && !Track->IsA<UMovieSceneSpawnTrack>()))
+					{
+						continue;
+					}
+
+					UnrealToUsd::FComponentBaker Baker;
 
-					// Evaluate constraints (these run on tick in the editor, so here we must trigger them manually)
-					for (const TWeakObjectPtr<UTickableConstraint>& Constraint : AllConstraints)
+					if (const UMovieScenePropertyTrack* PropertyTrack = Cast<UMovieScenePropertyTrack>(Track))
+					{
+						const FString PropertyPath = PropertyTrack->GetPropertyPath().ToString();
+						UnrealToUsd::CreateComponentPropertyBaker(Prim, *BoundComponent, PropertyPath, Baker);
+					}
+					else if (Track->IsA<UMovieSceneSpawnTrack>())
+					{
+						// Just handle spawnable tracks as if they're visibility tracks, and hide the prim when not "spawned"
+						// Remember that our spawn register just hides the spawnables when they're not spawned anyway, so this
+						// is essentially the same
+						const FString PropertyPath = TEXT("bHidden");
+						UnrealToUsd::CreateComponentPropertyBaker(Prim, *BoundComponent, PropertyPath, Baker);
+					}
+					// Check for the control rig tracks too, because if the user did "Bake to Control Rig" the controlrig code will silently
+					// set the original skeletal animation track sections as disabled, so they'd fail the "IsTrackAnimated" check above
+					else if (Track->IsA<UMovieSceneSkeletalAnimationTrack>() || Track->IsA<UMovieSceneControlRigParameterTrack>())
 					{
-						if (Constraint.IsValid())
+						if (USkeletalMeshComponent* SkeletalBoundComponent = Cast<USkeletalMeshComponent>(BoundComponent))
 						{
-							Constraint->Evaluate();
+							UE::FUsdPrim SkelAnimPrim = UsdStage.DefinePrim(UE::FSdfPath{*PrimPath}.AppendChild(TEXT("Anim")), TEXT("SkelAnimation"));
+
+							UE::FUsdPrim SkeletonPrim = UsdStage.DefinePrim(
+								UE::FSdfPath{*PrimPath}.AppendChild(UnrealIdentifiers::ExportedSkeletonPrimName),
+								TEXT("Skeleton")
+							);
+
+							if (!SkelAnimPrim || !SkeletonPrim)
+							{
+								UE_LOG(
+									LogUsd,
+									Warning,
+									TEXT("Failed to generate Skeleton or SkelAnimation prim when baking out SkelRoot '%s'"),
+									*PrimPath
+								);
+								continue;
+							}
+							UnrealToUsd::CreateSkeletalAnimationBaker(SkeletonPrim, SkelAnimPrim, *SkeletalBoundComponent, Baker);
 						}
 					}
 
-					FFrameTime KeyTime = FFrameRate::Snap( EvalTime, Resolution, DisplayRate ).FloorToFrame();
-					double UsdTimeCode = FFrameRate::TransformTime(KeyTime, Resolution, StageFrameRate).AsDecimal();
-
-					for (const UnrealToUsd::FComponentBaker& Baker : SortedBakers)
+					// If we made a baker and we don't have one of this type for this component yet, add its lambda to the array
+					FCombinedComponentBakers& ExistingBakers = InOutComponentBakers.FindOrAdd(BoundComponent);
+					if (Baker.BakerType != UnrealToUsd::EBakingType::None && !EnumHasAnyFlags(ExistingBakers.CombinedBakingType, Baker.BakerType))
 					{
-						Baker.BakerFunction(UsdTimeCode);
+						ExistingBakers.Bakers.Add(Baker);
+						ExistingBakers.CombinedBakingType |= Baker.BakerType;
 					}
-				}
 
-				// Convert all prims back to overs (so that this layer doesn't define anything on a stage that doesn't
-				// previously have it - it's only supposed to carry animation data)
-				// We need to do this in this way because when going from 'def' to 'over' we need to do it from leaf towards the
-				// root, as USD doesn't like a parent 'over' with a child 'def'
-				TFunction<void( FUsdPrim& )> RecursiveSetOver;
-				RecursiveSetOver = [&RecursiveSetOver]( FUsdPrim& Prim )
-				{
-					for ( FUsdPrim& Child : Prim.GetChildren() )
+					if (Baker.BakerType == UnrealToUsd::EBakingType::Transform)
 					{
-						RecursiveSetOver( Child );
+						bHasTransformBaker = true;
 					}
+				}
+			}
+
+			// If our component is attached to a specific socket of its parent, make sure its transform is baked.
+			// This because if this parent has any AnimSequence animating it, BoundComponent's world transform may change
+			// without it ever having an animated transform, and we don't have any form of rigging/socket attachment on USD that
+			// would cause the parent prim's skeletal animation to also affect its child prims.
+			// Ideally we'd actually search through the tracks to know for sure whether our parent has a SkeletalAnimation section,
+			// but it's probably safer to just do this in case it is hidden behind N subsequences or some obscure feature
+			if (!bHasTransformBaker && BoundComponent->GetAttachSocketName() != NAME_None)
+			{
+				UnrealToUsd::FComponentBaker Baker;
+				const FString PropertyPath = TEXT("Transform");
+				UnrealToUsd::CreateComponentPropertyBaker(Prim, *BoundComponent, PropertyPath, Baker);
 
-					Prim.SetSpecifier( ESdfSpecifier::Over );
-				};
-				FUsdPrim Root = UsdStage.GetPseudoRoot();
-				for ( FUsdPrim& TopLevelPrim : Root.GetChildren() )
+				FCombinedComponentBakers& ExistingBakers = InOutComponentBakers.FindOrAdd(BoundComponent);
+				if (Baker.BakerType != UnrealToUsd::EBakingType::None && !EnumHasAnyFlags(ExistingBakers.CombinedBakingType, Baker.BakerType))
 				{
-					RecursiveSetOver( TopLevelPrim );
+					ExistingBakers.Bakers.Add(Baker);
+					ExistingBakers.CombinedBakingType |= Baker.BakerType;
 				}
 			}
+		}
+
+		// We always have to generate bakers for subsequences (even if we're not exporting separate files for subsequences) since we
+		// will bake the full combined, composed level sequence as the main output USD layer. This because USD doesn't allow any form
+		// of animation blending, and so composing individual USD layers exported for each subsequence with the same result as the sequencer
+		// is impossible... we have to settle for having each USD layer represent the full effect of it's level sequence.
+		// Also note that we can't share these bakers with our parent movie scenes in case we're a subsequence, unfortunately, because our
+		// bakers will contain lambdas that write directly to a given prim, and those prims are specific to each layer that we're exporting
+		// (e.g. our parent level sequence will export to different prims than this movie scene will)
+		for (const UMovieSceneTrack* Track : MovieScene->GetTracks())
+		{
+			const UMovieSceneSubTrack* SubTrack = Cast<const UMovieSceneSubTrack>(Track);
+			if (!SubTrack)
+			{
+				continue;
+			}
 
-			void ExportMovieSceneSequence(
-				FLevelSequenceExportContext& Context,
-				UMovieSceneSequence& MovieSceneSequence,
-				const FString& FilePath,
-				TArray<UE::FUsdStage>& InOutExportedStages
-			)
+			for (const UMovieSceneSection* Section : SubTrack->GetAllSections())
 			{
-				if ( FilePath.IsEmpty() || Context.ExportedMovieScenes.Contains( &MovieSceneSequence ) )
+				const UMovieSceneSubSection* SubSection = Cast<const UMovieSceneSubSection>(Section);
+				if (!SubSection)
 				{
-					return;
+					continue;
 				}
 
-				UMovieScene* MovieScene = MovieSceneSequence.GetMovieScene();
-				if ( !MovieScene )
+				UMovieSceneSequence* SubSequence = SubSection->GetSequence();
+				if (!SubSequence)
 				{
-					return;
+					continue;
 				}
 
-				// Make sure we don't overwrite a file we just wrote *during this export*.
-				// Overwriting other files is OK, as we want to allow a "repeatedly export over the same files" workflow
-				FString UniqueFilePath = UsdUtils::GetUniqueName( FilePath, Context.UsedFilePaths );
+				GenerateBakersForMovieScene(Context, *SubSequence, SequenceInstances, UsdStage, InOutComponentBakers);
+			}
+		}
+	}
 
-				// Try exporting subsequences if needed
-				if ( Context.ExportOptions )
-				{
-					if ( Context.ExportOptions->bExportSubsequencesAsLayers )
-					{
-						FString Directory;
-						FString FileName;
-						FString Extension;
-						FPaths::Split( UniqueFilePath, Directory, FileName, Extension );
+	// Steps through MovieSceneSequence frame by frame, invoking all baker lambdas every frame
+	void BakeMovieSceneSequence(
+		FLevelSequenceExportContext& Context,
+		UMovieSceneSequence& MovieSceneSequence,
+		UE::FUsdStage& UsdStage,
+		const TMap<USceneComponent*, FCombinedComponentBakers>& ComponentBakers
+	)
+	{
+		UMovieScene* MovieScene = MovieSceneSequence.GetMovieScene();
+		if (!MovieScene)
+		{
+			return;
+		}
 
-						for ( const UMovieSceneTrack* Track : MovieScene->GetTracks() )
-						{
-							const UMovieSceneSubTrack* SubTrack = Cast<const UMovieSceneSubTrack>( Track );
-							if ( !SubTrack )
-							{
-								continue;
-							}
+		// Hide all our spawns again so we can pretend they haven't actually spawned and let the sequence spawn them
+		// as it plays
+		Context.SpawnRegister->bDestroyingJustHides = true;
+		Context.SpawnRegister->CleanUp(*Context.Sequencer);
 
-							for ( const UMovieSceneSection* Section : SubTrack->GetAllSections() )
-							{
-								const UMovieSceneSubSection* SubSection = Cast<const UMovieSceneSubSection>( Section );
-								if ( !SubSection )
-								{
-									continue;
-								}
+		const TRange<FFrameNumber> PlaybackRange = MovieScene->GetPlaybackRange();
+		const FFrameRate Resolution = MovieScene->GetTickResolution();
+		const FFrameRate DisplayRate = MovieScene->GetDisplayRate();
+		const double StageTimeCodesPerSecond = UsdStage.GetTimeCodesPerSecond();
+		const FFrameRate StageFrameRate(StageTimeCodesPerSecond, 1);
 
-								UMovieSceneSequence* SubSequence = SubSection->GetSequence();
-								if ( !SubSequence )
-								{
-									continue;
-								}
+		const FFrameTime Interval = FFrameRate::TransformTime(1, DisplayRate, Resolution);
+		FFrameNumber StartFrame = UE::MovieScene::DiscreteInclusiveLower(PlaybackRange);
+		FFrameNumber EndFrame = UE::MovieScene::DiscreteExclusiveUpper(PlaybackRange);
 
-								FString SubSequencePath = FPaths::Combine( Directory, FString::Printf( TEXT( "%s.%s" ), *SubSequence->GetName(), *Extension ) );
+		if (Context.ExportOptions && Context.ExportOptions->bOverrideExportRange)
+		{
+			StartFrame = FFrameRate::TransformTime(Context.ExportOptions->StartFrame, DisplayRate, Resolution).FloorToFrame();
+			EndFrame = FFrameRate::TransformTime(Context.ExportOptions->EndFrame, DisplayRate, Resolution).CeilToFrame();
+		}
 
-								ExportMovieSceneSequence( Context, *SubSequence, SubSequencePath, InOutExportedStages );
+		FFrameTime StartFrameUETime = FFrameRate::Snap(StartFrame, Resolution, DisplayRate).FloorToFrame();
+		double StartTimeCode = FFrameRate::TransformTime(StartFrameUETime, Resolution, StageFrameRate).AsDecimal();
+		UsdStage.SetStartTimeCode(StartTimeCode);
 
-								// For now we don't want to actually add the subsequence layers as sublayers since each exported level sequence
-								// contains the full baked result anyway, but this is how we'd do it:
-								/*
-								double Offset = 0.0;
-								if ( SubSection->HasStartFrame() )
-								{
-									FFrameNumber LowerBound = SubSection->GetTrueRange().GetLowerBoundValue();
-									FFrameTime LowerBoundTime{ FFrameRate::Snap( LowerBound, Resolution, DisplayRate ).FloorToFrame() };
-									Offset = FFrameRate::TransformTime( LowerBoundTime, Resolution, StageFrameRate ).AsDecimal();
-								}
-								float Scale = SubSection->Parameters.TimeScale;
-								const int32 Index = -1;
-								UsdUtils::InsertSubLayer( UsdStage.GetRootLayer(), *SubSequencePath, Index, Offset, Scale );
-								*/
-							}
-						}
-					}
-				}
+		FFrameTime EndFrameUETime = FFrameRate::Snap(EndFrame, Resolution, DisplayRate).FloorToFrame();
+		double EndTimeCode = FFrameRate::TransformTime(EndFrameUETime, Resolution, StageFrameRate).AsDecimal();
+		UsdStage.SetEndTimeCode(EndTimeCode);
 
-				if ( !IUsdExporterModule::CanExportToLayer( UniqueFilePath ) )
-				{
-					return;
-				}
+		const FConstraintsManagerController& Controller = FConstraintsManagerController::Get(Context.World);
+		static constexpr bool bSorted = true;
+		const TArray<TWeakObjectPtr<UTickableConstraint>> AllConstraints = Controller.GetAllConstraints(bSorted);
 
-				FString LevelSequenceVersion = MovieSceneSequence.GetSignature().ToString();
+		// Collect and sort the bakers: We need all skeletal animations evaluated first, as we need to manually force
+		// the component to update and attached components will only get the correct values if they are evaluated
+		// after that
+		TArray<UnrealToUsd::FComponentBaker> SortedBakers;
+		SortedBakers.Reserve(ComponentBakers.Num() * UnrealToUsd::NumBakingTypes);
+		for (const TPair<USceneComponent*, FCombinedComponentBakers>& Pair : ComponentBakers)
+		{
+			for (const UnrealToUsd::FComponentBaker& Baker : Pair.Value.Bakers)
+			{
+				if (Baker.BakerFunction)
 				{
-					// We could just use the GUID directly but all other asset types end up with SHA hash size so lets be
-					// consistent
-					FSHA1 SHA1;
-					SHA1.UpdateWithString( *LevelSequenceVersion, LevelSequenceVersion.Len() );
-					if ( Context.ExportOptions )
-					{
-						UsdUtils::HashForLevelSequenceExport( *Context.ExportOptions, SHA1 );
-					}
-					SHA1.Final();
-					FSHAHash Hash;
-					SHA1.GetHash( &Hash.Hash[ 0 ] );
-					LevelSequenceVersion = Hash.ToString();
+					SortedBakers.Add(Baker);
 				}
-
-				// Check if we already have exported what we plan on exporting anyway
-				if ( FPaths::FileExists( UniqueFilePath ) )
+			}
+		}
+		SortedBakers.Sort(
+			[](const UnrealToUsd::FComponentBaker& LHS, const UnrealToUsd::FComponentBaker& RHS) -> bool
+			{
+				if (LHS.BakerType != RHS.BakerType)
 				{
-					if ( !Context.bReplaceIdentical )
+					if (LHS.BakerType == UnrealToUsd::EBakingType::Skeletal)
 					{
-						UE_LOG( LogUsd, Log,
-							TEXT( "Skipping export of asset '%s' as the target file '%s' already exists." ),
-							*MovieSceneSequence.GetPathName(),
-							*UExporter::CurrentFilename
-						);
-						return;
+						// We want all skeletal bakers first so that the joints are updated before we refresh transforms attached
+						// to joints and sockets
+						return true;
 					}
-					// If we don't want to re-export this level sequence we need to check if its the same version
-					else if ( !Context.ExportOptions->bReExportIdenticalLevelsAndSequences )
+					else if (RHS.BakerType == UnrealToUsd::EBakingType::Skeletal)
 					{
-						// Don't use the stage cache here as we want this stage to close within this scope in case
-						// we have to overwrite its files due to e.g. missing payload or anything like that
-						const bool bUseStageCache = false;
-						const EUsdInitialLoadSet InitialLoadSet = EUsdInitialLoadSet::LoadNone;
-						if ( UE::FUsdStage TempStage = UnrealUSDWrapper::OpenStage( *UniqueFilePath, InitialLoadSet, bUseStageCache ) )
-						{
-							if ( UE::FUsdPrim RootPrim = TempStage.GetDefaultPrim() )
-							{
-								FUsdUnrealAssetInfo Info = UsdUtils::GetPrimAssetInfo( RootPrim );
-
-								const bool bVersionMatches = !Info.Version.IsEmpty() && Info.Version == LevelSequenceVersion;
-
-								const bool bAssetTypeMatches = !Info.UnrealAssetType.IsEmpty()
-									&& Info.UnrealAssetType == MovieSceneSequence.GetClass()->GetName();
-
-								if ( bVersionMatches && bAssetTypeMatches )
-								{
-									UE_LOG( LogUsd, Log,
-										TEXT( "Skipping export of asset '%s' as the target file '%s' already contains up-to-date exported data." ),
-										*MovieSceneSequence.GetPathName(),
-										*UniqueFilePath
-									);
-									return;
-								}
-							}
-						}
+						return false;
 					}
 				}
 
-				UE::FUsdStage UsdStage = UnrealUSDWrapper::NewStage( *UniqueFilePath );
-				if ( !UsdStage )
-				{
-					return;
-				}
+				// Parents should go first (so that if we have two nested skeletal mesh components attached to each other, we refresh the
+				// parent first). To be honest this is likely not necessary as child joint transforms shouldn't depend on parent joint
+				// transforms, but we should enforce some consistent ordering anyway so might as well use this
+				return LHS.ComponentPath < RHS.ComponentPath;
+			}
+		);
+
+		for (FFrameTime EvalTime = StartFrame; EvalTime <= EndFrame; EvalTime += Interval)
+		{
+			Context.Sequencer->SetLocalTimeDirectly(EvalTime);
+			Context.Sequencer->ForceEvaluate();
 
-				if ( Context.ExportOptions )
+			// Evaluate constraints (these run on tick in the editor, so here we must trigger them manually)
+			for (const TWeakObjectPtr<UTickableConstraint>& Constraint : AllConstraints)
+			{
+				if (Constraint.IsValid())
 				{
-					UsdUtils::SetUsdStageMetersPerUnit( UsdStage, Context.ExportOptions->StageOptions.MetersPerUnit );
-					UsdUtils::SetUsdStageUpAxis( UsdStage, Context.ExportOptions->StageOptions.UpAxis );
-					UsdStage.SetTimeCodesPerSecond( Context.ExportOptions->TimeCodesPerSecond );
+					Constraint->Evaluate();
 				}
+			}
 
-				// Set this so that if we open this exported sequence back up in UE the Sequencer will start up showing
-				// the same DisplayRate as the original exported sequence
-				UsdStage.SetFramesPerSecond( MovieScene->GetDisplayRate().AsDecimal() );
+			FFrameTime KeyTime = FFrameRate::Snap(EvalTime, Resolution, DisplayRate).FloorToFrame();
+			double UsdTimeCode = FFrameRate::TransformTime(KeyTime, Resolution, StageFrameRate).AsDecimal();
 
-				UE::FUsdPrim RootPrim = UsdStage.OverridePrim( UE::FSdfPath( TEXT( "/Root" ) ) );
-				if ( !RootPrim )
-				{
-					return;
-				}
+			for (const UnrealToUsd::FComponentBaker& Baker : SortedBakers)
+			{
+				Baker.BakerFunction(UsdTimeCode);
+			}
+		}
 
-				UsdStage.SetDefaultPrim( RootPrim );
+		// Convert all prims back to overs (so that this layer doesn't define anything on a stage that doesn't
+		// previously have it - it's only supposed to carry animation data)
+		// We need to do this in this way because when going from 'def' to 'over' we need to do it from leaf towards the
+		// root, as USD doesn't like a parent 'over' with a child 'def'
+		TFunction<void(FUsdPrim&)> RecursiveSetOver;
+		RecursiveSetOver = [&RecursiveSetOver](FUsdPrim& Prim)
+		{
+			for (FUsdPrim& Child : Prim.GetChildren())
+			{
+				RecursiveSetOver(Child);
+			}
 
-				// Currently we bake the full composed result of each level sequence into a single layer,
-				// because USD can't compose individual layers in the same way (with blending and so on). So here
-				// we make sure that our sequencer exports each MovieSceneSequence as if it was a root, emitting the
-				// same result as if we had exported that subsequence's LevelSequence by itself
-				Context.Sequencer->ResetToNewRootSequence( MovieSceneSequence );
+			Prim.SetSpecifier(ESdfSpecifier::Over);
+		};
+		FUsdPrim Root = UsdStage.GetPseudoRoot();
+		for (FUsdPrim& TopLevelPrim : Root.GetChildren())
+		{
+			RecursiveSetOver(TopLevelPrim);
+		}
+	}
 
-				TMap< UMovieSceneSequence*, TArray<FMovieSceneSequenceID> > SequenceInstances = GetSequenceHierarchyInstances( MovieSceneSequence, Context.Sequencer.Get() );
+	void ExportMovieSceneSequence(
+		FLevelSequenceExportContext& Context,
+		UMovieSceneSequence& MovieSceneSequence,
+		const FString& FilePath,
+		TArray<UE::FUsdStage>& InOutExportedStages
+	)
+	{
+		if (FilePath.IsEmpty() || Context.ExportedMovieScenes.Contains(&MovieSceneSequence))
+		{
+			return;
+		}
 
-				TMap<USceneComponent*, FCombinedComponentBakers> Bakers;
-				GenerateBakersForMovieScene( Context, MovieSceneSequence, SequenceInstances, UsdStage, Bakers );
+		UMovieScene* MovieScene = MovieSceneSequence.GetMovieScene();
+		if (!MovieScene)
+		{
+			return;
+		}
 
-				// Bake this MovieScene
-				// We bake each MovieScene individually instead of doing one large simultaneous bake because this way
-				// not only we avoid having to handle FMovieSceneSequenceTransforms when writing out the UsdTimeCodes,
-				// we guarantee we'll get the same result as if we exported each subsequence individually.
-				// They could have differed, for example, if we had a limited the playback range of a subsequence
-				BakeMovieSceneSequence( Context, MovieSceneSequence, UsdStage, Bakers );
+		// Make sure we don't overwrite a file we just wrote *during this export*.
+		// Overwriting other files is OK, as we want to allow a "repeatedly export over the same files" workflow
+		FString UniqueFilePath = UsdUtils::GetUniqueName(FilePath, Context.UsedFilePaths);
 
-				// We can add the level as a sublayer to every exported subsequence, so that each can be opened individually and
-				// automatically load the level layer. It doesn't matter much if the parent stage has composed the level
-				// sublayer multiple times (in case we add subsequence layers as sublayers in the future), as the prims will
-				// just all override each other with the same data
-				if ( Context.ExportOptions && Context.ExportOptions->bUseExportedLevelAsSublayer && FPaths::FileExists( Context.LevelFilePath ) )
-				{
-					UsdUtils::InsertSubLayer( UsdStage.GetRootLayer(), *Context.LevelFilePath );
-				}
+		// Try exporting subsequences if needed
+		if (Context.ExportOptions)
+		{
+			if (Context.ExportOptions->bExportSubsequencesAsLayers)
+			{
+				FString Directory;
+				FString FileName;
+				FString Extension;
+				FPaths::Split(UniqueFilePath, Directory, FileName, Extension);
 
-				if (UE::FUsdPrim AssetDefaultPrim = UsdStage.GetDefaultPrim())
+				for (const UMovieSceneTrack* Track : MovieScene->GetTracks())
 				{
-					if (Context.ExportOptions->LevelExportOptions.MetadataOptions.bExportAssetInfo)
+					const UMovieSceneSubTrack* SubTrack = Cast<const UMovieSceneSubTrack>(Track);
+					if (!SubTrack)
+					{
+						continue;
+					}
+
+					for (const UMovieSceneSection* Section : SubTrack->GetAllSections())
 					{
-						FUsdUnrealAssetInfo Info;
-						Info.Name = MovieSceneSequence.GetName();
-						Info.Identifier = UniqueFilePath;
-						Info.Version = LevelSequenceVersion;
-						Info.UnrealContentPath = MovieSceneSequence.GetPathName();
-						Info.UnrealAssetType = MovieSceneSequence.GetClass()->GetName();
-						Info.UnrealExportTime = FDateTime::Now().ToString();
-						Info.UnrealEngineVersion = FEngineVersion::Current().ToString();
-
-						UsdUtils::SetPrimAssetInfo(AssetDefaultPrim, Info);
+						const UMovieSceneSubSection* SubSection = Cast<const UMovieSceneSubSection>(Section);
+						if (!SubSection)
+						{
+							continue;
+						}
+
+						UMovieSceneSequence* SubSequence = SubSection->GetSequence();
+						if (!SubSequence)
+						{
+							continue;
+						}
+
+						FString SubSequencePath = FPaths::Combine(Directory, FString::Printf(TEXT("%s.%s"), *SubSequence->GetName(), *Extension));
+
+						ExportMovieSceneSequence(Context, *SubSequence, SubSequencePath, InOutExportedStages);
+
+						// For now we don't want to actually add the subsequence layers as sublayers since each exported level sequence
+						// contains the full baked result anyway, but this is how we'd do it:
+						/*
+						double Offset = 0.0;
+						if ( SubSection->HasStartFrame() )
+						{
+							FFrameNumber LowerBound = SubSection->GetTrueRange().GetLowerBoundValue();
+							FFrameTime LowerBoundTime{ FFrameRate::Snap( LowerBound, Resolution, DisplayRate ).FloorToFrame() };
+							Offset = FFrameRate::TransformTime( LowerBoundTime, Resolution, StageFrameRate ).AsDecimal();
+						}
+						float Scale = SubSection->Parameters.TimeScale;
+						const int32 Index = -1;
+						UsdUtils::InsertSubLayer( UsdStage.GetRootLayer(), *SubSequencePath, Index, Offset, Scale );
+						*/
 					}
+				}
+			}
+		}
+
+		if (!IUsdExporterModule::CanExportToLayer(UniqueFilePath))
+		{
+			return;
+		}
+
+		FString LevelSequenceVersion = MovieSceneSequence.GetSignature().ToString();
+		{
+			// We could just use the GUID directly but all other asset types end up with SHA hash size so lets be
+			// consistent
+			FSHA1 SHA1;
+			SHA1.UpdateWithString(*LevelSequenceVersion, LevelSequenceVersion.Len());
+			if (Context.ExportOptions)
+			{
+				UsdUtils::HashForLevelSequenceExport(*Context.ExportOptions, SHA1);
+			}
+			SHA1.Final();
+			FSHAHash Hash;
+			SHA1.GetHash(&Hash.Hash[0]);
+			LevelSequenceVersion = Hash.ToString();
+		}
 
-					if (Context.ExportOptions->LevelExportOptions.MetadataOptions.bExportAssetMetadata)
+		// Check if we already have exported what we plan on exporting anyway
+		if (FPaths::FileExists(UniqueFilePath))
+		{
+			if (!Context.bReplaceIdentical)
+			{
+				UE_LOG(
+					LogUsd,
+					Log,
+					TEXT("Skipping export of asset '%s' as the target file '%s' already exists."),
+					*MovieSceneSequence.GetPathName(),
+					*UExporter::CurrentFilename
+				);
+				return;
+			}
+			// If we don't want to re-export this level sequence we need to check if its the same version
+			else if (!Context.ExportOptions->bReExportIdenticalLevelsAndSequences)
+			{
+				// Don't use the stage cache here as we want this stage to close within this scope in case
+				// we have to overwrite its files due to e.g. missing payload or anything like that
+				const bool bUseStageCache = false;
+				const EUsdInitialLoadSet InitialLoadSet = EUsdInitialLoadSet::LoadNone;
+				if (UE::FUsdStage TempStage = UnrealUSDWrapper::OpenStage(*UniqueFilePath, InitialLoadSet, bUseStageCache))
+				{
+					if (UE::FUsdPrim RootPrim = TempStage.GetDefaultPrim())
 					{
-						if (UUsdAssetUserData* UserData = UsdUtils::GetAssetUserData(Cast<ULevelSequence>(&MovieSceneSequence)))
+						FUsdUnrealAssetInfo Info = UsdUtils::GetPrimAssetInfo(RootPrim);
+
+						const bool bVersionMatches = !Info.Version.IsEmpty() && Info.Version == LevelSequenceVersion;
+
+						const bool bAssetTypeMatches = !Info.UnrealAssetType.IsEmpty()
+													   && Info.UnrealAssetType == MovieSceneSequence.GetClass()->GetName();
+
+						if (bVersionMatches && bAssetTypeMatches)
 						{
-							UnrealToUsd::ConvertMetadata(
-								UserData,
-								AssetDefaultPrim,
-								Context.ExportOptions->LevelExportOptions.MetadataOptions.BlockedPrefixFilters,
-								Context.ExportOptions->LevelExportOptions.MetadataOptions.bInvertFilters
+							UE_LOG(
+								LogUsd,
+								Log,
+								TEXT("Skipping export of asset '%s' as the target file '%s' already contains up-to-date exported data."),
+								*MovieSceneSequence.GetPathName(),
+								*UniqueFilePath
 							);
+							return;
 						}
 					}
 				}
+			}
+		}
+
+		UE::FUsdStage UsdStage = UnrealUSDWrapper::NewStage(*UniqueFilePath);
+		if (!UsdStage)
+		{
+			return;
+		}
+
+		if (Context.ExportOptions)
+		{
+			UsdUtils::SetUsdStageMetersPerUnit(UsdStage, Context.ExportOptions->StageOptions.MetersPerUnit);
+			UsdUtils::SetUsdStageUpAxis(UsdStage, Context.ExportOptions->StageOptions.UpAxis);
+			UsdStage.SetTimeCodesPerSecond(Context.ExportOptions->TimeCodesPerSecond);
+		}
+
+		// Set this so that if we open this exported sequence back up in UE the Sequencer will start up showing
+		// the same DisplayRate as the original exported sequence
+		UsdStage.SetFramesPerSecond(MovieScene->GetDisplayRate().AsDecimal());
+
+		UE::FUsdPrim RootPrim = UsdStage.OverridePrim(UE::FSdfPath(TEXT("/Root")));
+		if (!RootPrim)
+		{
+			return;
+		}
 
-				Context.ExportedMovieScenes.Add( &MovieSceneSequence, UniqueFilePath );
-				Context.UsedFilePaths.Add( UniqueFilePath );
+		UsdStage.SetDefaultPrim(RootPrim);
+
+		// Currently we bake the full composed result of each level sequence into a single layer,
+		// because USD can't compose individual layers in the same way (with blending and so on). So here
+		// we make sure that our sequencer exports each MovieSceneSequence as if it was a root, emitting the
+		// same result as if we had exported that subsequence's LevelSequence by itself
+		Context.Sequencer->ResetToNewRootSequence(MovieSceneSequence);
+
+		TMap<UMovieSceneSequence*, TArray<FMovieSceneSequenceID>> SequenceInstances = GetSequenceHierarchyInstances(
+			MovieSceneSequence,
+			Context.Sequencer.Get()
+		);
+
+		TMap<USceneComponent*, FCombinedComponentBakers> Bakers;
+		GenerateBakersForMovieScene(Context, MovieSceneSequence, SequenceInstances, UsdStage, Bakers);
+
+		// Bake this MovieScene
+		// We bake each MovieScene individually instead of doing one large simultaneous bake because this way
+		// not only we avoid having to handle FMovieSceneSequenceTransforms when writing out the UsdTimeCodes,
+		// we guarantee we'll get the same result as if we exported each subsequence individually.
+		// They could have differed, for example, if we had a limited the playback range of a subsequence
+		BakeMovieSceneSequence(Context, MovieSceneSequence, UsdStage, Bakers);
+
+		// We can add the level as a sublayer to every exported subsequence, so that each can be opened individually and
+		// automatically load the level layer. It doesn't matter much if the parent stage has composed the level
+		// sublayer multiple times (in case we add subsequence layers as sublayers in the future), as the prims will
+		// just all override each other with the same data
+		if (Context.ExportOptions && Context.ExportOptions->bUseExportedLevelAsSublayer && FPaths::FileExists(Context.LevelFilePath))
+		{
+			UsdUtils::InsertSubLayer(UsdStage.GetRootLayer(), *Context.LevelFilePath);
+		}
 
-				UsdStage.GetRootLayer().Save();
+		if (UE::FUsdPrim AssetDefaultPrim = UsdStage.GetDefaultPrim())
+		{
+			if (Context.ExportOptions->LevelExportOptions.MetadataOptions.bExportAssetInfo)
+			{
+				FUsdUnrealAssetInfo Info;
+				Info.Name = MovieSceneSequence.GetName();
+				Info.Identifier = UniqueFilePath;
+				Info.Version = LevelSequenceVersion;
+				Info.UnrealContentPath = MovieSceneSequence.GetPathName();
+				Info.UnrealAssetType = MovieSceneSequence.GetClass()->GetName();
+				Info.UnrealExportTime = FDateTime::Now().ToString();
+				Info.UnrealEngineVersion = FEngineVersion::Current().ToString();
+
+				UsdUtils::SetPrimAssetInfo(AssetDefaultPrim, Info);
+			}
 
-				InOutExportedStages.Add( UsdStage );
+			if (Context.ExportOptions->LevelExportOptions.MetadataOptions.bExportAssetMetadata)
+			{
+				if (UUsdAssetUserData* UserData = UsdUtils::GetAssetUserData(Cast<ULevelSequence>(&MovieSceneSequence)))
+				{
+					UnrealToUsd::ConvertMetadata(
+						UserData,
+						AssetDefaultPrim,
+						Context.ExportOptions->LevelExportOptions.MetadataOptions.BlockedPrefixFilters,
+						Context.ExportOptions->LevelExportOptions.MetadataOptions.bInvertFilters
+					);
+				}
 			}
-#endif // USE_USD_SDK
 		}
+
+		Context.ExportedMovieScenes.Add(&MovieSceneSequence, UniqueFilePath);
+		Context.UsedFilePaths.Add(UniqueFilePath);
+
+		UsdStage.GetRootLayer().Save();
+
+		InOutExportedStages.Add(UsdStage);
 	}
-}
+#endif	  // USE_USD_SDK
+}	 // namespace UE::LevelSequenceExporterUSD::Private
 
 ULevelSequenceExporterUsd::ULevelSequenceExporterUsd()
 {
 #if USE_USD_SDK
-	for ( const FString& Extension : UnrealUSDWrapper::GetNativeFileFormats() )
+	for (const FString& Extension : UnrealUSDWrapper::GetNativeFileFormats())
 	{
 		// USDZ is not supported for writing for now
-		if ( Extension.Equals( TEXT( "usdz" ) ) )
+		if (Extension.Equals(TEXT("usdz")))
 		{
 			continue;
 		}
 
-		FormatExtension.Add( Extension );
-		FormatDescription.Add( TEXT( "Universal Scene Description file" ) );
+		FormatExtension.Add(Extension);
+		FormatDescription.Add(TEXT("Universal Scene Description file"));
 	}
 	SupportedClass = ULevelSequence::StaticClass();
 	bText = false;
-#endif // #if USE_USD_SDK
+#endif	  // #if USE_USD_SDK
 }
 
-bool ULevelSequenceExporterUsd::ExportBinary( UObject* Object, const TCHAR* Type, FArchive& Ar, FFeedbackContext* Warn, int32 FileIndex, uint32 PortFlags )
+bool ULevelSequenceExporterUsd::ExportBinary(
+	UObject* Object,
+	const TCHAR* Type,
+	FArchive& Ar,
+	FFeedbackContext* Warn,
+	int32 FileIndex,
+	uint32 PortFlags
+)
 {
 	namespace LevelSequenceExporterImpl = UE::LevelSequenceExporterUSD::Private;
 
 #if USE_USD_SDK
-	ULevelSequence* LevelSequence = Cast< ULevelSequence >( Object );
+	ULevelSequence* LevelSequence = Cast<ULevelSequence>(Object);
 	if (!GEditor || !GIsEditor || !LevelSequence)
 	{
 		return false;
 	}
 
 	UMovieScene* MovieScene = LevelSequence->GetMovieScene();
-	if ( !MovieScene )
+	if (!MovieScene)
 	{
 		return false;
 	}
 
 	ULevelSequenceExporterUsdOptions* Options = nullptr;
-	if ( ExportTask )
+	if (ExportTask)
 	{
-		Options = Cast<ULevelSequenceExporterUsdOptions>( ExportTask->Options );
+		Options = Cast<ULevelSequenceExporterUsdOptions>(ExportTask->Options);
 	}
-	if ( !Options )
+	if (!Options)
 	{
 		Options = GetMutableDefault<ULevelSequenceExporterUsdOptions>();
 
 		// Prompt with an options dialog if we can
-		if ( Options && ( !ExportTask || !ExportTask->bAutomated ) )
+		if (Options && (!ExportTask || !ExportTask->bAutomated))
 		{
-			Options->LevelExportOptions.AssetFolder.Path = FPaths::Combine( FPaths::GetPath( UExporter::CurrentFilename ), TEXT( "Assets" ) );
+			Options->LevelExportOptions.AssetFolder.Path = FPaths::Combine(FPaths::GetPath(UExporter::CurrentFilename), TEXT("Assets"));
 			Options->TimeCodesPerSecond = MovieScene->GetDisplayRate().AsDecimal();
 
-			const bool bContinue = SUsdOptionsWindow::ShowExportOptions( *Options );
-			if ( !bContinue )
+			const bool bContinue = SUsdOptionsWindow::ShowExportOptions(*Options);
+			if (!bContinue)
 			{
 				return false;
 			}
 		}
 	}
-	if ( !Options )
+	if (!Options)
 	{
 		return false;
 	}
@@ -1525,8 +1556,9 @@ bool ULevelSequenceExporterUsd::ExportBinary( UObject* Object, const TCHAR* Type
 
 	double StartTime = FPlatformTime::Cycles64();
 
-	TSharedPtr<LevelSequenceExporterImpl::FLevelSequenceHidingSpawnRegister> SpawnRegister = MakeShared<LevelSequenceExporterImpl::FLevelSequenceHidingSpawnRegister>();
-	if ( !SpawnRegister.IsValid() )
+	TSharedPtr<LevelSequenceExporterImpl::FLevelSequenceHidingSpawnRegister>
+		SpawnRegister = MakeShared<LevelSequenceExporterImpl::FLevelSequenceHidingSpawnRegister>();
+	if (!SpawnRegister.IsValid())
 	{
 		return false;
 	}
@@ -1539,36 +1571,38 @@ bool ULevelSequenceExporterUsd::ExportBinary( UObject* Object, const TCHAR* Type
 	Params.ViewParams.bReadOnly = true;
 	Params.bEditWithinLevelEditor = false;
 
-	// UE-132538: Use a getter for the playback context instead of just binding the world directly because FSequencer::UpdateCachedPlaybackContextAndClient will
-	// ignore the attribute's value and only check the getter
+	// UE-132538: Use a getter for the playback context instead of just binding the world directly because
+	// FSequencer::UpdateCachedPlaybackContextAndClient will ignore the attribute's value and only check the getter
 	UWorld* World = Options->Level.Get();
-	if ( !World )
+	if (!World)
 	{
 		const bool bEditorWorldsOnly = true;
-		World = IUsdClassesModule::GetCurrentWorld( bEditorWorldsOnly );
+		World = IUsdClassesModule::GetCurrentWorld(bEditorWorldsOnly);
 	}
-	Params.PlaybackContext = TAttribute<UObject*>::Create( TAttribute<UObject*>::FGetter::CreateLambda( [World]()
-	{
-		return World;
-	}));
+	Params.PlaybackContext = TAttribute<UObject*>::Create(TAttribute<UObject*>::FGetter::CreateLambda(
+		[World]()
+		{
+			return World;
+		}
+	));
 
 	// Set to read only or else CreateSequencer will change the playback range of the moviescene without even calling Modify()
 	const bool bOldReadOnly = MovieScene->IsReadOnly();
 	const bool bNewReadOnly = true;
-	MovieScene->SetReadOnly( bNewReadOnly );
-	TSharedPtr<ISequencer> TempSequencer = FModuleManager::LoadModuleChecked<ISequencerModule>( "Sequencer" ).CreateSequencer( Params );
-	MovieScene->SetReadOnly( bOldReadOnly );
-	if ( !TempSequencer.IsValid() )
+	MovieScene->SetReadOnly(bNewReadOnly);
+	TSharedPtr<ISequencer> TempSequencer = FModuleManager::LoadModuleChecked<ISequencerModule>("Sequencer").CreateSequencer(Params);
+	MovieScene->SetReadOnly(bOldReadOnly);
+	if (!TempSequencer.IsValid())
 	{
 		return false;
 	}
 
 	TempSequencer->EnterSilentMode();
-	TempSequencer->SetPlaybackStatus( EMovieScenePlayerStatus::Playing );
+	TempSequencer->SetPlaybackStatus(EMovieScenePlayerStatus::Playing);
 
-	SpawnRegister->SetSequencer( TempSequencer );
+	SpawnRegister->SetSequencer(TempSequencer);
 
-	LevelSequenceExporterImpl::FLevelSequenceExportContext Context{ *LevelSequence, TempSequencer.ToSharedRef(), SpawnRegister.ToSharedRef() };
+	LevelSequenceExporterImpl::FLevelSequenceExportContext Context{*LevelSequence, TempSequencer.ToSharedRef(), SpawnRegister.ToSharedRef()};
 	Context.ExportOptions = Options;
 	Context.bReplaceIdentical = ExportTask->bReplaceIdentical;
 	Context.World = World;
@@ -1578,7 +1612,7 @@ bool ULevelSequenceExporterUsd::ExportBinary( UObject* Object, const TCHAR* Type
 	// FMovieSceneSequenceID to spawn a separate object, so that they can become separate prims. Without doing this, if we used the same
 	// subsequence with spawnables multiple times within a parent sequence we'd only get one prim out, as the FMovieSceneSpawnable objects
 	// would be the exact same between all instances of the child sequence (same FGuid)
-	LevelSequenceExporterImpl::PreSpawnSpawnables( Context, *LevelSequence );
+	LevelSequenceExporterImpl::PreSpawnSpawnables(Context, *LevelSequence);
 
 	TArray<AActor*> OriginalSelectedSpawnableActors;
 	TArray<AActor*> OriginalSelectedNonSpawnableActors;
@@ -1619,43 +1653,43 @@ bool ULevelSequenceExporterUsd::ExportBinary( UObject* Object, const TCHAR* Type
 	UAssetEditorSubsystem* AssetEditorSubsystem = nullptr;
 	{
 		AssetEditorSubsystem = GEditor->GetEditorSubsystem<UAssetEditorSubsystem>();
-		if ( AssetEditorSubsystem )
+		if (AssetEditorSubsystem)
 		{
-			for ( const TWeakPtr<ISequencer>& Sequencer : FLevelEditorSequencerIntegration::Get().GetSequencers() )
+			for (const TWeakPtr<ISequencer>& Sequencer : FLevelEditorSequencerIntegration::Get().GetSequencers())
 			{
-				if ( TSharedPtr<ISequencer> PinnedSequencer = Sequencer.Pin() )
+				if (TSharedPtr<ISequencer> PinnedSequencer = Sequencer.Pin())
 				{
-					if ( PinnedSequencer == TempSequencer )
+					if (PinnedSequencer == TempSequencer)
 					{
 						continue;
 					}
 
-					ULevelSequence* OpenedSequence = Cast<ULevelSequence>( PinnedSequencer->GetRootMovieSceneSequence() );
-					if ( !OpenedSequence )
+					ULevelSequence* OpenedSequence = Cast<ULevelSequence>(PinnedSequencer->GetRootMovieSceneSequence());
+					if (!OpenedSequence)
 					{
 						continue;
 					}
 
-					AssetsToReopenEditorsFor.Add( OpenedSequence );
-					AssetEditorSubsystem->CloseAllEditorsForAsset( OpenedSequence );
+					AssetsToReopenEditorsFor.Add(OpenedSequence);
+					AssetEditorSubsystem->CloseAllEditorsForAsset(OpenedSequence);
 				}
 			}
 		}
 	}
 
 	// Export level if we need to
-	if ( Options && Options->bExportLevel )
+	if (Options && Options->bExportLevel)
 	{
-		if ( UWorld* WorldToExport = Options->Level.Get() )
+		if (UWorld* WorldToExport = Options->Level.Get())
 		{
 			// Come up with a file path for the level
 			FString Directory;
 			FString Filename;
 			FString Extension;
-			FPaths::Split( TargetFileName, Directory, Filename, Extension );
-			Context.LevelFilePath = FPaths::Combine( Directory, FString::Printf( TEXT( "%s.%s" ), *WorldToExport->GetFName().ToString(), *Extension ) );
+			FPaths::Split(TargetFileName, Directory, Filename, Extension);
+			Context.LevelFilePath = FPaths::Combine(Directory, FString::Printf(TEXT("%s.%s"), *WorldToExport->GetFName().ToString(), *Extension));
 
-			Context.UsedFilePaths.Add( Context.LevelFilePath );
+			Context.UsedFilePaths.Add(Context.LevelFilePath);
 
 			ULevelExporterUSDOptions* LevelOptions = GetMutableDefault<ULevelExporterUSDOptions>();
 			LevelOptions->StageOptions = Options->StageOptions;
@@ -1664,27 +1698,27 @@ bool ULevelSequenceExporterUsd::ExportBinary( UObject* Object, const TCHAR* Type
 			LevelOptions->bReExportIdenticalLevelsAndSequences = Options->bReExportIdenticalLevelsAndSequences;
 
 			UAssetExportTask* LevelExportTask = NewObject<UAssetExportTask>();
-			FGCObjectScopeGuard ExportTaskGuard( LevelExportTask );
+			FGCObjectScopeGuard ExportTaskGuard(LevelExportTask);
 			LevelExportTask->Object = WorldToExport;
 			LevelExportTask->Options = LevelOptions;
 			LevelExportTask->Exporter = nullptr;
 			LevelExportTask->Filename = Context.LevelFilePath;
-			LevelExportTask->bSelected = LevelOptions->Inner.bSelectionOnly; // Move this as the level exporter will favor bSelected
+			LevelExportTask->bSelected = LevelOptions->Inner.bSelectionOnly;	// Move this as the level exporter will favor bSelected
 			LevelExportTask->bReplaceIdentical = ExportTask->bReplaceIdentical;
 			LevelExportTask->bPrompt = false;
 			LevelExportTask->bUseFileArchive = false;
 			LevelExportTask->bWriteEmptyFiles = false;
-			LevelExportTask->bAutomated = true; // Pretend this is an automated task so it doesn't pop the options dialog
+			LevelExportTask->bAutomated = true;	   // Pretend this is an automated task so it doesn't pop the options dialog
 
-			UExporter::RunAssetExportTask( LevelExportTask );
+			UExporter::RunAssetExportTask(LevelExportTask);
 		}
 	}
 
 	TArray<UE::FUsdStage> ExportedStages;
-	LevelSequenceExporterImpl::ExportMovieSceneSequence( Context, *LevelSequence, TargetFileName, ExportedStages );
+	LevelSequenceExporterImpl::ExportMovieSceneSequence(Context, *LevelSequence, TargetFileName, ExportedStages);
 
 	// Set this back to Stopped or else it will keep the editor viewport controls permanently hidden
-	TempSequencer->SetPlaybackStatus( EMovieScenePlayerStatus::Stopped );
+	TempSequencer->SetPlaybackStatus(EMovieScenePlayerStatus::Stopped);
 	TempSequencer->Close();
 
 	if (AssetEditorSubsystem)
@@ -1697,16 +1731,16 @@ bool ULevelSequenceExporterUsd::ExportBinary( UObject* Object, const TCHAR* Type
 	// Analytics
 	{
 		bool bAutomated = ExportTask ? ExportTask->bAutomated : false;
-		double ElapsedSeconds = FPlatformTime::ToSeconds64( FPlatformTime::Cycles64() - StartTime );
-		FString Extension = FPaths::GetExtension( UExporter::CurrentFilename );
+		double ElapsedSeconds = FPlatformTime::ToSeconds64(FPlatformTime::Cycles64() - StartTime);
+		FString Extension = FPaths::GetExtension(UExporter::CurrentFilename);
 
-		UE::LevelSequenceExporterUSD::Private::SendAnalytics( LevelSequence, Options, ExportedStages, bAutomated, ElapsedSeconds, Extension );
+		UE::LevelSequenceExporterUSD::Private::SendAnalytics(LevelSequence, Options, ExportedStages, bAutomated, ElapsedSeconds, Extension);
 	}
 
 	return true;
 #else
 	return false;
-#endif // #if USE_USD_SDK
+#endif	  // #if USE_USD_SDK
 }
 
 #undef LOCTEXT_NAMESPACE
