USD: Reformat USDExporter module.

diff --git a/Engine/Plugins/Importers/USDImporter/Source/USDExporter/Private/MaterialExporterUSD.cpp b/Engine/Plugins/Importers/USDImporter/Source/USDExporter/Private/MaterialExporterUSD.cpp
--- a/Engine/Plugins/Importers/USDImporter/Source/USDExporter/Private/MaterialExporterUSD.cpp
+++ b/Engine/Plugins/Importers/USDImporter/Source/USDExporter/Private/MaterialExporterUSD.cpp
@@ -46,7 +46,7 @@ namespace UE::MaterialExporterUSD::Private
 		const FString& Extension
 	)
 	{
-		if ( !FEngineAnalytics::IsAvailable() )
+		if (!FEngineAnalytics::IsAvailable())
 		{
 			return;
 		}
@@ -54,13 +54,13 @@ namespace UE::MaterialExporterUSD::Private
 		FString ClassName = Material.GetClass()->GetName();
 
 		TArray<FAnalyticsEventAttribute> EventAttributes;
-		EventAttributes.Emplace( TEXT( "AssetType" ), ClassName );
+		EventAttributes.Emplace(TEXT("AssetType"), ClassName);
 
-		UsdUtils::AddAnalyticsAttributes( Options, EventAttributes );
+		UsdUtils::AddAnalyticsAttributes(Options, EventAttributes);
 
 		IUsdClassesModule::SendAnalytics(
-			MoveTemp( EventAttributes ),
-			FString::Printf( TEXT( "Export.%s" ), *ClassName ),
+			MoveTemp(EventAttributes),
+			FString::Printf(TEXT("Export.%s"), *ClassName),
 			bAutomated,
 			ElapsedSeconds,
 			NumberOfFrames,
@@ -68,25 +68,25 @@ namespace UE::MaterialExporterUSD::Private
 		);
 	}
 
-	void HashMaterial( const UMaterialInterface& Material, FSHA1& VersionHash )
+	void HashMaterial(const UMaterialInterface& Material, FSHA1& VersionHash)
 	{
-		if ( const UMaterial* Parent = Material.GetMaterial() )
+		if (const UMaterial* Parent = Material.GetMaterial())
 		{
 			// Add StateId to the combined hash, which is updated on every compile
-			VersionHash.Update( reinterpret_cast< const uint8* >( &Parent->StateId[ 0 ] ), sizeof( FGuid ) );
+			VersionHash.Update(reinterpret_cast<const uint8*>(&Parent->StateId[0]), sizeof(FGuid));
 
 			// Add material hash to the combined hash, which is updated when switching some material settings
-			if ( const FMaterialResource* MaterialResource = Parent->GetMaterialResource( GMaxRHIFeatureLevel ) )
+			if (const FMaterialResource* MaterialResource = Parent->GetMaterialResource(GMaxRHIFeatureLevel))
 			{
-				MaterialResource = Parent->GetMaterialResource( GMaxRHIFeatureLevel );
+				MaterialResource = Parent->GetMaterialResource(GMaxRHIFeatureLevel);
 
 				FMaterialShaderMapId ShaderMapID;
-				MaterialResource->GetShaderMapId( GMaxRHIShaderPlatform, nullptr, ShaderMapID );
+				MaterialResource->GetShaderMapId(GMaxRHIShaderPlatform, nullptr, ShaderMapID);
 
 				FSHAHash MaterialHash;
-				ShaderMapID.GetMaterialHash( MaterialHash );
+				ShaderMapID.GetMaterialHash(MaterialHash);
 
-				VersionHash.Update( &MaterialHash.Hash[ 0 ], sizeof( MaterialHash.Hash ) );
+				VersionHash.Update(&MaterialHash.Hash[0], sizeof(MaterialHash.Hash));
 			}
 			else
 			{
@@ -94,111 +94,102 @@ namespace UE::MaterialExporterUSD::Private
 				// we never consider this material similar to the one we previously exported, since we can't guarantee
 				// that at all without looking at its shader map ID hash
 				FGuid NewGuid = FGuid::NewGuid();
-				VersionHash.Update( reinterpret_cast< const uint8* >( &NewGuid[ 0 ] ), sizeof( FGuid ) );
+				VersionHash.Update(reinterpret_cast<const uint8*>(&NewGuid[0]), sizeof(FGuid));
 			}
 		}
 
 		// Add material override parameters to the hash
-		if ( const UMaterialInstanceConstant* InstanceConstant = Cast<const UMaterialInstanceConstant>( &Material ) )
+		if (const UMaterialInstanceConstant* InstanceConstant = Cast<const UMaterialInstanceConstant>(&Material))
 		{
 			// MIC already have a nice GUID we can use
-			VersionHash.Update( reinterpret_cast< const uint8* >( &InstanceConstant->ParameterStateId[ 0 ] ), sizeof( FGuid ) );
+			VersionHash.Update(reinterpret_cast<const uint8*>(&InstanceConstant->ParameterStateId[0]), sizeof(FGuid));
 		}
-		else if ( const UMaterialInstance* Instance = Cast<const UMaterialInstance>( &Material ) )
+		else if (const UMaterialInstance* Instance = Cast<const UMaterialInstance>(&Material))
 		{
 			// Manually hash everything...
-			for ( const FScalarParameterValue& Parameter : Instance->ScalarParameterValues )
+			for (const FScalarParameterValue& Parameter : Instance->ScalarParameterValues)
 			{
-				uint32 NameHash = GetTypeHash( Parameter.ParameterInfo.Name );
-				VersionHash.Update( reinterpret_cast< const uint8* >( &NameHash ), sizeof( uint32 ) );
-				VersionHash.Update( reinterpret_cast< const uint8* >( &Parameter.ParameterInfo.Association ), sizeof( Parameter.ParameterInfo.Association ) );
-				VersionHash.Update( reinterpret_cast< const uint8* >( &Parameter.ParameterInfo.Index ), sizeof( Parameter.ParameterInfo.Index ) );
-				VersionHash.Update( reinterpret_cast< const uint8* >( &Parameter.ParameterValue ), sizeof( Parameter.ParameterValue ) );
-				VersionHash.Update( reinterpret_cast< const uint8* >( &Parameter.ExpressionGUID ), sizeof( Parameter.ExpressionGUID ) );
+				uint32 NameHash = GetTypeHash(Parameter.ParameterInfo.Name);
+				VersionHash.Update(reinterpret_cast<const uint8*>(&NameHash), sizeof(uint32));
+				VersionHash.Update(reinterpret_cast<const uint8*>(&Parameter.ParameterInfo.Association), sizeof(Parameter.ParameterInfo.Association));
+				VersionHash.Update(reinterpret_cast<const uint8*>(&Parameter.ParameterInfo.Index), sizeof(Parameter.ParameterInfo.Index));
+				VersionHash.Update(reinterpret_cast<const uint8*>(&Parameter.ParameterValue), sizeof(Parameter.ParameterValue));
+				VersionHash.Update(reinterpret_cast<const uint8*>(&Parameter.ExpressionGUID), sizeof(Parameter.ExpressionGUID));
 			}
-			for ( const FVectorParameterValue& Parameter : Instance->VectorParameterValues )
+			for (const FVectorParameterValue& Parameter : Instance->VectorParameterValues)
 			{
-				uint32 NameHash = GetTypeHash( Parameter.ParameterInfo.Name );
-				VersionHash.Update( reinterpret_cast< const uint8* >( &NameHash ), sizeof( uint32 ) );
-				VersionHash.Update( reinterpret_cast< const uint8* >( &Parameter.ParameterInfo.Association ), sizeof( Parameter.ParameterInfo.Association ) );
-				VersionHash.Update( reinterpret_cast< const uint8* >( &Parameter.ParameterInfo.Index ), sizeof( Parameter.ParameterInfo.Index ) );
-				VersionHash.Update( reinterpret_cast< const uint8* >( &Parameter.ParameterValue ), sizeof( Parameter.ParameterValue ) );
-				VersionHash.Update( reinterpret_cast< const uint8* >( &Parameter.ExpressionGUID ), sizeof( Parameter.ExpressionGUID ) );
+				uint32 NameHash = GetTypeHash(Parameter.ParameterInfo.Name);
+				VersionHash.Update(reinterpret_cast<const uint8*>(&NameHash), sizeof(uint32));
+				VersionHash.Update(reinterpret_cast<const uint8*>(&Parameter.ParameterInfo.Association), sizeof(Parameter.ParameterInfo.Association));
+				VersionHash.Update(reinterpret_cast<const uint8*>(&Parameter.ParameterInfo.Index), sizeof(Parameter.ParameterInfo.Index));
+				VersionHash.Update(reinterpret_cast<const uint8*>(&Parameter.ParameterValue), sizeof(Parameter.ParameterValue));
+				VersionHash.Update(reinterpret_cast<const uint8*>(&Parameter.ExpressionGUID), sizeof(Parameter.ExpressionGUID));
 			}
-			for ( const FDoubleVectorParameterValue& Parameter : Instance->DoubleVectorParameterValues )
+			for (const FDoubleVectorParameterValue& Parameter : Instance->DoubleVectorParameterValues)
 			{
-				uint32 NameHash = GetTypeHash( Parameter.ParameterInfo.Name );
-				VersionHash.Update( reinterpret_cast< const uint8* >( &NameHash ), sizeof( uint32 ) );
-				VersionHash.Update( reinterpret_cast< const uint8* >( &Parameter.ParameterInfo.Association ), sizeof( Parameter.ParameterInfo.Association ) );
-				VersionHash.Update( reinterpret_cast< const uint8* >( &Parameter.ParameterInfo.Index ), sizeof( Parameter.ParameterInfo.Index ) );
-				VersionHash.Update( reinterpret_cast< const uint8* >( &Parameter.ParameterValue ), sizeof( Parameter.ParameterValue ) );
-				VersionHash.Update( reinterpret_cast< const uint8* >( &Parameter.ExpressionGUID ), sizeof( Parameter.ExpressionGUID ) );
+				uint32 NameHash = GetTypeHash(Parameter.ParameterInfo.Name);
+				VersionHash.Update(reinterpret_cast<const uint8*>(&NameHash), sizeof(uint32));
+				VersionHash.Update(reinterpret_cast<const uint8*>(&Parameter.ParameterInfo.Association), sizeof(Parameter.ParameterInfo.Association));
+				VersionHash.Update(reinterpret_cast<const uint8*>(&Parameter.ParameterInfo.Index), sizeof(Parameter.ParameterInfo.Index));
+				VersionHash.Update(reinterpret_cast<const uint8*>(&Parameter.ParameterValue), sizeof(Parameter.ParameterValue));
+				VersionHash.Update(reinterpret_cast<const uint8*>(&Parameter.ExpressionGUID), sizeof(Parameter.ExpressionGUID));
 			}
-			for ( const FTextureParameterValue& Parameter : Instance->TextureParameterValues )
+			for (const FTextureParameterValue& Parameter : Instance->TextureParameterValues)
 			{
-				uint32 NameHash = GetTypeHash( Parameter.ParameterInfo.Name );
-				VersionHash.Update( reinterpret_cast< const uint8* >( &NameHash ), sizeof( uint32 ) );
-				VersionHash.Update( reinterpret_cast< const uint8* >( &Parameter.ParameterInfo.Association ), sizeof( Parameter.ParameterInfo.Association ) );
-				VersionHash.Update( reinterpret_cast< const uint8* >( &Parameter.ParameterInfo.Index ), sizeof( Parameter.ParameterInfo.Index ) );
+				uint32 NameHash = GetTypeHash(Parameter.ParameterInfo.Name);
+				VersionHash.Update(reinterpret_cast<const uint8*>(&NameHash), sizeof(uint32));
+				VersionHash.Update(reinterpret_cast<const uint8*>(&Parameter.ParameterInfo.Association), sizeof(Parameter.ParameterInfo.Association));
+				VersionHash.Update(reinterpret_cast<const uint8*>(&Parameter.ParameterInfo.Index), sizeof(Parameter.ParameterInfo.Index));
 
 				FString TexturePath = Parameter.ParameterValue ? Parameter.ParameterValue->GetPathName() : FString{};
-				VersionHash.UpdateWithString( *TexturePath, TexturePath.Len() );
-				VersionHash.Update( reinterpret_cast< const uint8* >( &Parameter.ExpressionGUID ), sizeof( Parameter.ExpressionGUID ) );
+				VersionHash.UpdateWithString(*TexturePath, TexturePath.Len());
+				VersionHash.Update(reinterpret_cast<const uint8*>(&Parameter.ExpressionGUID), sizeof(Parameter.ExpressionGUID));
 			}
-			for ( const FRuntimeVirtualTextureParameterValue& Parameter : Instance->RuntimeVirtualTextureParameterValues )
+			for (const FRuntimeVirtualTextureParameterValue& Parameter : Instance->RuntimeVirtualTextureParameterValues)
 			{
-				uint32 NameHash = GetTypeHash( Parameter.ParameterInfo.Name );
-				VersionHash.Update( reinterpret_cast< const uint8* >( &NameHash ), sizeof( uint32 ) );
-				VersionHash.Update( reinterpret_cast< const uint8* >( &Parameter.ParameterInfo.Association ), sizeof( Parameter.ParameterInfo.Association ) );
-				VersionHash.Update( reinterpret_cast< const uint8* >( &Parameter.ParameterInfo.Index ), sizeof( Parameter.ParameterInfo.Index ) );
+				uint32 NameHash = GetTypeHash(Parameter.ParameterInfo.Name);
+				VersionHash.Update(reinterpret_cast<const uint8*>(&NameHash), sizeof(uint32));
+				VersionHash.Update(reinterpret_cast<const uint8*>(&Parameter.ParameterInfo.Association), sizeof(Parameter.ParameterInfo.Association));
+				VersionHash.Update(reinterpret_cast<const uint8*>(&Parameter.ParameterInfo.Index), sizeof(Parameter.ParameterInfo.Index));
 
 				FString TexturePath = Parameter.ParameterValue ? Parameter.ParameterValue->GetPathName() : FString{};
-				VersionHash.UpdateWithString( *TexturePath, TexturePath.Len() );
-				VersionHash.Update( reinterpret_cast< const uint8* >( &Parameter.ExpressionGUID ), sizeof( Parameter.ExpressionGUID ) );
+				VersionHash.UpdateWithString(*TexturePath, TexturePath.Len());
+				VersionHash.Update(reinterpret_cast<const uint8*>(&Parameter.ExpressionGUID), sizeof(Parameter.ExpressionGUID));
 			}
-			for ( const FFontParameterValue& Parameter : Instance->FontParameterValues )
+			for (const FFontParameterValue& Parameter : Instance->FontParameterValues)
 			{
-				uint32 NameHash = GetTypeHash( Parameter.ParameterInfo.Name );
-				VersionHash.Update( reinterpret_cast< const uint8* >( &NameHash ), sizeof( uint32 ) );
-				VersionHash.Update( reinterpret_cast< const uint8* >( &Parameter.ParameterInfo.Association ), sizeof( Parameter.ParameterInfo.Association ) );
-				VersionHash.Update( reinterpret_cast< const uint8* >( &Parameter.ParameterInfo.Index ), sizeof( Parameter.ParameterInfo.Index ) );
+				uint32 NameHash = GetTypeHash(Parameter.ParameterInfo.Name);
+				VersionHash.Update(reinterpret_cast<const uint8*>(&NameHash), sizeof(uint32));
+				VersionHash.Update(reinterpret_cast<const uint8*>(&Parameter.ParameterInfo.Association), sizeof(Parameter.ParameterInfo.Association));
+				VersionHash.Update(reinterpret_cast<const uint8*>(&Parameter.ParameterInfo.Index), sizeof(Parameter.ParameterInfo.Index));
 
 				FString FontPath = Parameter.FontValue ? Parameter.FontValue->GetPathName() : FString{};
-				VersionHash.UpdateWithString( *FontPath, FontPath.Len() );
-				VersionHash.Update( reinterpret_cast< const uint8* >( &Parameter.FontPage ), sizeof( Parameter.FontPage ) );
-				VersionHash.Update( reinterpret_cast< const uint8* >( &Parameter.ExpressionGUID ), sizeof( Parameter.ExpressionGUID ) );
+				VersionHash.UpdateWithString(*FontPath, FontPath.Len());
+				VersionHash.Update(reinterpret_cast<const uint8*>(&Parameter.FontPage), sizeof(Parameter.FontPage));
+				VersionHash.Update(reinterpret_cast<const uint8*>(&Parameter.ExpressionGUID), sizeof(Parameter.ExpressionGUID));
 			}
 		}
 	}
 
-	void HashOptions( const FUsdMaterialBakingOptions& Options, FSHA1& VersionHash )
+	void HashOptions(const FUsdMaterialBakingOptions& Options, FSHA1& VersionHash)
 	{
-		VersionHash.Update(
-			reinterpret_cast< const uint8* >( &Options.DefaultTextureSize ),
-			sizeof( Options.DefaultTextureSize )
-		);
+		VersionHash.Update(reinterpret_cast<const uint8*>(&Options.DefaultTextureSize), sizeof(Options.DefaultTextureSize));
 
-		VersionHash.Update(
-			reinterpret_cast<const uint8*>(&Options.bConstantColorAsSingleValue),
-			sizeof(Options.bConstantColorAsSingleValue)
-		);
+		VersionHash.Update(reinterpret_cast<const uint8*>(&Options.bConstantColorAsSingleValue), sizeof(Options.bConstantColorAsSingleValue));
 
-		VersionHash.Update(
-			reinterpret_cast< const uint8* >( Options.Properties.GetData() ),
-			Options.Properties.Num() * Options.Properties.GetTypeSize()
-		);
+		VersionHash.Update(reinterpret_cast<const uint8*>(Options.Properties.GetData()), Options.Properties.Num() * Options.Properties.GetTypeSize());
 
-		VersionHash.UpdateWithString( *Options.TexturesDir.Path, Options.TexturesDir.Path.Len() );
+		VersionHash.UpdateWithString(*Options.TexturesDir.Path, Options.TexturesDir.Path.Len());
 	}
-}
+}	 // namespace UE::MaterialExporterUSD::Private
 
 UMaterialExporterUsd::UMaterialExporterUsd()
 {
 #if USE_USD_SDK
-	for ( const FString& Extension : UnrealUSDWrapper::GetNativeFileFormats() )
+	for (const FString& Extension : UnrealUSDWrapper::GetNativeFileFormats())
 	{
 		// USDZ is not supported for writing for now
-		if ( Extension.Equals( TEXT( "usdz" ) ) )
+		if (Extension.Equals(TEXT("usdz")))
 		{
 			continue;
 		}
@@ -208,40 +199,40 @@ UMaterialExporterUsd::UMaterialExporterUsd()
 	}
 	SupportedClass = UMaterialInterface::StaticClass();
 	bText = false;
-#endif // #if USE_USD_SDK
+#endif	  // #if USE_USD_SDK
 }
 
-bool UMaterialExporterUsd::ExportBinary( UObject* Object, const TCHAR* Type, FArchive& Ar, FFeedbackContext* Warn, int32 FileIndex, uint32 PortFlags )
+bool UMaterialExporterUsd::ExportBinary(UObject* Object, const TCHAR* Type, FArchive& Ar, FFeedbackContext* Warn, int32 FileIndex, uint32 PortFlags)
 {
 #if USE_USD_SDK
-	UMaterialInterface* Material = Cast< UMaterialInterface >( Object );
-	if ( !Material )
+	UMaterialInterface* Material = Cast<UMaterialInterface>(Object);
+	if (!Material)
 	{
 		return false;
 	}
 
 	UMaterialExporterUSDOptions* Options = nullptr;
-	if ( ExportTask )
+	if (ExportTask)
 	{
-		Options = Cast<UMaterialExporterUSDOptions>( ExportTask->Options );
+		Options = Cast<UMaterialExporterUSDOptions>(ExportTask->Options);
 	}
-	if ( !Options )
+	if (!Options)
 	{
 		Options = GetMutableDefault<UMaterialExporterUSDOptions>();
 
 		// Prompt with an options dialog if we can
-		if ( Options && ( !ExportTask || !ExportTask->bAutomated ) )
+		if (Options && (!ExportTask || !ExportTask->bAutomated))
 		{
-			Options->MaterialBakingOptions.TexturesDir.Path = FPaths::Combine( FPaths::GetPath( UExporter::CurrentFilename ), TEXT( "Textures" ) );
+			Options->MaterialBakingOptions.TexturesDir.Path = FPaths::Combine(FPaths::GetPath(UExporter::CurrentFilename), TEXT("Textures"));
 
-			const bool bContinue = SUsdOptionsWindow::ShowExportOptions( *Options );
-			if ( !bContinue )
+			const bool bContinue = SUsdOptionsWindow::ShowExportOptions(*Options);
+			if (!bContinue)
 			{
 				return false;
 			}
 		}
 	}
-	if ( !Options )
+	if (!Options)
 	{
 		return false;
 	}
@@ -251,10 +242,10 @@ bool UMaterialExporterUsd::ExportBinary( UObject* Object, const TCHAR* Type, FAr
 
 	// Make sure the material has a resource as we'll use it to fetch some version GUIDs and ExportMaterial
 	// can't do it directly due to receiving the material by const ref
-	if ( UMaterial* Parent = Material->GetMaterial() )
+	if (UMaterial* Parent = Material->GetMaterial())
 	{
-		const FMaterialResource* MaterialResource = Parent->GetMaterialResource( GMaxRHIFeatureLevel );
-		if ( MaterialResource == nullptr )
+		const FMaterialResource* MaterialResource = Parent->GetMaterialResource(GMaxRHIFeatureLevel);
+		if (MaterialResource == nullptr)
 		{
 			Parent->ForceRecompileForRendering();
 		}
@@ -264,14 +255,14 @@ bool UMaterialExporterUsd::ExportBinary( UObject* Object, const TCHAR* Type, FAr
 		*Material,
 		Options->MaterialBakingOptions,
 		Options->MetadataOptions,
-		FFilePath{ UExporter::CurrentFilename },
+		FFilePath{UExporter::CurrentFilename},
 		ExportTask->bReplaceIdentical,
 		Options->bReExportIdenticalAssets,
 		ExportTask->bAutomated
 	);
 #else
 	return false;
-#endif // #if USE_USD_SDK
+#endif	  // #if USE_USD_SDK
 }
 
 bool UMaterialExporterUsd::ExportMaterial(
@@ -288,51 +279,54 @@ bool UMaterialExporterUsd::ExportMaterial(
 
 	FSHAHash Hash;
 	FSHA1 SHA1;
-	UE::MaterialExporterUSD::Private::HashMaterial( Material, SHA1 );
-	UE::MaterialExporterUSD::Private::HashOptions( Options, SHA1 );
+	UE::MaterialExporterUSD::Private::HashMaterial(Material, SHA1);
+	UE::MaterialExporterUSD::Private::HashOptions(Options, SHA1);
 	SHA1.Final();
-	SHA1.GetHash( &Hash.Hash[ 0 ] );
+	SHA1.GetHash(&Hash.Hash[0]);
 	FString MaterialHashString = Hash.ToString();
 
-	if ( !IUsdExporterModule::CanExportToLayer( FilePath.FilePath ) )
+	if (!IUsdExporterModule::CanExportToLayer(FilePath.FilePath))
 	{
 		return false;
 	}
 
 	// Check if we already have exported what we plan on exporting anyway
-	if ( FPaths::FileExists( FilePath.FilePath ) )
+	if (FPaths::FileExists(FilePath.FilePath))
 	{
-		if ( !bReplaceIdentical )
+		if (!bReplaceIdentical)
 		{
-			UE_LOG( LogUsd, Log,
-				TEXT( "Skipping export of asset '%s' as the target file '%s' already exists." ),
+			UE_LOG(
+				LogUsd,
+				Log,
+				TEXT("Skipping export of asset '%s' as the target file '%s' already exists."),
 				*Material.GetPathName(),
 				*UExporter::CurrentFilename
 			);
 			return false;
 		}
 		// If we don't want to re-export this asset we need to check if its the same version
-		else if ( !bReExportIdenticalAssets )
+		else if (!bReExportIdenticalAssets)
 		{
 			// Don't use the asset cache here as we want this stage to close within this scope in case
 			// we have to overwrite its files due to e.g. missing payload or anything like that
 			const bool bUseAssetCache = false;
 			const EUsdInitialLoadSet InitialLoadSet = EUsdInitialLoadSet::LoadNone;
-			if ( UE::FUsdStage TempStage = UnrealUSDWrapper::OpenStage( *FilePath.FilePath, InitialLoadSet, bUseAssetCache ) )
+			if (UE::FUsdStage TempStage = UnrealUSDWrapper::OpenStage(*FilePath.FilePath, InitialLoadSet, bUseAssetCache))
 			{
-				if ( UE::FUsdPrim DefaultPrim = TempStage.GetDefaultPrim() )
+				if (UE::FUsdPrim DefaultPrim = TempStage.GetDefaultPrim())
 				{
-					FUsdUnrealAssetInfo Info = UsdUtils::GetPrimAssetInfo( DefaultPrim );
+					FUsdUnrealAssetInfo Info = UsdUtils::GetPrimAssetInfo(DefaultPrim);
 
 					const bool bVersionMatches = !Info.Version.IsEmpty() && Info.Version == MaterialHashString;
 
-					const bool bAssetTypeMatches = !Info.UnrealAssetType.IsEmpty()
-						&& Info.UnrealAssetType == Material.GetClass()->GetName();
+					const bool bAssetTypeMatches = !Info.UnrealAssetType.IsEmpty() && Info.UnrealAssetType == Material.GetClass()->GetName();
 
-					if ( bVersionMatches && bAssetTypeMatches )
+					if (bVersionMatches && bAssetTypeMatches)
 					{
-						UE_LOG( LogUsd, Log,
-							TEXT( "Skipping export of asset '%s' as the target file '%s' already contains up-to-date exported data." ),
+						UE_LOG(
+							LogUsd,
+							Log,
+							TEXT("Skipping export of asset '%s' as the target file '%s' already contains up-to-date exported data."),
 							*Material.GetPathName(),
 							*FilePath.FilePath
 						);
@@ -345,23 +339,23 @@ bool UMaterialExporterUsd::ExportMaterial(
 
 	double StartTime = FPlatformTime::Cycles64();
 
-	UE::FUsdStage UsdStage = UnrealUSDWrapper::NewStage( *FilePath.FilePath );
-	if ( !UsdStage )
+	UE::FUsdStage UsdStage = UnrealUSDWrapper::NewStage(*FilePath.FilePath);
+	if (!UsdStage)
 	{
 		return false;
 	}
 
-	FString RootPrimPath = TEXT( "/" ) + UsdUtils::SanitizeUsdIdentifier( *Material.GetName() );
+	FString RootPrimPath = TEXT("/") + UsdUtils::SanitizeUsdIdentifier(*Material.GetName());
 
-	UE::FUsdPrim RootPrim = UsdStage.DefinePrim( UE::FSdfPath( *RootPrimPath ), TEXT( "Material" ) );
-	if ( !RootPrim )
+	UE::FUsdPrim RootPrim = UsdStage.DefinePrim(UE::FSdfPath(*RootPrimPath), TEXT("Material"));
+	if (!RootPrim)
 	{
 		return false;
 	}
 
-	UsdStage.SetDefaultPrim( RootPrim );
+	UsdStage.SetDefaultPrim(RootPrim);
 
-	UsdUtils::SetUnrealSurfaceOutput( RootPrim, Material.GetPathName() );
+	UsdUtils::SetUnrealSurfaceOutput(RootPrim, Material.GetPathName());
 
 	UnrealToUsd::ConvertMaterialToBakedSurface(
 		Material,
@@ -390,12 +384,7 @@ bool UMaterialExporterUsd::ExportMaterial(
 	{
 		if (UUsdAssetUserData* UserData = UsdUtils::GetAssetUserData(&Material))
 		{
-			UnrealToUsd::ConvertMetadata(
-				UserData,
-				RootPrim,
-				MetadataOptions.BlockedPrefixFilters,
-				MetadataOptions.bInvertFilters
-			);
+			UnrealToUsd::ConvertMetadata(UserData, RootPrim, MetadataOptions.BlockedPrefixFilters, MetadataOptions.bInvertFilters);
 		}
 	}
 
@@ -403,24 +392,17 @@ bool UMaterialExporterUsd::ExportMaterial(
 
 	// Analytics
 	{
-		double ElapsedSeconds = FPlatformTime::ToSeconds64( FPlatformTime::Cycles64() - StartTime );
-		FString Extension = FPaths::GetExtension( FilePath.FilePath );
+		double ElapsedSeconds = FPlatformTime::ToSeconds64(FPlatformTime::Cycles64() - StartTime);
+		FString Extension = FPaths::GetExtension(FilePath.FilePath);
 		double NumberOfFrames = 1 + UsdStage.GetEndTimeCode() - UsdStage.GetStartTimeCode();
 
-		UE::MaterialExporterUSD::Private::SendAnalytics(
-			Material,
-			Options,
-			bIsAutomated,
-			ElapsedSeconds,
-			NumberOfFrames,
-			Extension
-		);
+		UE::MaterialExporterUSD::Private::SendAnalytics(Material, Options, bIsAutomated, ElapsedSeconds, NumberOfFrames, Extension);
 	}
 
 	return true;
 #else
 	return false;
-#endif // USE_USD_SDK
+#endif	  // USE_USD_SDK
 }
 
 // Deprecated signature
@@ -445,7 +427,6 @@ bool UMaterialExporterUsd::ExportMaterial(
 	);
 }
 
-
 bool UMaterialExporterUsd::ExportMaterialsForStage(
 	const TArray<UMaterialInterface*>& Materials,
 	const FUsdMaterialBakingOptions& Options,
@@ -459,17 +440,17 @@ bool UMaterialExporterUsd::ExportMaterialsForStage(
 )
 {
 #if USE_USD_SDK
-	if ( StageRootLayerPath.IsEmpty() )
+	if (StageRootLayerPath.IsEmpty())
 	{
 		return false;
 	}
 
-	if ( Materials.Num() == 0 )
+	if (Materials.Num() == 0)
 	{
 		return true;
 	}
 
-	const FString ExtensionNoDot = FPaths::GetExtension( StageRootLayerPath );
+	const FString ExtensionNoDot = FPaths::GetExtension(StageRootLayerPath);
 
 	// If we have multiple materials *within this mesh* that want to be emitted to the same filepath we'll append
 	// a suffix, but we will otherwise overwrite any unrelated existing files that were there before we began the export.
@@ -477,19 +458,19 @@ bool UMaterialExporterUsd::ExportMaterialsForStage(
 	TSet<FString> UsedFilePathsWithoutExt;
 
 	TMap<FString, FString> MaterialPathNameToFilePath;
-	for ( UMaterialInterface* Material : Materials )
+	for (UMaterialInterface* Material : Materials)
 	{
-		if ( !Material )
+		if (!Material)
 		{
 			continue;
 		}
 
 		// Make sure the material has a resource as we'll use it to fetch some version GUIDs and ExportMaterial
 		// can't do it directly due to receiving the material by const ref
-		if ( UMaterial* Parent = Material->GetMaterial() )
+		if (UMaterial* Parent = Material->GetMaterial())
 		{
-			const FMaterialResource* MaterialResource = Parent->GetMaterialResource( GMaxRHIFeatureLevel );
-			if ( MaterialResource == nullptr )
+			const FMaterialResource* MaterialResource = Parent->GetMaterialResource(GMaxRHIFeatureLevel);
+			if (MaterialResource == nullptr)
 			{
 				Parent->ForceRecompileForRendering();
 			}
@@ -499,23 +480,23 @@ bool UMaterialExporterUsd::ExportMaterialsForStage(
 		FString MaterialPathName = Material->GetPathName();
 
 		// "C:/MyFolder/Export/Blue"
-		FString MaterialFilePath = FPaths::Combine( FPaths::GetPath( UExporter::CurrentFilename ), FPaths::GetBaseFilename( MaterialPathName ) );
+		FString MaterialFilePath = FPaths::Combine(FPaths::GetPath(UExporter::CurrentFilename), FPaths::GetBaseFilename(MaterialPathName));
 
 		// "C:/MyFolder/Export/Blue_4"
-		FString FinalPathNoExt = UsdUtils::GetUniqueName( MaterialFilePath, UsedFilePathsWithoutExt );
+		FString FinalPathNoExt = UsdUtils::GetUniqueName(MaterialFilePath, UsedFilePathsWithoutExt);
 
 		// "C:/MyFolder/Export/Blue_4.usda"
-		FString FinalPath = FString::Printf( TEXT( "%s.%s" ), *FinalPathNoExt, *ExtensionNoDot );
+		FString FinalPath = FString::Printf(TEXT("%s.%s"), *FinalPathNoExt, *ExtensionNoDot);
 
 		if (UMaterialExporterUsd::ExportMaterial(
-			*Material,
-			Options,
-			MetadataOptions,
-			FFilePath{FinalPath},
-			bReplaceIdentical,
-			bReExportIdenticalAssets,
-			bIsAutomated
-		))
+				*Material,
+				Options,
+				MetadataOptions,
+				FFilePath{FinalPath},
+				bReplaceIdentical,
+				bReExportIdenticalAssets,
+				bIsAutomated
+			))
 		{
 			UsedFilePathsWithoutExt.Add(FinalPathNoExt);
 			MaterialPathNameToFilePath.Add(MaterialPathName, FinalPath);
@@ -527,32 +508,26 @@ bool UMaterialExporterUsd::ExportMaterialsForStage(
 	// material exports would try to replace, meaning the exports would fail as those files would be currently open.
 	const EUsdInitialLoadSet InitialLoadSet = EUsdInitialLoadSet::LoadAll;
 	const bool bUseStageCache = false;
-	UE::FUsdStage Stage = UnrealUSDWrapper::OpenStage( *StageRootLayerPath, InitialLoadSet, bUseStageCache );
-	if ( !Stage )
+	UE::FUsdStage Stage = UnrealUSDWrapper::OpenStage(*StageRootLayerPath, InitialLoadSet, bUseStageCache);
+	if (!Stage)
 	{
 		return false;
 	}
 
 	UE::FSdfLayer RootLayer = Stage.GetRootLayer();
-	if ( !RootLayer )
+	if (!RootLayer)
 	{
 		return false;
 	}
 
-	UsdUtils::ReplaceUnrealMaterialsWithBaked(
-		Stage,
-		RootLayer,
-		MaterialPathNameToFilePath,
-		bIsAssetLayer,
-		bUsePayload
-	);
+	UsdUtils::ReplaceUnrealMaterialsWithBaked(Stage, RootLayer, MaterialPathNameToFilePath, bIsAssetLayer, bUsePayload);
 
 	RootLayer.Save();
 
 	return true;
 #else
 	return false;
-#endif // USE_USD_SDK
+#endif	  // USE_USD_SDK
 }
 
 // Deprecated signature
diff --git a/Engine/Plugins/Importers/USDImporter/Source/USDExporter/Private/SkeletalMeshExporterUSD.cpp b/Engine/Plugins/Importers/USDImporter/Source/USDExporter/Private/SkeletalMeshExporterUSD.cpp
--- a/Engine/Plugins/Importers/USDImporter/Source/USDExporter/Private/SkeletalMeshExporterUSD.cpp
+++ b/Engine/Plugins/Importers/USDImporter/Source/USDExporter/Private/SkeletalMeshExporterUSD.cpp
@@ -38,7 +38,7 @@ namespace UE::SkeletalMeshExporterUSD::Private
 		const FString& Extension
 	)
 	{
-		if ( !Asset || !FEngineAnalytics::IsAvailable() )
+		if (!Asset || !FEngineAnalytics::IsAvailable())
 		{
 			return;
 		}
@@ -46,16 +46,16 @@ namespace UE::SkeletalMeshExporterUSD::Private
 		FString ClassName = Asset->GetClass()->GetName();
 
 		TArray<FAnalyticsEventAttribute> EventAttributes;
-		EventAttributes.Emplace( TEXT( "AssetType" ), ClassName );
+		EventAttributes.Emplace(TEXT("AssetType"), ClassName);
 
-		if ( Options )
+		if (Options)
 		{
-			UsdUtils::AddAnalyticsAttributes( *Options, EventAttributes );
+			UsdUtils::AddAnalyticsAttributes(*Options, EventAttributes);
 		}
 
 		IUsdClassesModule::SendAnalytics(
-			MoveTemp( EventAttributes ),
-			FString::Printf( TEXT( "Export.%s" ), *ClassName ),
+			MoveTemp(EventAttributes),
+			FString::Printf(TEXT("Export.%s"), *ClassName),
 			bAutomated,
 			ElapsedSeconds,
 			NumberOfFrames,
@@ -67,10 +67,10 @@ namespace UE::SkeletalMeshExporterUSD::Private
 USkeletalMeshExporterUsd::USkeletalMeshExporterUsd()
 {
 #if USE_USD_SDK
-	for ( const FString& Extension : UnrealUSDWrapper::GetNativeFileFormats() )
+	for (const FString& Extension : UnrealUSDWrapper::GetNativeFileFormats())
 	{
 		// USDZ is not supported for writing for now
-		if ( Extension.Equals( TEXT( "usdz" ) ) )
+		if (Extension.Equals(TEXT("usdz")))
 		{
 			continue;
 		}
@@ -80,40 +80,50 @@ USkeletalMeshExporterUsd::USkeletalMeshExporterUsd()
 	}
 	SupportedClass = USkeletalMesh::StaticClass();
 	bText = false;
-#endif // #if USE_USD_SDK
+#endif	  // #if USE_USD_SDK
 }
 
-bool USkeletalMeshExporterUsd::ExportBinary( UObject* Object, const TCHAR* Type, FArchive& Ar, FFeedbackContext* Warn, int32 FileIndex, uint32 PortFlags )
+bool USkeletalMeshExporterUsd::ExportBinary(
+	UObject* Object,
+	const TCHAR* Type,
+	FArchive& Ar,
+	FFeedbackContext* Warn,
+	int32 FileIndex,
+	uint32 PortFlags
+)
 {
 #if USE_USD_SDK
-	USkeletalMesh* SkeletalMesh = CastChecked< USkeletalMesh >( Object );
-	if ( !SkeletalMesh )
+	USkeletalMesh* SkeletalMesh = CastChecked<USkeletalMesh>(Object);
+	if (!SkeletalMesh)
 	{
 		return false;
 	}
 
 	USkeletalMeshExporterUSDOptions* Options = nullptr;
-	if ( ExportTask )
+	if (ExportTask)
 	{
-		Options = Cast<USkeletalMeshExporterUSDOptions>( ExportTask->Options );
+		Options = Cast<USkeletalMeshExporterUSDOptions>(ExportTask->Options);
 	}
-	if ( !Options )
+	if (!Options)
 	{
 		Options = GetMutableDefault<USkeletalMeshExporterUSDOptions>();
 
 		// Prompt with an options dialog if we can
-		if ( Options && ( !ExportTask || !ExportTask->bAutomated ) )
+		if (Options && (!ExportTask || !ExportTask->bAutomated))
 		{
-			Options->MeshAssetOptions.MaterialBakingOptions.TexturesDir.Path = FPaths::Combine( FPaths::GetPath( UExporter::CurrentFilename ), TEXT( "Textures" ) );
+			Options->MeshAssetOptions.MaterialBakingOptions.TexturesDir.Path = FPaths::Combine(
+				FPaths::GetPath(UExporter::CurrentFilename),
+				TEXT("Textures")
+			);
 
-			const bool bContinue = SUsdOptionsWindow::ShowExportOptions( *Options );
-			if ( !bContinue )
+			const bool bContinue = SUsdOptionsWindow::ShowExportOptions(*Options);
+			if (!bContinue)
 			{
 				return false;
 			}
 		}
 	}
-	if ( !Options )
+	if (!Options)
 	{
 		return false;
 	}
@@ -125,55 +135,57 @@ bool USkeletalMeshExporterUsd::ExportBinary( UObject* Object, const TCHAR* Type,
 	// "C:/MyFolder/file_payload.usda" and create an "asset" file "C:/MyFolder/file.usda" that uses it
 	// as a payload, pointing at the default prim
 	FString PayloadFilename = UExporter::CurrentFilename;
-	if ( Options && Options->MeshAssetOptions.bUsePayload )
+	if (Options && Options->MeshAssetOptions.bUsePayload)
 	{
 		FString PathPart;
 		FString FilenamePart;
 		FString ExtensionPart;
-		FPaths::Split( PayloadFilename, PathPart, FilenamePart, ExtensionPart );
+		FPaths::Split(PayloadFilename, PathPart, FilenamePart, ExtensionPart);
 
-		if ( FormatExtension.Contains( Options->MeshAssetOptions.PayloadFormat ) )
+		if (FormatExtension.Contains(Options->MeshAssetOptions.PayloadFormat))
 		{
 			ExtensionPart = Options->MeshAssetOptions.PayloadFormat;
 		}
 
-		PayloadFilename = FPaths::Combine( PathPart, FilenamePart + TEXT( "_payload." ) + ExtensionPart );
+		PayloadFilename = FPaths::Combine(PathPart, FilenamePart + TEXT("_payload.") + ExtensionPart);
 	}
 
-	if ( !IUsdExporterModule::CanExportToLayer( UExporter::CurrentFilename ) ||
-		( Options->MeshAssetOptions.bUsePayload && !IUsdExporterModule::CanExportToLayer( PayloadFilename ) ) )
+	if (!IUsdExporterModule::CanExportToLayer(UExporter::CurrentFilename)
+		|| (Options->MeshAssetOptions.bUsePayload && !IUsdExporterModule::CanExportToLayer(PayloadFilename)))
 	{
 		return false;
 	}
 
 	// Get a simple GUID hash/identifier of our mesh
 	FString DDCKeyHash;
-	if( FSkeletalMeshRenderData* RenderData = SkeletalMesh->GetResourceForRendering() )
+	if (FSkeletalMeshRenderData* RenderData = SkeletalMesh->GetResourceForRendering())
 	{
 		const FString& DDCKey = RenderData->DerivedDataKey;
 		FSHA1 SHA1;
-		SHA1.UpdateWithString( *DDCKey, DDCKey.Len() );
-		UsdUtils::HashForSkeletalMeshExport( *Options, SHA1 );
+		SHA1.UpdateWithString(*DDCKey, DDCKey.Len());
+		UsdUtils::HashForSkeletalMeshExport(*Options, SHA1);
 		SHA1.Final();
 		FSHAHash Hash;
-		SHA1.GetHash( &Hash.Hash[ 0 ] );
+		SHA1.GetHash(&Hash.Hash[0]);
 		DDCKeyHash = Hash.ToString();
 	}
 
 	// Check if we already have exported what we plan on exporting anyway
-	if ( FPaths::FileExists( UExporter::CurrentFilename ) && FPaths::FileExists( PayloadFilename ) )
+	if (FPaths::FileExists(UExporter::CurrentFilename) && FPaths::FileExists(PayloadFilename))
 	{
-		if ( !ExportTask->bReplaceIdentical )
+		if (!ExportTask->bReplaceIdentical)
 		{
-			UE_LOG( LogUsd, Log,
-				TEXT( "Skipping export of asset '%s' as the target file '%s' already exists." ),
+			UE_LOG(
+				LogUsd,
+				Log,
+				TEXT("Skipping export of asset '%s' as the target file '%s' already exists."),
 				*Object->GetPathName(),
 				*UExporter::CurrentFilename
 			);
 			return false;
 		}
 		// If we don't want to re-export this asset we need to check if its the same version
-		else if ( !Options->bReExportIdenticalAssets )
+		else if (!Options->bReExportIdenticalAssets)
 		{
 			bool bSkipMeshExport = false;
 
@@ -181,21 +193,22 @@ bool USkeletalMeshExporterUsd::ExportBinary( UObject* Object, const TCHAR* Type,
 			// we have to overwrite its files due to e.g. missing payload or anything like that
 			const bool bUseStageCache = false;
 			const EUsdInitialLoadSet InitialLoadSet = EUsdInitialLoadSet::LoadNone;
-			if ( UE::FUsdStage TempStage = UnrealUSDWrapper::OpenStage( *UExporter::CurrentFilename, InitialLoadSet, bUseStageCache ) )
+			if (UE::FUsdStage TempStage = UnrealUSDWrapper::OpenStage(*UExporter::CurrentFilename, InitialLoadSet, bUseStageCache))
 			{
-				if ( UE::FUsdPrim DefaultPrim = TempStage.GetDefaultPrim() )
+				if (UE::FUsdPrim DefaultPrim = TempStage.GetDefaultPrim())
 				{
-					FUsdUnrealAssetInfo Info = UsdUtils::GetPrimAssetInfo( DefaultPrim );
+					FUsdUnrealAssetInfo Info = UsdUtils::GetPrimAssetInfo(DefaultPrim);
 
 					const bool bVersionMatches = !Info.Version.IsEmpty() && Info.Version == DDCKeyHash;
 
-					const bool bAssetTypeMatches = !Info.UnrealAssetType.IsEmpty()
-						&& Info.UnrealAssetType == SkeletalMesh->GetClass()->GetName();
+					const bool bAssetTypeMatches = !Info.UnrealAssetType.IsEmpty() && Info.UnrealAssetType == SkeletalMesh->GetClass()->GetName();
 
-					if ( bVersionMatches && bAssetTypeMatches )
+					if (bVersionMatches && bAssetTypeMatches)
 					{
-						UE_LOG( LogUsd, Log,
-							TEXT( "Skipping export of asset '%s' as the target file '%s' already contains up-to-date exported data." ),
+						UE_LOG(
+							LogUsd,
+							Log,
+							TEXT("Skipping export of asset '%s' as the target file '%s' already contains up-to-date exported data."),
 							*SkeletalMesh->GetPathName(),
 							*UExporter::CurrentFilename
 						);
@@ -205,15 +218,15 @@ bool USkeletalMeshExporterUsd::ExportBinary( UObject* Object, const TCHAR* Type,
 				}
 			}
 
-			if ( bSkipMeshExport )
+			if (bSkipMeshExport)
 			{
 				// Even if we're not going to export the mesh, we may still need to re-bake materials
-				if ( Options->MeshAssetOptions.bBakeMaterials )
+				if (Options->MeshAssetOptions.bBakeMaterials)
 				{
 					TSet<UMaterialInterface*> MaterialsToBake;
-					for ( const FSkeletalMaterial& SkeletalMaterial : SkeletalMesh->GetMaterials() )
+					for (const FSkeletalMaterial& SkeletalMaterial : SkeletalMesh->GetMaterials())
 					{
-						MaterialsToBake.Add( SkeletalMaterial.MaterialInterface );
+						MaterialsToBake.Add(SkeletalMaterial.MaterialInterface);
 					}
 
 					const bool bIsAssetLayer = true;
@@ -238,45 +251,45 @@ bool USkeletalMeshExporterUsd::ExportBinary( UObject* Object, const TCHAR* Type,
 	double StartTime = FPlatformTime::Cycles64();
 
 	// UsdStage is the payload stage when exporting with payloads, or just the single stage otherwise
-	UE::FUsdStage UsdStage = UnrealUSDWrapper::NewStage( *PayloadFilename );
-	if ( !UsdStage )
+	UE::FUsdStage UsdStage = UnrealUSDWrapper::NewStage(*PayloadFilename);
+	if (!UsdStage)
 	{
 		return false;
 	}
 
-	if ( Options )
+	if (Options)
 	{
-		UsdUtils::SetUsdStageMetersPerUnit( UsdStage, Options->StageOptions.MetersPerUnit );
-		UsdUtils::SetUsdStageUpAxis( UsdStage, Options->StageOptions.UpAxis );
+		UsdUtils::SetUsdStageMetersPerUnit(UsdStage, Options->StageOptions.MetersPerUnit);
+		UsdUtils::SetUsdStageUpAxis(UsdStage, Options->StageOptions.UpAxis);
 	}
 
-	FString RootPrimPath = ( TEXT( "/" ) + UsdUtils::SanitizeUsdIdentifier( *SkeletalMesh->GetName() ) );
+	FString RootPrimPath = (TEXT("/") + UsdUtils::SanitizeUsdIdentifier(*SkeletalMesh->GetName()));
 
-	UE::FUsdPrim RootPrim = UsdStage.DefinePrim( UE::FSdfPath( *RootPrimPath ), TEXT("SkelRoot") );
-	if ( !RootPrim )
+	UE::FUsdPrim RootPrim = UsdStage.DefinePrim(UE::FSdfPath(*RootPrimPath), TEXT("SkelRoot"));
+	if (!RootPrim)
 	{
 		return false;
 	}
 
-	UsdStage.SetDefaultPrim( RootPrim );
+	UsdStage.SetDefaultPrim(RootPrim);
 
 	// Asset stage always the stage where we write the material assignments
 	UE::FUsdStage AssetStage;
 
 	// Using payload: Convert mesh data through the asset stage (that references the payload) so that we can
 	// author mesh data on the payload layer and material data on the asset layer
-	if ( Options && Options->MeshAssetOptions.bUsePayload )
+	if (Options && Options->MeshAssetOptions.bUsePayload)
 	{
-		AssetStage = UnrealUSDWrapper::NewStage( *UExporter::CurrentFilename );
-		if ( AssetStage )
+		AssetStage = UnrealUSDWrapper::NewStage(*UExporter::CurrentFilename);
+		if (AssetStage)
 		{
-			UsdUtils::SetUsdStageMetersPerUnit( AssetStage, Options->StageOptions.MetersPerUnit );
-			UsdUtils::SetUsdStageUpAxis( AssetStage, Options->StageOptions.UpAxis );
+			UsdUtils::SetUsdStageMetersPerUnit(AssetStage, Options->StageOptions.MetersPerUnit);
+			UsdUtils::SetUsdStageUpAxis(AssetStage, Options->StageOptions.UpAxis);
 
-			if ( UE::FUsdPrim AssetRootPrim = AssetStage.DefinePrim( UE::FSdfPath( *RootPrimPath ) ) )
+			if (UE::FUsdPrim AssetRootPrim = AssetStage.DefinePrim(UE::FSdfPath(*RootPrimPath)))
 			{
-				AssetStage.SetDefaultPrim( AssetRootPrim );
-				UsdUtils::AddPayload( AssetRootPrim, *PayloadFilename );
+				AssetStage.SetDefaultPrim(AssetRootPrim);
+				UsdUtils::AddPayload(AssetRootPrim, *PayloadFilename);
 			}
 		}
 	}
@@ -286,7 +299,14 @@ bool USkeletalMeshExporterUsd::ExportBinary( UObject* Object, const TCHAR* Type,
 		AssetStage = UsdStage;
 	}
 
-	UnrealToUsd::ConvertSkeletalMesh( SkeletalMesh, RootPrim, UsdUtils::GetDefaultTimeCode(), &AssetStage, Options->MeshAssetOptions.LowestMeshLOD, Options->MeshAssetOptions.HighestMeshLOD );
+	UnrealToUsd::ConvertSkeletalMesh(
+		SkeletalMesh,
+		RootPrim,
+		UsdUtils::GetDefaultTimeCode(),
+		&AssetStage,
+		Options->MeshAssetOptions.LowestMeshLOD,
+		Options->MeshAssetOptions.HighestMeshLOD
+	);
 
 	if (UE::FUsdPrim AssetDefaultPrim = AssetStage.GetDefaultPrim())
 	{
@@ -338,12 +358,12 @@ bool USkeletalMeshExporterUsd::ExportBinary( UObject* Object, const TCHAR* Type,
 	}
 
 	// Bake materials and replace unrealMaterials with references to the baked files.
-	if ( Options->MeshAssetOptions.bBakeMaterials )
+	if (Options->MeshAssetOptions.bBakeMaterials)
 	{
 		TSet<UMaterialInterface*> MaterialsToBake;
-		for ( const FSkeletalMaterial& SkeletalMaterial : SkeletalMesh->GetMaterials() )
+		for (const FSkeletalMaterial& SkeletalMaterial : SkeletalMesh->GetMaterials())
 		{
-			MaterialsToBake.Add( SkeletalMaterial.MaterialInterface );
+			MaterialsToBake.Add(SkeletalMaterial.MaterialInterface);
 		}
 
 		const bool bIsAssetLayer = true;
@@ -360,7 +380,7 @@ bool USkeletalMeshExporterUsd::ExportBinary( UObject* Object, const TCHAR* Type,
 		);
 	}
 
-	if ( AssetStage && UsdStage != AssetStage )
+	if (AssetStage && UsdStage != AssetStage)
 	{
 		AssetStage.GetRootLayer().Save();
 	}
@@ -369,16 +389,15 @@ bool USkeletalMeshExporterUsd::ExportBinary( UObject* Object, const TCHAR* Type,
 	// Analytics
 	{
 		bool bAutomated = ExportTask ? ExportTask->bAutomated : false;
-		double ElapsedSeconds = FPlatformTime::ToSeconds64( FPlatformTime::Cycles64() - StartTime );
-		FString Extension = FPaths::GetExtension( UExporter::CurrentFilename );
-		double NumberOfFrames = 1 + UsdStage.GetEndTimeCode() - UsdStage.GetStartTimeCode();
+		double ElapsedSeconds = FPlatformTime::ToSeconds64(FPlatformTime::Cycles64() - StartTime);
+		FString Extension = FPaths::GetExtension(UExporter::CurrentFilename);
 
 		UE::SkeletalMeshExporterUSD::Private::SendAnalytics(
 			Object,
 			Options,
 			bAutomated,
 			ElapsedSeconds,
-			UsdUtils::GetUsdStageNumFrames( AssetStage ),
+			UsdUtils::GetUsdStageNumFrames(AssetStage),
 			Extension
 		);
 	}
@@ -386,5 +405,5 @@ bool USkeletalMeshExporterUsd::ExportBinary( UObject* Object, const TCHAR* Type,
 	return true;
 #else
 	return false;
-#endif // #if USE_USD_SDK
+#endif	  // #if USE_USD_SDK
 }
diff --git a/Engine/Plugins/Importers/USDImporter/Source/USDExporter/Private/USDConversionBlueprintContext.cpp b/Engine/Plugins/Importers/USDImporter/Source/USDExporter/Private/USDConversionBlueprintContext.cpp
--- a/Engine/Plugins/Importers/USDImporter/Source/USDExporter/Private/USDConversionBlueprintContext.cpp
+++ b/Engine/Plugins/Importers/USDImporter/Source/USDExporter/Private/USDConversionBlueprintContext.cpp
@@ -25,51 +25,49 @@
 #include "RHI.h"
 #include "StaticMeshAttributes.h"
 
-namespace UE
+namespace UE::USDExporter::Private
 {
-	namespace USDExporter
+	UE::FUsdPrim GetPrim(const UE::FUsdStage& Stage, const FString& PrimPath)
 	{
-		namespace Private
+		if (!Stage)
 		{
-			UE::FUsdPrim GetPrim( const UE::FUsdStage& Stage, const FString& PrimPath )
-			{
-				if ( !Stage )
-				{
-					UE_LOG( LogUsd, Error, TEXT( "Export context has no stage set! Call SetStage with a root layer filepath first." ) );
-					return {};
-				}
+			UE_LOG(LogUsd, Error, TEXT("Export context has no stage set! Call SetStage with a root layer filepath first."));
+			return {};
+		}
 
-				return Stage.GetPrimAtPath( UE::FSdfPath( *PrimPath ) );
+		return Stage.GetPrimAtPath(UE::FSdfPath(*PrimPath));
+	}
+
+	static void ConvertPropertyEntriesIntoFlattenMaterial(
+		const TArray<FPropertyEntry>& PropertiesToBake,
+		const FIntPoint& DefaultTextureSize,
+		FFlattenMaterial& LandscapeFlattenMaterial
+	)
+	{
+#if USE_USD_SDK
+		for (FPropertyEntry Property : PropertiesToBake)
+		{
+			EFlattenMaterialProperties AnalogueProperty = UsdUtils::MaterialPropertyToFlattenProperty(Property.Property);
+			if (AnalogueProperty == EFlattenMaterialProperties::NumFlattenMaterialProperties)
+			{
+				continue;
 			}
 
-			static void ConvertPropertyEntriesIntoFlattenMaterial( const TArray<FPropertyEntry>& PropertiesToBake, const FIntPoint& DefaultTextureSize, FFlattenMaterial& LandscapeFlattenMaterial )
+			FIntPoint Size = DefaultTextureSize;
+			if (Property.bUseCustomSize)
 			{
-#if USE_USD_SDK
-				for ( FPropertyEntry Property : PropertiesToBake )
-				{
-					EFlattenMaterialProperties AnalogueProperty = UsdUtils::MaterialPropertyToFlattenProperty( Property.Property );
-					if ( AnalogueProperty == EFlattenMaterialProperties::NumFlattenMaterialProperties )
-					{
-						continue;
-					}
-
-					FIntPoint Size = DefaultTextureSize;
-					if ( Property.bUseCustomSize )
-					{
-						Size = Property.CustomSize;
-					}
-					if ( Property.bUseConstantValue )
-					{
-						Size = FIntPoint( 0, 0 );
-					}
-
-					LandscapeFlattenMaterial.SetPropertySize( AnalogueProperty, Size );
-				}
-#endif
+				Size = Property.CustomSize;
 			}
+			if (Property.bUseConstantValue)
+			{
+				Size = FIntPoint(0, 0);
+			}
+
+			LandscapeFlattenMaterial.SetPropertySize(AnalogueProperty, Size);
 		}
+#endif
 	}
-}
+}	 // namespace UE::USDExporter::Private
 namespace UnrealToUsdImpl = UE::USDExporter::Private;
 
 UUsdConversionBlueprintContext::~UUsdConversionBlueprintContext()
@@ -77,158 +75,162 @@ UUsdConversionBlueprintContext::~UUsdConversionBlueprintContext()
 	Cleanup();
 }
 
-void UUsdConversionBlueprintContext::SetStageRootLayer( FFilePath StageRootLayerPath )
+void UUsdConversionBlueprintContext::SetStageRootLayer(FFilePath StageRootLayerPath)
 {
 	Cleanup();
 
-	TArray< UE::FUsdStage > PreviouslyOpenedStages = UnrealUSDWrapper::GetAllStagesFromCache();
+	TArray<UE::FUsdStage> PreviouslyOpenedStages = UnrealUSDWrapper::GetAllStagesFromCache();
 
-	Stage = UnrealUSDWrapper::OpenStage( *StageRootLayerPath.FilePath, EUsdInitialLoadSet::LoadAll );
-	bEraseFromStageCache = !PreviouslyOpenedStages.Contains( Stage );
+	Stage = UnrealUSDWrapper::OpenStage(*StageRootLayerPath.FilePath, EUsdInitialLoadSet::LoadAll);
+	bEraseFromStageCache = !PreviouslyOpenedStages.Contains(Stage);
 }
 
 FFilePath UUsdConversionBlueprintContext::GetStageRootLayer()
 {
-	if ( Stage )
+	if (Stage)
 	{
-		return FFilePath{ Stage.GetRootLayer().GetRealPath() };
+		return FFilePath{Stage.GetRootLayer().GetRealPath()};
 	}
 
-	UE_LOG( LogUsd, Error, TEXT( "There is no stage currently open!" ) );
+	UE_LOG(LogUsd, Error, TEXT("There is no stage currently open!"));
 	return {};
 }
 
-void UUsdConversionBlueprintContext::SetEditTarget( FFilePath EditTargetLayerPath )
+void UUsdConversionBlueprintContext::SetEditTarget(FFilePath EditTargetLayerPath)
 {
-	if ( Stage )
+	if (Stage)
 	{
-		if ( UE::FSdfLayer EditTargetLayer = UE::FSdfLayer::FindOrOpen( *EditTargetLayerPath.FilePath ) )
+		if (UE::FSdfLayer EditTargetLayer = UE::FSdfLayer::FindOrOpen(*EditTargetLayerPath.FilePath))
 		{
-			Stage.SetEditTarget( EditTargetLayer );
+			Stage.SetEditTarget(EditTargetLayer);
 		}
 		else
 		{
-			UE_LOG( LogUsd, Error, TEXT( "Failed to find or open USD layer with filepath '%s'!" ), *EditTargetLayerPath.FilePath );
+			UE_LOG(LogUsd, Error, TEXT("Failed to find or open USD layer with filepath '%s'!"), *EditTargetLayerPath.FilePath);
 		}
 	}
 	else
 	{
-		UE_LOG( LogUsd, Error, TEXT( "There is no stage currently open!" ) );
+		UE_LOG(LogUsd, Error, TEXT("There is no stage currently open!"));
 	}
 }
 
 FFilePath UUsdConversionBlueprintContext::GetEditTarget()
 {
-	if ( Stage )
+	if (Stage)
 	{
-		return FFilePath{ Stage.GetEditTarget().GetRealPath() };
+		return FFilePath{Stage.GetEditTarget().GetRealPath()};
 	}
 
-	UE_LOG( LogUsd, Error, TEXT( "There is no stage currently open!" ) );
+	UE_LOG(LogUsd, Error, TEXT("There is no stage currently open!"));
 	return {};
 }
 
 void UUsdConversionBlueprintContext::Cleanup()
 {
-	if ( Stage )
+	if (Stage)
 	{
-		if ( bEraseFromStageCache )
+		if (bEraseFromStageCache)
 		{
-			UnrealUSDWrapper::EraseStageFromCache( Stage );
+			UnrealUSDWrapper::EraseStageFromCache(Stage);
 		}
 
 		Stage = UE::FUsdStage();
 	}
 }
 
-bool UUsdConversionBlueprintContext::ConvertLightComponent( const ULightComponentBase* Component, const FString& PrimPath, float TimeCode )
+bool UUsdConversionBlueprintContext::ConvertLightComponent(const ULightComponentBase* Component, const FString& PrimPath, float TimeCode)
 {
 #if USE_USD_SDK
-	UE::FUsdPrim Prim = UnrealToUsdImpl::GetPrim( Stage, PrimPath );
-	if ( !Prim || !Component )
+	UE::FUsdPrim Prim = UnrealToUsdImpl::GetPrim(Stage, PrimPath);
+	if (!Prim || !Component)
 	{
 		return false;
 	}
 
-	return UnrealToUsd::ConvertLightComponent( *Component, Prim, TimeCode == FLT_MAX ? UsdUtils::GetDefaultTimeCode() : TimeCode );
+	return UnrealToUsd::ConvertLightComponent(*Component, Prim, TimeCode == FLT_MAX ? UsdUtils::GetDefaultTimeCode() : TimeCode);
 #else
 	return false;
-#endif // USE_USD_SDK
+#endif	  // USE_USD_SDK
 }
 
-bool UUsdConversionBlueprintContext::ConvertDirectionalLightComponent( const UDirectionalLightComponent* Component, const FString& PrimPath, float TimeCode )
+bool UUsdConversionBlueprintContext::ConvertDirectionalLightComponent(
+	const UDirectionalLightComponent* Component,
+	const FString& PrimPath,
+	float TimeCode
+)
 {
 #if USE_USD_SDK
-	UE::FUsdPrim Prim = UnrealToUsdImpl::GetPrim( Stage, PrimPath );
-	if ( !Prim || !Component )
+	UE::FUsdPrim Prim = UnrealToUsdImpl::GetPrim(Stage, PrimPath);
+	if (!Prim || !Component)
 	{
 		return false;
 	}
 
-	return UnrealToUsd::ConvertDirectionalLightComponent( *Component, Prim, TimeCode == FLT_MAX ? UsdUtils::GetDefaultTimeCode() : TimeCode );
+	return UnrealToUsd::ConvertDirectionalLightComponent(*Component, Prim, TimeCode == FLT_MAX ? UsdUtils::GetDefaultTimeCode() : TimeCode);
 #else
 	return false;
-#endif // USE_USD_SDK
+#endif	  // USE_USD_SDK
 }
 
-bool UUsdConversionBlueprintContext::ConvertRectLightComponent( const URectLightComponent* Component, const FString& PrimPath, float TimeCode )
+bool UUsdConversionBlueprintContext::ConvertRectLightComponent(const URectLightComponent* Component, const FString& PrimPath, float TimeCode)
 {
 #if USE_USD_SDK
-	UE::FUsdPrim Prim = UnrealToUsdImpl::GetPrim( Stage, PrimPath );
-	if ( !Prim || !Component )
+	UE::FUsdPrim Prim = UnrealToUsdImpl::GetPrim(Stage, PrimPath);
+	if (!Prim || !Component)
 	{
 		return false;
 	}
 
-	return UnrealToUsd::ConvertRectLightComponent( *Component, Prim, TimeCode == FLT_MAX ? UsdUtils::GetDefaultTimeCode() : TimeCode );
+	return UnrealToUsd::ConvertRectLightComponent(*Component, Prim, TimeCode == FLT_MAX ? UsdUtils::GetDefaultTimeCode() : TimeCode);
 #else
 	return false;
-#endif // USE_USD_SDK
+#endif	  // USE_USD_SDK
 }
 
-bool UUsdConversionBlueprintContext::ConvertPointLightComponent( const UPointLightComponent* Component, const FString& PrimPath, float TimeCode )
+bool UUsdConversionBlueprintContext::ConvertPointLightComponent(const UPointLightComponent* Component, const FString& PrimPath, float TimeCode)
 {
 #if USE_USD_SDK
-	UE::FUsdPrim Prim = UnrealToUsdImpl::GetPrim( Stage, PrimPath );
-	if ( !Prim || !Component )
+	UE::FUsdPrim Prim = UnrealToUsdImpl::GetPrim(Stage, PrimPath);
+	if (!Prim || !Component)
 	{
 		return false;
 	}
 
-	return UnrealToUsd::ConvertPointLightComponent( *Component, Prim, TimeCode == FLT_MAX ? UsdUtils::GetDefaultTimeCode() : TimeCode );
+	return UnrealToUsd::ConvertPointLightComponent(*Component, Prim, TimeCode == FLT_MAX ? UsdUtils::GetDefaultTimeCode() : TimeCode);
 #else
 	return false;
-#endif // USE_USD_SDK
+#endif	  // USE_USD_SDK
 }
 
-bool UUsdConversionBlueprintContext::ConvertSkyLightComponent( const USkyLightComponent* Component, const FString& PrimPath, float TimeCode )
+bool UUsdConversionBlueprintContext::ConvertSkyLightComponent(const USkyLightComponent* Component, const FString& PrimPath, float TimeCode)
 {
 #if USE_USD_SDK
-	UE::FUsdPrim Prim = UnrealToUsdImpl::GetPrim( Stage, PrimPath );
-	if ( !Prim || !Component )
+	UE::FUsdPrim Prim = UnrealToUsdImpl::GetPrim(Stage, PrimPath);
+	if (!Prim || !Component)
 	{
 		return false;
 	}
 
-	return UnrealToUsd::ConvertSkyLightComponent( *Component, Prim, TimeCode == FLT_MAX ? UsdUtils::GetDefaultTimeCode() : TimeCode );
+	return UnrealToUsd::ConvertSkyLightComponent(*Component, Prim, TimeCode == FLT_MAX ? UsdUtils::GetDefaultTimeCode() : TimeCode);
 #else
 	return false;
-#endif // USE_USD_SDK
+#endif	  // USE_USD_SDK
 }
 
-bool UUsdConversionBlueprintContext::ConvertSpotLightComponent( const USpotLightComponent* Component, const FString& PrimPath, float TimeCode )
+bool UUsdConversionBlueprintContext::ConvertSpotLightComponent(const USpotLightComponent* Component, const FString& PrimPath, float TimeCode)
 {
 #if USE_USD_SDK
-	UE::FUsdPrim Prim = UnrealToUsdImpl::GetPrim( Stage, PrimPath );
-	if ( !Prim || !Component )
+	UE::FUsdPrim Prim = UnrealToUsdImpl::GetPrim(Stage, PrimPath);
+	if (!Prim || !Component)
 	{
 		return false;
 	}
 
-	return UnrealToUsd::ConvertSpotLightComponent( *Component, Prim, TimeCode == FLT_MAX ? UsdUtils::GetDefaultTimeCode() : TimeCode );
+	return UnrealToUsd::ConvertSpotLightComponent(*Component, Prim, TimeCode == FLT_MAX ? UsdUtils::GetDefaultTimeCode() : TimeCode);
 #else
 	return false;
-#endif // USE_USD_SDK
+#endif	  // USE_USD_SDK
 }
 
 bool UUsdConversionBlueprintContext::ConvertBoundsComponent(const UUsdDrawModeComponent* Component, const FString& PrimPath, float TimeCode)
@@ -247,117 +249,136 @@ bool UUsdConversionBlueprintContext::ConvertBoundsComponent(const UUsdDrawModeCo
 #endif	  // USE_USD_SDK
 }
 
-bool UUsdConversionBlueprintContext::ConvertSceneComponent( const USceneComponent* Component, const FString& PrimPath )
+bool UUsdConversionBlueprintContext::ConvertSceneComponent(const USceneComponent* Component, const FString& PrimPath)
 {
 #if USE_USD_SDK
-	UE::FUsdPrim Prim = UnrealToUsdImpl::GetPrim( Stage, PrimPath );
-	if ( !Prim || !Component )
+	UE::FUsdPrim Prim = UnrealToUsdImpl::GetPrim(Stage, PrimPath);
+	if (!Prim || !Component)
 	{
 		return false;
 	}
 
-	return UnrealToUsd::ConvertSceneComponent( Stage, Component, Prim );
+	return UnrealToUsd::ConvertSceneComponent(Stage, Component, Prim);
 #else
 	return false;
-#endif // USE_USD_SDK
+#endif	  // USE_USD_SDK
 }
 
-bool UUsdConversionBlueprintContext::ConvertHismComponent( const UHierarchicalInstancedStaticMeshComponent* Component, const FString& PrimPath, float TimeCode )
+bool UUsdConversionBlueprintContext::ConvertHismComponent(
+	const UHierarchicalInstancedStaticMeshComponent* Component,
+	const FString& PrimPath,
+	float TimeCode
+)
 {
 #if USE_USD_SDK
-	UE::FUsdPrim Prim = UnrealToUsdImpl::GetPrim( Stage, PrimPath );
-	if ( !Prim || !Component )
+	UE::FUsdPrim Prim = UnrealToUsdImpl::GetPrim(Stage, PrimPath);
+	if (!Prim || !Component)
 	{
 		return false;
 	}
 
-	return UnrealToUsd::ConvertHierarchicalInstancedStaticMeshComponent( Component, Prim, TimeCode == FLT_MAX ? UsdUtils::GetDefaultTimeCode() : TimeCode );
+	return UnrealToUsd::ConvertHierarchicalInstancedStaticMeshComponent(
+		Component,
+		Prim,
+		TimeCode == FLT_MAX ? UsdUtils::GetDefaultTimeCode() : TimeCode
+	);
 #else
 	return false;
-#endif // USE_USD_SDK
+#endif	  // USE_USD_SDK
 }
 
-bool UUsdConversionBlueprintContext::ConvertMeshComponent( const UMeshComponent* Component, const FString& PrimPath )
+bool UUsdConversionBlueprintContext::ConvertMeshComponent(const UMeshComponent* Component, const FString& PrimPath)
 {
 #if USE_USD_SDK
-	UE::FUsdPrim Prim = UnrealToUsdImpl::GetPrim( Stage, PrimPath );
-	if ( !Prim || !Component )
+	UE::FUsdPrim Prim = UnrealToUsdImpl::GetPrim(Stage, PrimPath);
+	if (!Prim || !Component)
 	{
 		return false;
 	}
 
-	return UnrealToUsd::ConvertMeshComponent( Stage, Component, Prim );
+	return UnrealToUsd::ConvertMeshComponent(Stage, Component, Prim);
 #else
 	return false;
-#endif // USE_USD_SDK
+#endif	  // USE_USD_SDK
 }
 
-bool UUsdConversionBlueprintContext::ConvertCineCameraComponent( const UCineCameraComponent* Component, const FString& PrimPath, float TimeCode  )
+bool UUsdConversionBlueprintContext::ConvertCineCameraComponent(const UCineCameraComponent* Component, const FString& PrimPath, float TimeCode)
 {
 #if USE_USD_SDK
-	UE::FUsdPrim Prim = UnrealToUsdImpl::GetPrim( Stage, PrimPath );
-	if ( !Prim || !Component )
+	UE::FUsdPrim Prim = UnrealToUsdImpl::GetPrim(Stage, PrimPath);
+	if (!Prim || !Component)
 	{
 		return false;
 	}
 
-	return UnrealToUsd::ConvertCameraComponent( *Component, Prim, TimeCode == FLT_MAX ? UsdUtils::GetDefaultTimeCode() : TimeCode );
+	return UnrealToUsd::ConvertCameraComponent(*Component, Prim, TimeCode == FLT_MAX ? UsdUtils::GetDefaultTimeCode() : TimeCode);
 #else
 	return false;
-#endif // USE_USD_SDK
+#endif	  // USE_USD_SDK
 }
 
-bool UUsdConversionBlueprintContext::ConvertInstancedFoliageActor( const AInstancedFoliageActor* Actor, const FString& PrimPath, ULevel* InstancesLevel, float TimeCode )
+bool UUsdConversionBlueprintContext::ConvertInstancedFoliageActor(
+	const AInstancedFoliageActor* Actor,
+	const FString& PrimPath,
+	ULevel* InstancesLevel,
+	float TimeCode
+)
 {
 #if USE_USD_SDK
-	UE::FUsdPrim Prim = UnrealToUsdImpl::GetPrim( Stage, PrimPath );
-	if ( !Prim || !Actor )
+	UE::FUsdPrim Prim = UnrealToUsdImpl::GetPrim(Stage, PrimPath);
+	if (!Prim || !Actor)
 	{
 		return false;
 	}
 
-	return UnrealToUsd::ConvertInstancedFoliageActor( *Actor, Prim, TimeCode == FLT_MAX ? UsdUtils::GetDefaultTimeCode() : TimeCode, InstancesLevel );
+	return UnrealToUsd::ConvertInstancedFoliageActor(*Actor, Prim, TimeCode == FLT_MAX ? UsdUtils::GetDefaultTimeCode() : TimeCode, InstancesLevel);
 #else
 	return false;
-#endif // USE_USD_SDK
+#endif	  // USE_USD_SDK
 }
 
-bool UUsdConversionBlueprintContext::ConvertLandscapeProxyActorMesh( const ALandscapeProxy* Actor, const FString& PrimPath, int32 LowestLOD, int32 HighestLOD, float TimeCode /*= 3.402823466e+38F */ )
+bool UUsdConversionBlueprintContext::ConvertLandscapeProxyActorMesh(
+	const ALandscapeProxy* Actor,
+	const FString& PrimPath,
+	int32 LowestLOD,
+	int32 HighestLOD,
+	float TimeCode /*= 3.402823466e+38F */
+)
 {
 #if USE_USD_SDK
-	UE::FUsdPrim Prim = UnrealToUsdImpl::GetPrim( Stage, PrimPath );
-	if ( !Prim || !Actor )
+	UE::FUsdPrim Prim = UnrealToUsdImpl::GetPrim(Stage, PrimPath);
+	if (!Prim || !Actor)
 	{
 		return false;
 	}
 
 	// Make sure they're both >= 0 (the options dialog slider is clamped, but this may be called directly)
-	LowestLOD = FMath::Max( LowestLOD, 0 );
-	HighestLOD = FMath::Max( HighestLOD, 0 );
+	LowestLOD = FMath::Max(LowestLOD, 0);
+	HighestLOD = FMath::Max(HighestLOD, 0);
 
 	// Make sure Lowest <= Highest
-	int32 Temp = FMath::Min( LowestLOD, HighestLOD );
-	HighestLOD = FMath::Max( LowestLOD, HighestLOD );
+	int32 Temp = FMath::Min(LowestLOD, HighestLOD);
+	HighestLOD = FMath::Max(LowestLOD, HighestLOD);
 	LowestLOD = Temp;
 
 	// Make sure it's at least 1 LOD level
-	int32 NumLODs = FMath::Max( HighestLOD - LowestLOD + 1, 1 );
+	int32 NumLODs = FMath::Max(HighestLOD - LowestLOD + 1, 1);
 
 	TArray<FMeshDescription> LODMeshDescriptions;
-	LODMeshDescriptions.SetNum( NumLODs );
+	LODMeshDescriptions.SetNum(NumLODs);
 
-	for ( int32 LODIndex = 0; LODIndex < NumLODs; ++LODIndex )
+	for (int32 LODIndex = 0; LODIndex < NumLODs; ++LODIndex)
 	{
-		FMeshDescription& MeshDescription = LODMeshDescriptions[ LODIndex ];
+		FMeshDescription& MeshDescription = LODMeshDescriptions[LODIndex];
 
-		FStaticMeshAttributes Attributes( MeshDescription );
+		FStaticMeshAttributes Attributes(MeshDescription);
 		Attributes.Register();
 
 		ALandscapeProxy::FRawMeshExportParams ExportParams;
 		ExportParams.ExportLOD = LODIndex + LowestLOD;
-		if ( !Actor->ExportToRawMesh(ExportParams, MeshDescription ) )
+		if (!Actor->ExportToRawMesh(ExportParams, MeshDescription))
 		{
-			UE_LOG( LogUsd, Error, TEXT( "Failed to convert LOD %d of landscape actor '%s''s mesh data" ), LODIndex, *Actor->GetName() );
+			UE_LOG(LogUsd, Error, TEXT("Failed to convert LOD %d of landscape actor '%s''s mesh data"), LODIndex, *Actor->GetName());
 			return false;
 		}
 	}
@@ -368,7 +389,12 @@ bool UUsdConversionBlueprintContext::ConvertLandscapeProxyActorMesh( const ALand
 	// the parent actor, generating the same transform as the one you'd get from calling LandscapeActorToWorld() directly on it.
 	// We don't want this here: We want to specifically compensate the proxy actor's transform ourselves, so we need just ActorToWorld.
 	FMatrix ActorToWorldInv = Actor->ActorToWorld().ToInverseMatrixWithScale();
-	if ( !UnrealToUsd::ConvertMeshDescriptions( LODMeshDescriptions, Prim, ActorToWorldInv, TimeCode == FLT_MAX ? UsdUtils::GetDefaultTimeCode() : TimeCode ) )
+	if (!UnrealToUsd::ConvertMeshDescriptions(
+			LODMeshDescriptions,
+			Prim,
+			ActorToWorldInv,
+			TimeCode == FLT_MAX ? UsdUtils::GetDefaultTimeCode() : TimeCode
+		))
 	{
 		return false;
 	}
@@ -376,72 +402,85 @@ bool UUsdConversionBlueprintContext::ConvertLandscapeProxyActorMesh( const ALand
 	return true;
 #else
 	return false;
-#endif // USE_USD_SDK
+#endif	  // USE_USD_SDK
 }
 
-bool UUsdConversionBlueprintContext::ConvertLandscapeProxyActorMaterial( ALandscapeProxy* Actor, const FString& PrimPath, const TArray<FPropertyEntry>& PropertiesToBake, const FIntPoint& DefaultTextureSize, const FDirectoryPath& TexturesDir, float TimeCode /*= 3.402823466e+38F */ )
+bool UUsdConversionBlueprintContext::ConvertLandscapeProxyActorMaterial(
+	ALandscapeProxy* Actor,
+	const FString& PrimPath,
+	const TArray<FPropertyEntry>& PropertiesToBake,
+	const FIntPoint& DefaultTextureSize,
+	const FDirectoryPath& TexturesDir,
+	float TimeCode /*= 3.402823466e+38F */
+)
 {
 #if USE_USD_SDK
-	UE::FUsdPrim Prim = UnrealToUsdImpl::GetPrim( Stage, PrimPath );
-	if ( !Prim || !Actor )
+	UE::FUsdPrim Prim = UnrealToUsdImpl::GetPrim(Stage, PrimPath);
+	if (!Prim || !Actor)
 	{
 		return false;
 	}
 
 	FFlattenMaterial LandscapeFlattenMaterial;
-	UnrealToUsdImpl::ConvertPropertyEntriesIntoFlattenMaterial( PropertiesToBake, DefaultTextureSize, LandscapeFlattenMaterial );
+	UnrealToUsdImpl::ConvertPropertyEntriesIntoFlattenMaterial(PropertiesToBake, DefaultTextureSize, LandscapeFlattenMaterial);
 
 	// FMaterialUtilities::ExportLandscapeMaterial will basically render the scene with an orthographic transform. We need to hide
 	// each and every actor and component that is not part of the landscape for that render, or else they will show up on the bake
 	TSet<FPrimitiveComponentId> PrimitivesToHide;
-	for ( TObjectIterator<UPrimitiveComponent> It; It; ++It )
+	for (TObjectIterator<UPrimitiveComponent> It; It; ++It)
 	{
 		UPrimitiveComponent* PrimitiveComp = *It;
-		const bool bTargetPrim = ( PrimitiveComp->GetOuter() == Actor );
+		const bool bTargetPrim = (PrimitiveComp->GetOuter() == Actor);
 
-		if ( !bTargetPrim && PrimitiveComp->IsRegistered() && PrimitiveComp->SceneProxy )
+		if (!bTargetPrim && PrimitiveComp->IsRegistered() && PrimitiveComp->SceneProxy)
 		{
-			PrimitivesToHide.Add( PrimitiveComp->SceneProxy->GetPrimitiveComponentId() );
+			PrimitivesToHide.Add(PrimitiveComp->SceneProxy->GetPrimitiveComponentId());
 		}
 	}
-	FMaterialUtilities::ExportLandscapeMaterial( Actor, PrimitivesToHide, LandscapeFlattenMaterial );
+	FMaterialUtilities::ExportLandscapeMaterial(Actor, PrimitivesToHide, LandscapeFlattenMaterial);
 
 	// If we asked to bake a texture with some size and it's just a constant value, FMaterialUtilities::ExportLandscapeMaterial will
 	// emit the full texture anyway, so we need to collapse it back into being just one pixel
-	UsdUtils::CollapseConstantChannelsToSinglePixel( LandscapeFlattenMaterial );
+	UsdUtils::CollapseConstantChannelsToSinglePixel(LandscapeFlattenMaterial);
 
 	// Synchronize the intended bake size for channels that FMaterialUtilities::ExportLandscapeMaterial didn't bake (or that we collapsed),
 	// or else downstream code may assume that there was some form of error
-	for ( FPropertyEntry PropertyToBake : PropertiesToBake )
+	for (FPropertyEntry PropertyToBake : PropertiesToBake)
 	{
-		EFlattenMaterialProperties AnalogueProperty = UsdUtils::MaterialPropertyToFlattenProperty( PropertyToBake.Property );
-		if ( AnalogueProperty == EFlattenMaterialProperties::NumFlattenMaterialProperties )
+		EFlattenMaterialProperties AnalogueProperty = UsdUtils::MaterialPropertyToFlattenProperty(PropertyToBake.Property);
+		if (AnalogueProperty == EFlattenMaterialProperties::NumFlattenMaterialProperties)
 		{
 			continue;
 		}
 
 		const TArray<FColor>& Samples = LandscapeFlattenMaterial.GetPropertySamples(AnalogueProperty);
-		if ( Samples.Num() == 0 )
+		if (Samples.Num() == 0)
 		{
-			LandscapeFlattenMaterial.SetPropertySize( AnalogueProperty, FIntPoint( 0, 0 ) );
+			LandscapeFlattenMaterial.SetPropertySize(AnalogueProperty, FIntPoint(0, 0));
 		}
 	}
 
-	bool bSuccess = UnrealToUsd::ConvertFlattenMaterial( Actor->GetPathName() + TEXT( "_BakedMaterial" ), LandscapeFlattenMaterial, PropertiesToBake, TexturesDir, Prim );
+	bool bSuccess = UnrealToUsd::ConvertFlattenMaterial(
+		Actor->GetPathName() + TEXT("_BakedMaterial"),
+		LandscapeFlattenMaterial,
+		PropertiesToBake,
+		TexturesDir,
+		Prim
+	);
 
 	// FMaterialUtilities::ExportLandscapeMaterial always bakes WorldNormals, so we need to write to this material that it's meant to be using
 	// world-space normals. On import, we check for this and set the proper scalar parameter on our UsdPreviewSurface material to compensate for it.
-	// This is done on the shader via a scalar parameter to avoid quantizing the normal info twice (which converting on the CPU would do) and to allow usage
-	// during runtime (which wouldn't be possible with static switch parameters)
-	if ( bSuccess )
+	// This is done on the shader via a scalar parameter to avoid quantizing the normal info twice (which converting on the CPU would do) and to allow
+	// usage during runtime (which wouldn't be possible with static switch parameters)
+	if (bSuccess)
 	{
-		bSuccess &= UsdUtils::MarkMaterialPrimWithWorldSpaceNormals( Prim );
+		bSuccess &= UsdUtils::MarkMaterialPrimWithWorldSpaceNormals(Prim);
 	}
 
 	return bSuccess;
 #else
 	return false;
-#endif // USE_USD_SDK
+#endif	  // USE_USD_SDK
 }
 
 bool UUsdConversionBlueprintContext::ConvertMaterialOverrides(
@@ -465,95 +504,101 @@ bool UUsdConversionBlueprintContext::ConvertMaterialOverrides(
 #endif	  // USE_USD_SDK
 }
 
-void UUsdConversionBlueprintContext::ReplaceUnrealMaterialsWithBaked( const FFilePath& LayerToAuthorIn, const TMap<FString, FString>& BakedMaterials, bool bIsAssetLayer, bool bUsePayload, bool bRemoveUnrealMaterials )
+void UUsdConversionBlueprintContext::ReplaceUnrealMaterialsWithBaked(
+	const FFilePath& LayerToAuthorIn,
+	const TMap<FString, FString>& BakedMaterials,
+	bool bIsAssetLayer,
+	bool bUsePayload,
+	bool bRemoveUnrealMaterials
+)
 {
 #if USE_USD_SDK
-	if ( !Stage || LayerToAuthorIn.FilePath.IsEmpty() )
+	if (!Stage || LayerToAuthorIn.FilePath.IsEmpty())
 	{
 		return;
 	}
 
-	UE::FSdfLayer Layer = UE::FSdfLayer::FindOrOpen( *LayerToAuthorIn.FilePath );
-	if ( !Layer )
+	UE::FSdfLayer Layer = UE::FSdfLayer::FindOrOpen(*LayerToAuthorIn.FilePath);
+	if (!Layer)
 	{
 		return;
 	}
 
-	UsdUtils::ReplaceUnrealMaterialsWithBaked( Stage, Layer, BakedMaterials, bIsAssetLayer, bUsePayload );
-#endif // USE_USD_SDK
+	UsdUtils::ReplaceUnrealMaterialsWithBaked(Stage, Layer, BakedMaterials, bIsAssetLayer, bUsePayload);
+#endif	  // USE_USD_SDK
 }
 
-bool UUsdConversionBlueprintContext::RemoveUnrealSurfaceOutput( const FString& PrimPath, const FFilePath& LayerToAuthorIn )
+bool UUsdConversionBlueprintContext::RemoveUnrealSurfaceOutput(const FString& PrimPath, const FFilePath& LayerToAuthorIn)
 {
 #if USE_USD_SDK
-	if ( !Stage )
+	if (!Stage)
 	{
 		return false;
 	}
 
-	UE::FUsdPrim Prim = UnrealToUsdImpl::GetPrim( Stage, PrimPath );
-	if ( !Prim )
+	UE::FUsdPrim Prim = UnrealToUsdImpl::GetPrim(Stage, PrimPath);
+	if (!Prim)
 	{
 		return false;
 	}
 
 	UE::FSdfLayer Layer;
-	if ( !LayerToAuthorIn.FilePath.IsEmpty() )
+	if (!LayerToAuthorIn.FilePath.IsEmpty())
 	{
-		Layer = UE::FSdfLayer::FindOrOpen( *LayerToAuthorIn.FilePath );
+		Layer = UE::FSdfLayer::FindOrOpen(*LayerToAuthorIn.FilePath);
 	}
 
 	PRAGMA_DISABLE_DEPRECATION_WARNINGS
-	return UsdUtils::RemoveUnrealSurfaceOutput( Prim, Layer );
+	return UsdUtils::RemoveUnrealSurfaceOutput(Prim, Layer);
 	PRAGMA_ENABLE_DEPRECATION_WARNINGS
 #else
 	return false;
-#endif // USE_USD_SDK
+#endif	  // USE_USD_SDK
 }
 
 int32 UUsdConversionBlueprintContext::GetUsdStageNumFrames()
 {
 #if USE_USD_SDK
-	if ( !Stage )
+	if (!Stage)
 	{
 		return 0;
 	}
 
-	return UsdUtils::GetUsdStageNumFrames( Stage );
+	return UsdUtils::GetUsdStageNumFrames(Stage);
 #else
 	return 0;
-#endif // USE_USD_SDK
+#endif	  // USE_USD_SDK
 }
 
-void UUsdConversionBlueprintContext::SetPrimAssetInfo( const FString& PrimPath, const FUsdUnrealAssetInfo& Info )
+void UUsdConversionBlueprintContext::SetPrimAssetInfo(const FString& PrimPath, const FUsdUnrealAssetInfo& Info)
 {
 #if USE_USD_SDK
-	if ( !Stage )
+	if (!Stage)
 	{
 		return;
 	}
 
-	UE::FUsdPrim Prim = UnrealToUsdImpl::GetPrim( Stage, PrimPath );
-	if ( !Prim )
+	UE::FUsdPrim Prim = UnrealToUsdImpl::GetPrim(Stage, PrimPath);
+	if (!Prim)
 	{
 		return;
 	}
 
-	UsdUtils::SetPrimAssetInfo( Prim, Info );
-#endif // USE_USD_SDK
+	UsdUtils::SetPrimAssetInfo(Prim, Info);
+#endif	  // USE_USD_SDK
 }
 
-FUsdUnrealAssetInfo UUsdConversionBlueprintContext::GetPrimAssetInfo( const FString& PrimPath )
+FUsdUnrealAssetInfo UUsdConversionBlueprintContext::GetPrimAssetInfo(const FString& PrimPath)
 {
 #if USE_USD_SDK
-	if ( Stage )
+	if (Stage)
 	{
-		if ( UE::FUsdPrim Prim = UnrealToUsdImpl::GetPrim( Stage, PrimPath ) )
+		if (UE::FUsdPrim Prim = UnrealToUsdImpl::GetPrim(Stage, PrimPath))
 		{
-			return UsdUtils::GetPrimAssetInfo( Prim );
+			return UsdUtils::GetPrimAssetInfo(Prim);
 		}
 	}
-#endif // USE_USD_SDK
+#endif	  // USE_USD_SDK
 
 	return {};
 }
