USD: Reformat USDExporter module.

diff --git a/Engine/Plugins/Importers/USDImporter/Source/USDExporter/Private/USDConversionBlueprintLibrary.cpp b/Engine/Plugins/Importers/USDImporter/Source/USDExporter/Private/USDConversionBlueprintLibrary.cpp
--- a/Engine/Plugins/Importers/USDImporter/Source/USDExporter/Private/USDConversionBlueprintLibrary.cpp
+++ b/Engine/Plugins/Importers/USDImporter/Source/USDExporter/Private/USDConversionBlueprintLibrary.cpp
@@ -2,7 +2,6 @@
 
 #include "USDConversionBlueprintLibrary.h"
 
-#include "LevelExporterUSD.h"
 #include "LevelExporterUSDOptions.h"
 #include "UnrealUSDWrapper.h"
 #include "USDAssetUserData.h"
@@ -11,20 +10,17 @@
 #include "USDExporterModule.h"
 #include "USDLayerUtils.h"
 #include "USDLog.h"
-#include "USDTypesConversion.h"
 #include "USDValueConversion.h"
 
 #include "UsdWrappers/SdfLayer.h"
 #include "UsdWrappers/SdfPath.h"
 #include "UsdWrappers/UsdPrim.h"
 #include "UsdWrappers/UsdStage.h"
-#include "UsdWrappers/VtValue.h"
 
 #include "AnalyticsBlueprintLibrary.h"
 #include "AnalyticsEventAttribute.h"
 #include "AssetCompilingManager.h"
 #include "ContentStreaming.h"
-#include "CoreMinimal.h"
 #include "Editor.h"
 #include "Engine/Engine.h"
 #include "Engine/Level.h"
@@ -36,141 +32,135 @@
 #include "LevelEditorSequencerIntegration.h"
 #include "UObject/ObjectMacros.h"
 
-namespace UE
+namespace UE::UsdConversionBlueprintLibrary::Private
 {
-	namespace UsdConversionBlueprintLibraryImpl
+	void WaitForAllAsyncAndSteamingTasks(UWorld* World)
 	{
-		namespace Private
+		FlushAsyncLoading();
+
+		if (World)
 		{
-			void WaitForAllAsyncAndSteamingTasks(UWorld* World)
+			World->BlockTillLevelStreamingCompleted();
+
+			if (!FPlatformProperties::RequiresCookedData())
 			{
-				FlushAsyncLoading();
+				UMaterialInterface::SubmitRemainingJobsForWorld(World);
+				FAssetCompilingManager::Get().FinishAllCompilation();
+			}
+		}
 
-				if (World)
-				{
-					World->BlockTillLevelStreamingCompleted();
+		UTexture::ForceUpdateTextureStreaming();
 
-					if (!FPlatformProperties::RequiresCookedData())
-					{
-						UMaterialInterface::SubmitRemainingJobsForWorld(World);
-						FAssetCompilingManager::Get().FinishAllCompilation();
-					}
-				}
+		IStreamingManager::Get().StreamAllResources(0.0f);
+	}
 
-				UTexture::ForceUpdateTextureStreaming();
+	void StreamInLevels(ULevel* Level, const TSet<FString>& LevelsToIgnore)
+	{
+		UWorld* InnerWorld = Level->GetTypedOuter<UWorld>();
+		if (!InnerWorld || InnerWorld->GetStreamingLevels().Num() == 0)
+		{
+			return;
+		}
 
-				IStreamingManager::Get().StreamAllResources(0.0f);
-			}
+		// Ensure our world to export has a context so that level streaming doesn't crash.
+		// This is needed exclusively to be able to load other levels from python scripts using just `load_asset`
+		// and have them be exportable.
+		bool bCreatedContext = false;
+		FWorldContext* Context = GEngine->GetWorldContextFromWorld(InnerWorld);
+		if (!Context)
+		{
+			FWorldContext& NewContext = GEngine->CreateNewWorldContext(EWorldType::EditorPreview);
+			NewContext.SetCurrentWorld(InnerWorld);
 
-			void StreamInLevels( ULevel* Level, const TSet<FString>& LevelsToIgnore )
-			{
-				UWorld* InnerWorld = Level->GetTypedOuter<UWorld>();
-				if ( !InnerWorld || InnerWorld->GetStreamingLevels().Num() == 0 )
-				{
-					return;
-				}
+			bCreatedContext = true;
+		}
 
-				// Ensure our world to export has a context so that level streaming doesn't crash.
-				// This is needed exclusively to be able to load other levels from python scripts using just `load_asset`
-				// and have them be exportable.
-				bool bCreatedContext = false;
-				FWorldContext* Context = GEngine->GetWorldContextFromWorld( InnerWorld );
-				if ( !Context )
+		// Mark all sublevels that we need to be loaded
+		for (ULevelStreaming* StreamingLevel : InnerWorld->GetStreamingLevels())
+		{
+			if (StreamingLevel)
+			{
+				// Note how we will always load the sublevel's package even if will ignore this level. This is a workaround
+				// to a level streaming bug/quirk:
+				// As soon as we first load the sublevels's package the component scene proxies will be (incorrectly?) placed on
+				// the vestigial worlds' FScene. When exporting sublevels though, we need the scene proxies to be on the owning
+				// world, especially for landscapes as their materials will be baked by essentially taking a camera screenshot
+				// from the FScene.
+				// Because of that, we have to always at least load the sublevels here so that the FlushLevelStreaming call below can
+				// call World->RemoveFromWorld from within ULevelStreaming::UpdateStreamingState when we call FlushLevelStreaming
+				// below and leave the loaded level's component's unregistered.
+				// This ensures that if we later try exporting this same sublevel, we won't get the components first registered
+				// during the process of actually first loading the sublevel (which would have placed them on the vestigial world), but
+				// instead, since the level *is already loaded*, a future call to to FlushLevelStreaming below can trigger the landscape
+				// components to be registered on the owning world by the World->AddToWorld call.
+				const FString StreamingLevelWorldAssetPackageName = StreamingLevel->GetWorldAssetPackageName();
+				const FName StreamingLevelWorldAssetPackageFName = StreamingLevel->GetWorldAssetPackageFName();
+				ULevel::StreamedLevelsOwningWorld.Add(StreamingLevelWorldAssetPackageFName, InnerWorld);
+				UPackage* Package = LoadPackage(nullptr, *StreamingLevelWorldAssetPackageName, LOAD_None);
+				ULevel::StreamedLevelsOwningWorld.Remove(StreamingLevelWorldAssetPackageFName);
+
+				const FString LevelName = FPaths::GetBaseFilename(StreamingLevel->GetWorldAssetPackageName());
+				if (LevelsToIgnore.Contains(LevelName))
 				{
-					FWorldContext& NewContext = GEngine->CreateNewWorldContext( EWorldType::EditorPreview );
-					NewContext.SetCurrentWorld( InnerWorld );
-
-					bCreatedContext = true;
+					continue;
 				}
 
-				// Mark all sublevels that we need to be loaded
-				for ( ULevelStreaming* StreamingLevel : InnerWorld->GetStreamingLevels() )
+				const bool bInShouldBeLoaded = true;
+				StreamingLevel->SetShouldBeLoaded(bInShouldBeLoaded);
+
+				// This is a workaround to a level streaming bug/quirk:
+				// We must set these to false here in order to get both our streaming level's current and target
+				// state to LoadedNotVisible. If the level is already visible when we call FlushLevelStreaming,
+				// our level will go directly to the LoadedVisible state.
+				// This may seem desirable, but it means the second FlushLevelStreaming call below won't really do anything.
+				// We need ULevelStreaming::UpdateStreamingState to call World->RemoveFromWorld and World->AddToWorld though,
+				// as that is the only thing that will force the sublevel components' scene proxies to being added to the
+				// *owning* world's FScene, as opposed to being left at the vestigial world's FScenes instead.
+				// If we don't do this, we may get undesired effects when exporting anything that relies on the actual FScene,
+				// like baking landscape materials (see UE-126953)
 				{
-					if ( StreamingLevel )
-					{
-						// Note how we will always load the sublevel's package even if will ignore this level. This is a workaround
-						// to a level streaming bug/quirk:
-						// As soon as we first load the sublevels's package the component scene proxies will be (incorrectly?) placed on
-						// the vestigial worlds' FScene. When exporting sublevels though, we need the scene proxies to be on the owning
-						// world, especially for landscapes as their materials will be baked by essentially taking a camera screenshot
-						// from the FScene.
-						// Because of that, we have to always at least load the sublevels here so that the FlushLevelStreaming call below can
-						// call World->RemoveFromWorld from within ULevelStreaming::UpdateStreamingState when we call FlushLevelStreaming
-						// below and leave the loaded level's component's unregistered.
-						// This ensures that if we later try exporting this same sublevel, we won't get the components first registered
-						// during the process of actually first loading the sublevel (which would have placed them on the vestigial world), but
-						// instead, since the level *is already loaded*, a future call to to FlushLevelStreaming below can trigger the landscape
-						// components to be registered on the owning world by the World->AddToWorld call.
-						const FString StreamingLevelWorldAssetPackageName = StreamingLevel->GetWorldAssetPackageName();
-						const FName StreamingLevelWorldAssetPackageFName = StreamingLevel->GetWorldAssetPackageFName();
-						ULevel::StreamedLevelsOwningWorld.Add( StreamingLevelWorldAssetPackageFName, InnerWorld );
-						UPackage* Package = LoadPackage( nullptr, *StreamingLevelWorldAssetPackageName, LOAD_None );
-						ULevel::StreamedLevelsOwningWorld.Remove( StreamingLevelWorldAssetPackageFName );
-
-						const FString LevelName = FPaths::GetBaseFilename( StreamingLevel->GetWorldAssetPackageName() );
-						if ( LevelsToIgnore.Contains( LevelName ) )
-						{
-							continue;
-						}
-
-						const bool bInShouldBeLoaded = true;
-						StreamingLevel->SetShouldBeLoaded( bInShouldBeLoaded );
-
-						// This is a workaround to a level streaming bug/quirk:
-						// We must set these to false here in order to get both our streaming level's current and target
-						// state to LoadedNotVisible. If the level is already visible when we call FlushLevelStreaming,
-						// our level will go directly to the LoadedVisible state.
-						// This may seem desirable, but it means the second FlushLevelStreaming call below won't really do anything.
-						// We need ULevelStreaming::UpdateStreamingState to call World->RemoveFromWorld and World->AddToWorld though,
-						// as that is the only thing that will force the sublevel components' scene proxies to being added to the
-						// *owning* world's FScene, as opposed to being left at the vestigial world's FScenes instead.
-						// If we don't do this, we may get undesired effects when exporting anything that relies on the actual FScene,
-						// like baking landscape materials (see UE-126953)
-						{
-							const bool bShouldBeVisible = false;
-							StreamingLevel->SetShouldBeVisible( bShouldBeVisible );
-							StreamingLevel->SetShouldBeVisibleInEditor( bShouldBeVisible );
-						}
-					}
+					const bool bShouldBeVisible = false;
+					StreamingLevel->SetShouldBeVisible(bShouldBeVisible);
+					StreamingLevel->SetShouldBeVisibleInEditor(bShouldBeVisible);
 				}
+			}
+		}
 
-				// Synchronously stream in levels
-				InnerWorld->FlushLevelStreaming( EFlushLevelStreamingType::Full );
+		// Synchronously stream in levels
+		InnerWorld->FlushLevelStreaming(EFlushLevelStreamingType::Full);
 
-				// Mark all sublevels that we need to be made visible
-				// For whatever reason this needs to be done with two separate flushes: One for loading, one for
-				// turning visible. If we do this with a single flush the levels will not be synchronously loaded *and* made visible
-				// on this same exact frame, and if we're e.g. baking a landscape immediately after this,
-				// its material will be baked incorrectly (check test_export_level_landscape_bake.py)
-				for ( ULevelStreaming* StreamingLevel : InnerWorld->GetStreamingLevels() )
+		// Mark all sublevels that we need to be made visible
+		// For whatever reason this needs to be done with two separate flushes: One for loading, one for
+		// turning visible. If we do this with a single flush the levels will not be synchronously loaded *and* made visible
+		// on this same exact frame, and if we're e.g. baking a landscape immediately after this,
+		// its material will be baked incorrectly (check test_export_level_landscape_bake.py)
+		for (ULevelStreaming* StreamingLevel : InnerWorld->GetStreamingLevels())
+		{
+			if (StreamingLevel)
+			{
+				const FString LevelName = FPaths::GetBaseFilename(StreamingLevel->GetWorldAssetPackageName());
+				if (LevelsToIgnore.Contains(LevelName))
 				{
-					if ( StreamingLevel )
-					{
-						const FString LevelName = FPaths::GetBaseFilename( StreamingLevel->GetWorldAssetPackageName() );
-						if ( LevelsToIgnore.Contains( LevelName ) )
-						{
-							continue;
-						}
-
-						const bool bInShouldBeVisible = true;
-						StreamingLevel->SetShouldBeVisible( bInShouldBeVisible );
-						StreamingLevel->SetShouldBeVisibleInEditor( bInShouldBeVisible );
-					}
+					continue;
 				}
 
-				// Synchronously show levels right now
-				InnerWorld->FlushLevelStreaming( EFlushLevelStreamingType::Visibility );
+				const bool bInShouldBeVisible = true;
+				StreamingLevel->SetShouldBeVisible(bInShouldBeVisible);
+				StreamingLevel->SetShouldBeVisibleInEditor(bInShouldBeVisible);
+			}
+		}
 
-				WaitForAllAsyncAndSteamingTasks(InnerWorld);
+		// Synchronously show levels right now
+		InnerWorld->FlushLevelStreaming(EFlushLevelStreamingType::Visibility);
 
-				if ( bCreatedContext )
-				{
-					GEngine->DestroyWorldContext( InnerWorld );
-				}
-			}
+		WaitForAllAsyncAndSteamingTasks(InnerWorld);
+
+		if (bCreatedContext)
+		{
+			GEngine->DestroyWorldContext(InnerWorld);
 		}
 	}
-}
+}	 // namespace UE::UsdConversionBlueprintLibrary::Private
 
 int32 UUsdConversionBlueprintLibrary::GetNumLevelsToExport(UWorld* World, const TSet<FString>& LevelsToIgnore)
 {
@@ -203,24 +193,24 @@ int32 UUsdConversionBlueprintLibrary::GetNumLevelsToExport(UWorld* World, const
 	return Count;
 }
 
-void UUsdConversionBlueprintLibrary::StreamInRequiredLevels( UWorld* World, const TSet<FString>& LevelsToIgnore )
+void UUsdConversionBlueprintLibrary::StreamInRequiredLevels(UWorld* World, const TSet<FString>& LevelsToIgnore)
 {
-	if ( !World )
+	if (!World)
 	{
 		return;
 	}
 
-	if ( ULevel* PersistentLevel = World->PersistentLevel )
+	if (ULevel* PersistentLevel = World->PersistentLevel)
 	{
-		UE::UsdConversionBlueprintLibraryImpl::Private::StreamInLevels( PersistentLevel, LevelsToIgnore );
+		UE::UsdConversionBlueprintLibrary::Private::StreamInLevels(PersistentLevel, LevelsToIgnore);
 	}
 }
 
 void UUsdConversionBlueprintLibrary::RevertSequencerAnimations()
 {
-	for ( const TWeakPtr<ISequencer>& Sequencer : FLevelEditorSequencerIntegration::Get().GetSequencers() )
+	for (const TWeakPtr<ISequencer>& Sequencer : FLevelEditorSequencerIntegration::Get().GetSequencers())
 	{
-		if ( TSharedPtr<ISequencer> PinnedSequencer = Sequencer.Pin() )
+		if (TSharedPtr<ISequencer> PinnedSequencer = Sequencer.Pin())
 		{
 			PinnedSequencer->EnterSilentMode();
 			PinnedSequencer->RestorePreAnimatedState();
@@ -230,9 +220,9 @@ void UUsdConversionBlueprintLibrary::RevertSequencerAnimations()
 
 void UUsdConversionBlueprintLibrary::ReapplySequencerAnimations()
 {
-	for ( const TWeakPtr<ISequencer>& Sequencer : FLevelEditorSequencerIntegration::Get().GetSequencers() )
+	for (const TWeakPtr<ISequencer>& Sequencer : FLevelEditorSequencerIntegration::Get().GetSequencers())
 	{
-		if ( TSharedPtr<ISequencer> PinnedSequencer = Sequencer.Pin() )
+		if (TSharedPtr<ISequencer> PinnedSequencer = Sequencer.Pin())
 		{
 			PinnedSequencer->InvalidateCachedData();
 			PinnedSequencer->ForceEvaluate();
@@ -241,307 +231,319 @@ void UUsdConversionBlueprintLibrary::ReapplySequencerAnimations()
 	}
 }
 
-TArray<FString> UUsdConversionBlueprintLibrary::GetLoadedLevelNames( UWorld* World )
+TArray<FString> UUsdConversionBlueprintLibrary::GetLoadedLevelNames(UWorld* World)
 {
 	TArray<FString> Result;
 
-	for ( ULevelStreaming* StreamingLevel : World->GetStreamingLevels() )
+	for (ULevelStreaming* StreamingLevel : World->GetStreamingLevels())
 	{
-		if ( StreamingLevel && StreamingLevel->IsLevelLoaded() )
+		if (StreamingLevel && StreamingLevel->IsLevelLoaded())
 		{
-			Result.Add( StreamingLevel->GetWorldAssetPackageName() );
+			Result.Add(StreamingLevel->GetWorldAssetPackageName());
 		}
 	}
 
 	return Result;
 }
 
-TArray<FString> UUsdConversionBlueprintLibrary::GetVisibleInEditorLevelNames( UWorld* World )
+TArray<FString> UUsdConversionBlueprintLibrary::GetVisibleInEditorLevelNames(UWorld* World)
 {
 	TArray<FString> Result;
 
-	for ( ULevelStreaming* StreamingLevel : World->GetStreamingLevels() )
+	for (ULevelStreaming* StreamingLevel : World->GetStreamingLevels())
 	{
-		if ( StreamingLevel && StreamingLevel->GetShouldBeVisibleInEditor() )
+		if (StreamingLevel && StreamingLevel->GetShouldBeVisibleInEditor())
 		{
-			Result.Add( StreamingLevel->GetWorldAssetPackageName() );
+			Result.Add(StreamingLevel->GetWorldAssetPackageName());
 		}
 	}
 
 	return Result;
 }
 
-void UUsdConversionBlueprintLibrary::StreamOutLevels( UWorld* OwningWorld, const TArray<FString>& LevelNamesToStreamOut, const TArray<FString>& LevelNamesToHide )
+void UUsdConversionBlueprintLibrary::StreamOutLevels(
+	UWorld* OwningWorld,
+	const TArray<FString>& LevelNamesToStreamOut,
+	const TArray<FString>& LevelNamesToHide
+)
 {
-	if ( LevelNamesToStreamOut.Num() == 0 && LevelNamesToHide.Num() == 0 )
+	if (LevelNamesToStreamOut.Num() == 0 && LevelNamesToHide.Num() == 0)
 	{
 		return;
 	}
 
 	bool bCreatedContext = false;
-	FWorldContext* Context = GEngine->GetWorldContextFromWorld( OwningWorld );
-	if ( !Context )
+	FWorldContext* Context = GEngine->GetWorldContextFromWorld(OwningWorld);
+	if (!Context)
 	{
-		FWorldContext& NewContext = GEngine->CreateNewWorldContext( EWorldType::EditorPreview );
-		NewContext.SetCurrentWorld( OwningWorld );
+		FWorldContext& NewContext = GEngine->CreateNewWorldContext(EWorldType::EditorPreview);
+		NewContext.SetCurrentWorld(OwningWorld);
 
 		bCreatedContext = true;
 	}
 
-	for ( ULevelStreaming* StreamingLevel : OwningWorld->GetStreamingLevels() )
+	for (ULevelStreaming* StreamingLevel : OwningWorld->GetStreamingLevels())
 	{
-		if ( StreamingLevel )
+		if (StreamingLevel)
 		{
 			const FString& LevelName = StreamingLevel->GetWorldAssetPackageName();
 
-			if ( LevelNamesToHide.Contains( LevelName ) )
+			if (LevelNamesToHide.Contains(LevelName))
 			{
 				const bool bInShouldBeVisible = false;
-				StreamingLevel->SetShouldBeVisible( bInShouldBeVisible );
-				StreamingLevel->SetShouldBeVisibleInEditor( bInShouldBeVisible );
+				StreamingLevel->SetShouldBeVisible(bInShouldBeVisible);
+				StreamingLevel->SetShouldBeVisibleInEditor(bInShouldBeVisible);
 			}
 
-			if ( LevelNamesToStreamOut.Contains( LevelName ) )
+			if (LevelNamesToStreamOut.Contains(LevelName))
 			{
 				const bool bInShouldBeLoaded = false;
-				StreamingLevel->SetShouldBeVisible( bInShouldBeLoaded );
-				StreamingLevel->SetShouldBeLoaded( bInShouldBeLoaded );
+				StreamingLevel->SetShouldBeVisible(bInShouldBeLoaded);
+				StreamingLevel->SetShouldBeLoaded(bInShouldBeLoaded);
 			}
 		}
 	}
 
-	UE::UsdConversionBlueprintLibraryImpl::Private::WaitForAllAsyncAndSteamingTasks(OwningWorld);
+	UE::UsdConversionBlueprintLibrary::Private::WaitForAllAsyncAndSteamingTasks(OwningWorld);
 
-	if ( bCreatedContext )
+	if (bCreatedContext)
 	{
-		GEngine->DestroyWorldContext( OwningWorld );
+		GEngine->DestroyWorldContext(OwningWorld);
 	}
 }
 
-TSet<AActor*> UUsdConversionBlueprintLibrary::GetActorsToConvert( UWorld* World )
+TSet<AActor*> UUsdConversionBlueprintLibrary::GetActorsToConvert(UWorld* World)
 {
 	TSet<AActor*> Result;
-	if ( !World )
+	if (!World)
 	{
 		return Result;
 	}
 
-	auto CollectActors = [ &Result ]( ULevel* Level )
+	auto CollectActors = [&Result](ULevel* Level)
 	{
-		if ( !Level )
+		if (!Level)
 		{
 			return;
 		}
 
-		Result.Append( ObjectPtrDecay(Level->Actors) );
+		Result.Append(ObjectPtrDecay(Level->Actors));
 	};
 
-	CollectActors( World->PersistentLevel );
+	CollectActors(World->PersistentLevel);
 
-	for ( ULevelStreaming* StreamingLevel : World->GetStreamingLevels() )
+	for (ULevelStreaming* StreamingLevel : World->GetStreamingLevels())
 	{
-		if ( StreamingLevel && StreamingLevel->IsLevelLoaded() && StreamingLevel->GetShouldBeVisibleInEditor() )
+		if (StreamingLevel && StreamingLevel->IsLevelLoaded() && StreamingLevel->GetShouldBeVisibleInEditor())
 		{
-			if ( ULevel* Level = StreamingLevel->GetLoadedLevel() )
+			if (ULevel* Level = StreamingLevel->GetLoadedLevel())
 			{
-				CollectActors( Level );
+				CollectActors(Level);
 			}
 		}
 	}
 
-	Result.Remove( nullptr );
+	Result.Remove(nullptr);
 
 	// Remove transient actors here because it is not possible to do this via Python
 	TSet<AActor*> ActorsToRemove;
-	ActorsToRemove.Reserve( Result.Num() );
-	for ( AActor* Actor : Result )
+	ActorsToRemove.Reserve(Result.Num());
+	for (AActor* Actor : Result)
 	{
 		// Actors marked with this tag are transient because they're spawnables: We still want to export those, in case
 		// we're exporting a level for a LevelSequence export with spawnables.
-		if ( Actor->HasAnyFlags( EObjectFlags::RF_Transient ) && !Actor->Tags.Contains( TEXT( "SequencerActor" ) ) )
+		if (Actor->HasAnyFlags(EObjectFlags::RF_Transient) && !Actor->Tags.Contains(TEXT("SequencerActor")))
 		{
-			ActorsToRemove.Add( Actor );
+			ActorsToRemove.Add(Actor);
 		}
 	}
-	Result = Result.Difference( ActorsToRemove );
+	Result = Result.Difference(ActorsToRemove);
 	return Result;
 }
 
-FString UUsdConversionBlueprintLibrary::GenerateObjectVersionString( const UObject* ObjectToExport, UObject* ExportOptions )
+FString UUsdConversionBlueprintLibrary::GenerateObjectVersionString(const UObject* ObjectToExport, UObject* ExportOptions)
 {
-	if ( !ObjectToExport )
+	if (!ObjectToExport)
 	{
 		return {};
 	}
 
 	FSHA1 SHA1;
 
-	if ( !IUsdClassesModule::HashObjectPackage( ObjectToExport, SHA1 ) )
+	if (!IUsdClassesModule::HashObjectPackage(ObjectToExport, SHA1))
 	{
 		return {};
 	}
 
-	if ( ULevelExporterUSDOptions* LevelExportOptions = Cast<ULevelExporterUSDOptions>( ExportOptions ) )
+	if (ULevelExporterUSDOptions* LevelExportOptions = Cast<ULevelExporterUSDOptions>(ExportOptions))
 	{
-		UsdUtils::HashForLevelExport( *LevelExportOptions, SHA1 );
+		UsdUtils::HashForLevelExport(*LevelExportOptions, SHA1);
 	}
 
 	FSHAHash Hash;
 	SHA1.Final();
-	SHA1.GetHash( &Hash.Hash[ 0 ] );
+	SHA1.GetHash(&Hash.Hash[0]);
 	return Hash.ToString();
 }
 
-bool UUsdConversionBlueprintLibrary::CanExportToLayer( const FString& TargetFilePath )
+bool UUsdConversionBlueprintLibrary::CanExportToLayer(const FString& TargetFilePath)
 {
-	return IUsdExporterModule::CanExportToLayer( TargetFilePath );
+	return IUsdExporterModule::CanExportToLayer(TargetFilePath);
 }
 
-FString UUsdConversionBlueprintLibrary::MakePathRelativeToLayer( const FString& AnchorLayerPath, const FString& PathToMakeRelative )
+FString UUsdConversionBlueprintLibrary::MakePathRelativeToLayer(const FString& AnchorLayerPath, const FString& PathToMakeRelative)
 {
 #if USE_USD_SDK
-	if ( UE::FSdfLayer Layer = UE::FSdfLayer::FindOrOpen( *AnchorLayerPath ) )
+	if (UE::FSdfLayer Layer = UE::FSdfLayer::FindOrOpen(*AnchorLayerPath))
 	{
 		FString Path = PathToMakeRelative;
-		UsdUtils::MakePathRelativeToLayer( Layer, Path );
+		UsdUtils::MakePathRelativeToLayer(Layer, Path);
 		return Path;
 	}
 	else
 	{
-		UE_LOG(LogUsd, Error, TEXT("Failed to find a layer with path '%s' to make the path '%s' relative to"), *AnchorLayerPath, *PathToMakeRelative );
+		UE_LOG(LogUsd, Error, TEXT("Failed to find a layer with path '%s' to make the path '%s' relative to"), *AnchorLayerPath, *PathToMakeRelative);
 		return PathToMakeRelative;
 	}
 #else
 	return FString();
-#endif // USE_USD_SDK
+#endif	  // USE_USD_SDK
 }
 
-void UUsdConversionBlueprintLibrary::InsertSubLayer( const FString& ParentLayerPath, const FString& SubLayerPath, int32 Index /*= -1 */ )
+void UUsdConversionBlueprintLibrary::InsertSubLayer(const FString& ParentLayerPath, const FString& SubLayerPath, int32 Index /*= -1 */)
 {
 #if USE_USD_SDK
-	if ( ParentLayerPath.IsEmpty() || SubLayerPath.IsEmpty() )
+	if (ParentLayerPath.IsEmpty() || SubLayerPath.IsEmpty())
 	{
 		return;
 	}
 
-	if ( UE::FSdfLayer Layer = UE::FSdfLayer::FindOrOpen( *ParentLayerPath ) )
+	if (UE::FSdfLayer Layer = UE::FSdfLayer::FindOrOpen(*ParentLayerPath))
 	{
-		UsdUtils::InsertSubLayer( Layer, *SubLayerPath, Index );
+		UsdUtils::InsertSubLayer(Layer, *SubLayerPath, Index);
 	}
 	else
 	{
-		UE_LOG( LogUsd, Error, TEXT( "Failed to find a parent layer '%s' when trying to insert sublayer '%s'" ), *ParentLayerPath, *SubLayerPath );
+		UE_LOG(LogUsd, Error, TEXT("Failed to find a parent layer '%s' when trying to insert sublayer '%s'"), *ParentLayerPath, *SubLayerPath);
 	}
-#endif // USE_USD_SDK
+#endif	  // USE_USD_SDK
 }
 
-void UUsdConversionBlueprintLibrary::AddReference( const FString& ReferencingStagePath, const FString& ReferencingPrimPath, const FString& TargetStagePath )
+void UUsdConversionBlueprintLibrary::AddReference(
+	const FString& ReferencingStagePath,
+	const FString& ReferencingPrimPath,
+	const FString& TargetStagePath
+)
 {
 #if USE_USD_SDK
 	TArray<UE::FUsdStage> PreviouslyOpenedStages = UnrealUSDWrapper::GetAllStagesFromCache();
 
 	// Open using the stage cache as it's very likely this stage is already in there anyway
-	UE::FUsdStage ReferencingStage = UnrealUSDWrapper::OpenStage( *ReferencingStagePath, EUsdInitialLoadSet::LoadAll );
-	if ( ReferencingStage )
+	UE::FUsdStage ReferencingStage = UnrealUSDWrapper::OpenStage(*ReferencingStagePath, EUsdInitialLoadSet::LoadAll);
+	if (ReferencingStage)
 	{
-		if ( UE::FUsdPrim ReferencingPrim = ReferencingStage.GetPrimAtPath( UE::FSdfPath( *ReferencingPrimPath ) ) )
+		if (UE::FUsdPrim ReferencingPrim = ReferencingStage.GetPrimAtPath(UE::FSdfPath(*ReferencingPrimPath)))
 		{
-			UsdUtils::AddReference( ReferencingPrim, *TargetStagePath );
+			UsdUtils::AddReference(ReferencingPrim, *TargetStagePath);
 		}
 	}
 
 	// Cleanup or else the stage cache will keep these stages open forever
-	if ( !PreviouslyOpenedStages.Contains( ReferencingStage ) )
+	if (!PreviouslyOpenedStages.Contains(ReferencingStage))
 	{
-		UnrealUSDWrapper::EraseStageFromCache( ReferencingStage );
+		UnrealUSDWrapper::EraseStageFromCache(ReferencingStage);
 	}
-#endif // USE_USD_SDK
+#endif	  // USE_USD_SDK
 }
 
-void UUsdConversionBlueprintLibrary::AddPayload( const FString& ReferencingStagePath, const FString& ReferencingPrimPath, const FString& TargetStagePath )
+void UUsdConversionBlueprintLibrary::AddPayload(
+	const FString& ReferencingStagePath,
+	const FString& ReferencingPrimPath,
+	const FString& TargetStagePath
+)
 {
 #if USE_USD_SDK
 	TArray<UE::FUsdStage> PreviouslyOpenedStages = UnrealUSDWrapper::GetAllStagesFromCache();
 
 	// Open using the stage cache as it's very likely this stage is already in there anyway
-	UE::FUsdStage ReferencingStage = UnrealUSDWrapper::OpenStage( *ReferencingStagePath, EUsdInitialLoadSet::LoadAll );
-	if ( ReferencingStage )
+	UE::FUsdStage ReferencingStage = UnrealUSDWrapper::OpenStage(*ReferencingStagePath, EUsdInitialLoadSet::LoadAll);
+	if (ReferencingStage)
 	{
-		if ( UE::FUsdPrim ReferencingPrim = ReferencingStage.GetPrimAtPath( UE::FSdfPath( *ReferencingPrimPath ) ) )
+		if (UE::FUsdPrim ReferencingPrim = ReferencingStage.GetPrimAtPath(UE::FSdfPath(*ReferencingPrimPath)))
 		{
-			UsdUtils::AddPayload( ReferencingPrim, *TargetStagePath );
+			UsdUtils::AddPayload(ReferencingPrim, *TargetStagePath);
 		}
 	}
 
 	// Cleanup or else the stage cache will keep these stages open forever
-	if ( !PreviouslyOpenedStages.Contains( ReferencingStage ) )
+	if (!PreviouslyOpenedStages.Contains(ReferencingStage))
 	{
-		UnrealUSDWrapper::EraseStageFromCache( ReferencingStage );
+		UnrealUSDWrapper::EraseStageFromCache(ReferencingStage);
 	}
-#endif // USE_USD_SDK
+#endif	  // USE_USD_SDK
 }
 
-FString UUsdConversionBlueprintLibrary::GetPrimPathForObject( const UObject* ActorOrComponent, const FString& ParentPrimPath, bool bUseActorFolders )
+FString UUsdConversionBlueprintLibrary::GetPrimPathForObject(const UObject* ActorOrComponent, const FString& ParentPrimPath, bool bUseActorFolders)
 {
 #if USE_USD_SDK
-	return UsdUtils::GetPrimPathForObject( ActorOrComponent, ParentPrimPath, bUseActorFolders );
+	return UsdUtils::GetPrimPathForObject(ActorOrComponent, ParentPrimPath, bUseActorFolders);
 #else
 	return {};
-#endif // USE_USD_SDK
+#endif	  // USE_USD_SDK
 }
 
-FString UUsdConversionBlueprintLibrary::GetSchemaNameForComponent( const USceneComponent* Component )
+FString UUsdConversionBlueprintLibrary::GetSchemaNameForComponent(const USceneComponent* Component)
 {
 #if USE_USD_SDK
-	if ( Component )
+	if (Component)
 	{
-		return UsdUtils::GetSchemaNameForComponent( *Component );
+		return UsdUtils::GetSchemaNameForComponent(*Component);
 	}
-#endif // USE_USD_SDK
+#endif	  // USE_USD_SDK
 
 	return {};
 }
 
-AInstancedFoliageActor* UUsdConversionBlueprintLibrary::GetInstancedFoliageActorForLevel( bool bCreateIfNone /*= false */, ULevel* Level /*= nullptr */ )
+AInstancedFoliageActor* UUsdConversionBlueprintLibrary::GetInstancedFoliageActorForLevel(bool bCreateIfNone /*= false */, ULevel* Level /*= nullptr */)
 {
-	if ( !Level )
+	if (!Level)
 	{
 		const bool bEnsureIsGWorld = false;
-		UWorld* EditorWorld = GEditor ? GEditor->GetEditorWorldContext( bEnsureIsGWorld ).World() : nullptr;
-		if ( !EditorWorld )
+		UWorld* EditorWorld = GEditor ? GEditor->GetEditorWorldContext(bEnsureIsGWorld).World() : nullptr;
+		if (!EditorWorld)
 		{
 			return nullptr;
 		}
 
 		Level = EditorWorld->GetCurrentLevel();
-		if ( !Level )
+		if (!Level)
 		{
 			return nullptr;
 		}
 	}
 
-	return AInstancedFoliageActor::GetInstancedFoliageActorForLevel( Level, bCreateIfNone );
+	return AInstancedFoliageActor::GetInstancedFoliageActorForLevel(Level, bCreateIfNone);
 }
 
-TArray<UFoliageType*> UUsdConversionBlueprintLibrary::GetUsedFoliageTypes( AInstancedFoliageActor* Actor )
+TArray<UFoliageType*> UUsdConversionBlueprintLibrary::GetUsedFoliageTypes(AInstancedFoliageActor* Actor)
 {
 	TArray<UFoliageType*> Result;
-	if ( !Actor )
+	if (!Actor)
 	{
 		return Result;
 	}
 
-	for ( const TPair<UFoliageType*, TUniqueObj<FFoliageInfo>>& FoliagePair : Actor->GetFoliageInfos() )
+	for (const TPair<UFoliageType*, TUniqueObj<FFoliageInfo>>& FoliagePair : Actor->GetFoliageInfos())
 	{
-		Result.Add( FoliagePair.Key );
+		Result.Add(FoliagePair.Key);
 	}
 
 	return Result;
 }
 
-UObject* UUsdConversionBlueprintLibrary::GetSource( UFoliageType* FoliageType )
+UObject* UUsdConversionBlueprintLibrary::GetSource(UFoliageType* FoliageType)
 {
-	if ( FoliageType )
+	if (FoliageType)
 	{
 		return FoliageType->GetSource();
 	}
@@ -549,37 +551,41 @@ UObject* UUsdConversionBlueprintLibrary::GetSource( UFoliageType* FoliageType )
 	return nullptr;
 }
 
-TArray<FTransform> UUsdConversionBlueprintLibrary::GetInstanceTransforms( AInstancedFoliageActor* Actor, UFoliageType* FoliageType, ULevel* InstancesLevel )
+TArray<FTransform> UUsdConversionBlueprintLibrary::GetInstanceTransforms(
+	AInstancedFoliageActor* Actor,
+	UFoliageType* FoliageType,
+	ULevel* InstancesLevel
+)
 {
 	TArray<FTransform> Result;
-	if ( !Actor || !FoliageType )
+	if (!Actor || !FoliageType)
 	{
 		return Result;
 	}
 
 	// Modified from AInstancedFoliageActor::GetInstancesForComponent to limit traversal only to our FoliageType
 
-	if ( const TUniqueObj<FFoliageInfo>* FoundInfo = Actor->GetFoliageInfos().Find( FoliageType ) )
+	if (const TUniqueObj<FFoliageInfo>* FoundInfo = Actor->GetFoliageInfos().Find(FoliageType))
 	{
-		const FFoliageInfo& Info = ( *FoundInfo ).Get();
+		const FFoliageInfo& Info = (*FoundInfo).Get();
 
 		// Collect IDs of components that are on the same level as the actor's level. This because later on we'll have level-by-level
 		// export, and we'd want one point instancer per level
-		for ( const TPair<FFoliageInstanceBaseId, FFoliageInstanceBaseInfo>& FoliageInstancePair : Actor->InstanceBaseCache.InstanceBaseMap )
+		for (const TPair<FFoliageInstanceBaseId, FFoliageInstanceBaseInfo>& FoliageInstancePair : Actor->InstanceBaseCache.InstanceBaseMap)
 		{
 			UActorComponent* Comp = FoliageInstancePair.Value.BasePtr.Get();
-			if ( !Comp || ( InstancesLevel && ( Comp->GetComponentLevel() != InstancesLevel ) ) )
+			if (!Comp || (InstancesLevel && (Comp->GetComponentLevel() != InstancesLevel)))
 			{
 				continue;
 			}
 
-			if ( const auto* InstanceSet = Info.ComponentHash.Find( FoliageInstancePair.Key ) )
+			if (const auto* InstanceSet = Info.ComponentHash.Find(FoliageInstancePair.Key))
 			{
-				Result.Reserve( Result.Num() + InstanceSet->Num() );
-				for ( int32 InstanceIndex : *InstanceSet )
+				Result.Reserve(Result.Num() + InstanceSet->Num());
+				for (int32 InstanceIndex : *InstanceSet)
 				{
-					const FFoliageInstancePlacementInfo* Instance = &Info.Instances[ InstanceIndex ];
-					Result.Emplace( FQuat( Instance->Rotation ), Instance->Location, ( FVector ) Instance->DrawScale3D );
+					const FFoliageInstancePlacementInfo* Instance = &Info.Instances[InstanceIndex];
+					Result.Emplace(FQuat(Instance->Rotation), Instance->Location, (FVector)Instance->DrawScale3D);
 				}
 			}
 		}
@@ -588,16 +594,16 @@ TArray<FTransform> UUsdConversionBlueprintLibrary::GetInstanceTransforms( AInsta
 	return Result;
 }
 
-TArray<FAnalyticsEventAttr> UUsdConversionBlueprintLibrary::GetAnalyticsAttributes( const ULevelExporterUSDOptions* Options )
+TArray<FAnalyticsEventAttr> UUsdConversionBlueprintLibrary::GetAnalyticsAttributes(const ULevelExporterUSDOptions* Options)
 {
 	TArray<FAnalyticsEventAttr> Attrs;
-	if ( Options )
+	if (Options)
 	{
 		TArray<FAnalyticsEventAttribute> Attributes;
-		UsdUtils::AddAnalyticsAttributes( *Options, Attributes );
+		UsdUtils::AddAnalyticsAttributes(*Options, Attributes);
 
-		Attrs.Reserve( Attributes.Num() );
-		for(const FAnalyticsEventAttribute& Attribute : Attributes )
+		Attrs.Reserve(Attributes.Num());
+		for (const FAnalyticsEventAttribute& Attribute : Attributes)
 		{
 			FAnalyticsEventAttr& NewAttr = Attrs.Emplace_GetRef();
 			NewAttr.Name = Attribute.GetName();
@@ -607,102 +613,106 @@ TArray<FAnalyticsEventAttr> UUsdConversionBlueprintLibrary::GetAnalyticsAttribut
 	return Attrs;
 }
 
-void UUsdConversionBlueprintLibrary::SendAnalytics( const TArray<FAnalyticsEventAttr>& Attrs, const FString& EventName, bool bAutomated, double ElapsedSeconds, double NumberOfFrames, const FString& Extension )
+void UUsdConversionBlueprintLibrary::SendAnalytics(
+	const TArray<FAnalyticsEventAttr>& Attrs,
+	const FString& EventName,
+	bool bAutomated,
+	double ElapsedSeconds,
+	double NumberOfFrames,
+	const FString& Extension
+)
 {
 	TArray<FAnalyticsEventAttribute> Converted;
-	Converted.Reserve( Attrs.Num() );
-	for ( const FAnalyticsEventAttr& Attr : Attrs )
+	Converted.Reserve(Attrs.Num());
+	for (const FAnalyticsEventAttr& Attr : Attrs)
 	{
-		Converted.Emplace( Attr.Name, Attr.Value );
+		Converted.Emplace(Attr.Name, Attr.Value);
 	}
 
-	IUsdClassesModule::SendAnalytics( MoveTemp( Converted ), EventName, bAutomated, ElapsedSeconds, NumberOfFrames, Extension );
+	IUsdClassesModule::SendAnalytics(MoveTemp(Converted), EventName, bAutomated, ElapsedSeconds, NumberOfFrames, Extension);
 }
 
-void UUsdConversionBlueprintLibrary::RemoveAllPrimSpecs( const FString& StageRootLayer, const FString& PrimPath, const FString& TargetLayer )
+void UUsdConversionBlueprintLibrary::RemoveAllPrimSpecs(const FString& StageRootLayer, const FString& PrimPath, const FString& TargetLayer)
 {
 #if USE_USD_SDK
 	const bool bUseStageCache = true;
-	UE::FUsdStage Stage = UnrealUSDWrapper::OpenStage( *StageRootLayer, EUsdInitialLoadSet::LoadAll, bUseStageCache );
-	if ( !Stage )
+	UE::FUsdStage Stage = UnrealUSDWrapper::OpenStage(*StageRootLayer, EUsdInitialLoadSet::LoadAll, bUseStageCache);
+	if (!Stage)
 	{
 		return;
 	}
 
-	UsdUtils::RemoveAllLocalPrimSpecs(
-		Stage.GetPrimAtPath( UE::FSdfPath{ *PrimPath } ),
-		UE::FSdfLayer::FindOrOpen( *TargetLayer )
-	);
-#endif // USE_USD_SDK
+	UsdUtils::RemoveAllLocalPrimSpecs(Stage.GetPrimAtPath(UE::FSdfPath{*PrimPath}), UE::FSdfLayer::FindOrOpen(*TargetLayer));
+#endif	  // USE_USD_SDK
 }
 
-bool UUsdConversionBlueprintLibrary::CutPrims( const FString& StageRootLayer, const TArray<FString>& PrimPaths )
+bool UUsdConversionBlueprintLibrary::CutPrims(const FString& StageRootLayer, const TArray<FString>& PrimPaths)
 {
 #if USE_USD_SDK
 	const bool bUseStageCache = true;
-	UE::FUsdStage Stage = UnrealUSDWrapper::OpenStage( *StageRootLayer, EUsdInitialLoadSet::LoadAll, bUseStageCache );
-	if ( !Stage )
+	UE::FUsdStage Stage = UnrealUSDWrapper::OpenStage(*StageRootLayer, EUsdInitialLoadSet::LoadAll, bUseStageCache);
+	if (!Stage)
 	{
 		return false;
 	}
 
 	TArray<UE::FUsdPrim> Prims;
-	Prims.Reserve( PrimPaths.Num() );
+	Prims.Reserve(PrimPaths.Num());
 
-	for ( const FString& PrimPath : PrimPaths )
+	for (const FString& PrimPath : PrimPaths)
 	{
-		Prims.Add( Stage.GetPrimAtPath( UE::FSdfPath{ *PrimPath } ) );
+		Prims.Add(Stage.GetPrimAtPath(UE::FSdfPath{*PrimPath}));
 	}
 
-	return UsdUtils::CutPrims( Prims );
+	return UsdUtils::CutPrims(Prims);
 #else
 	return false;
-#endif // USE_USD_SDK
+#endif	  // USE_USD_SDK
 }
 
-bool UUsdConversionBlueprintLibrary::CopyPrims( const FString& StageRootLayer, const TArray<FString>& PrimPaths )
+bool UUsdConversionBlueprintLibrary::CopyPrims(const FString& StageRootLayer, const TArray<FString>& PrimPaths)
 {
 #if USE_USD_SDK
 	const bool bUseStageCache = true;
-	UE::FUsdStage Stage = UnrealUSDWrapper::OpenStage( *StageRootLayer, EUsdInitialLoadSet::LoadAll, bUseStageCache );
-	if ( !Stage )
+	UE::FUsdStage Stage = UnrealUSDWrapper::OpenStage(*StageRootLayer, EUsdInitialLoadSet::LoadAll, bUseStageCache);
+	if (!Stage)
 	{
 		return false;
 	}
 
 	TArray<UE::FUsdPrim> Prims;
-	Prims.Reserve( PrimPaths.Num() );
+	Prims.Reserve(PrimPaths.Num());
 
-	for ( const FString& PrimPath : PrimPaths )
+	for (const FString& PrimPath : PrimPaths)
 	{
-		Prims.Add( Stage.GetPrimAtPath( UE::FSdfPath{ *PrimPath } ) );
+		Prims.Add(Stage.GetPrimAtPath(UE::FSdfPath{*PrimPath}));
 	}
 
-	return UsdUtils::CopyPrims( Prims );
+	return UsdUtils::CopyPrims(Prims);
 #else
 	return false;
-#endif // USE_USD_SDK
+#endif	  // USE_USD_SDK
 }
 
-TArray<FString> UUsdConversionBlueprintLibrary::PastePrims( const FString& StageRootLayer, const FString& ParentPrimPath )
+TArray<FString> UUsdConversionBlueprintLibrary::PastePrims(const FString& StageRootLayer, const FString& ParentPrimPath)
 {
 	TArray<FString> Result;
 
 #if USE_USD_SDK
 	const bool bUseStageCache = true;
-	UE::FUsdStage Stage = UnrealUSDWrapper::OpenStage( *StageRootLayer, EUsdInitialLoadSet::LoadAll, bUseStageCache );
-	if ( !Stage )
+	UE::FUsdStage Stage = UnrealUSDWrapper::OpenStage(*StageRootLayer, EUsdInitialLoadSet::LoadAll, bUseStageCache);
+	if (!Stage)
 	{
 		return Result;
 	}
 
-	TArray<UE::FSdfPath> PastedPrims = UsdUtils::PastePrims( Stage.GetPrimAtPath( UE::FSdfPath{ *ParentPrimPath } ) );
+	TArray<UE::FSdfPath> PastedPrims = UsdUtils::PastePrims(Stage.GetPrimAtPath(UE::FSdfPath{*ParentPrimPath}));
 
-	for ( const UE::FSdfPath& DuplicatePrim : PastedPrims )
+	for (const UE::FSdfPath& DuplicatePrim : PastedPrims)
 	{
-		Result.Add( DuplicatePrim.GetString() );
+		Result.Add(DuplicatePrim.GetString());
 	}
-#endif // USE_USD_SDK
+#endif	  // USE_USD_SDK
 
 	return Result;
 }
@@ -717,39 +727,40 @@ void UUsdConversionBlueprintLibrary::ClearPrimClipboard()
 	UsdUtils::ClearPrimClipboard();
 }
 
-TArray<FString> UUsdConversionBlueprintLibrary::DuplicatePrims( const FString& StageRootLayer, const TArray<FString>& PrimPaths, EUsdDuplicateType DuplicateType, const FString& TargetLayer )
+TArray<FString> UUsdConversionBlueprintLibrary::DuplicatePrims(
+	const FString& StageRootLayer,
+	const TArray<FString>& PrimPaths,
+	EUsdDuplicateType DuplicateType,
+	const FString& TargetLayer
+)
 {
 	TArray<FString> Result;
-	Result.SetNum( PrimPaths.Num() );
+	Result.SetNum(PrimPaths.Num());
 
 #if USE_USD_SDK
 	const bool bUseStageCache = true;
-	UE::FUsdStage Stage = UnrealUSDWrapper::OpenStage( *StageRootLayer, EUsdInitialLoadSet::LoadAll, bUseStageCache );
-	if ( !Stage )
+	UE::FUsdStage Stage = UnrealUSDWrapper::OpenStage(*StageRootLayer, EUsdInitialLoadSet::LoadAll, bUseStageCache);
+	if (!Stage)
 	{
 		return Result;
 	}
 
 	TArray<UE::FUsdPrim> Prims;
-	Prims.Reserve( PrimPaths.Num() );
+	Prims.Reserve(PrimPaths.Num());
 
-	for ( const FString& PrimPath : PrimPaths )
+	for (const FString& PrimPath : PrimPaths)
 	{
-		Prims.Add( Stage.GetPrimAtPath( UE::FSdfPath{ *PrimPath } ) );
+		Prims.Add(Stage.GetPrimAtPath(UE::FSdfPath{*PrimPath}));
 	}
 
-	TArray<UE::FSdfPath> DuplicatedPrims = UsdUtils::DuplicatePrims(
-		Prims,
-		DuplicateType,
-		UE::FSdfLayer::FindOrOpen( *TargetLayer )
-	);
+	TArray<UE::FSdfPath> DuplicatedPrims = UsdUtils::DuplicatePrims(Prims, DuplicateType, UE::FSdfLayer::FindOrOpen(*TargetLayer));
 
-	for ( int32 Index = 0; Index < DuplicatedPrims.Num(); ++Index )
+	for (int32 Index = 0; Index < DuplicatedPrims.Num(); ++Index)
 	{
 		const UE::FSdfPath& DuplicatePrim = DuplicatedPrims[Index];
-		Result[ Index ] = DuplicatePrim.GetString();
+		Result[Index] = DuplicatePrim.GetString();
 	}
-#endif // USE_USD_SDK
+#endif	  // USE_USD_SDK
 
 	return Result;
 }
@@ -764,7 +775,7 @@ bool UUsdConversionBlueprintLibrary::SetUsdAssetUserData(UObject* Object, UUsdAs
 	return UsdUtils::SetAssetUserData(Object, AssetUserData);
 }
 
-namespace UE::UsdConversionBlueprintLibraryImpl::Private
+namespace UE::UsdConversionBlueprintLibrary::Private
 {
 	FEditPropertyChain& GetMetadataPropertyChain()
 	{
@@ -781,7 +792,9 @@ namespace UE::UsdConversionBlueprintLibraryImpl::Private
 
 			Chain.Emplace();
 			Chain->AddHead(StageIdentifierProp);
-			Chain->AddTail(FUsdCombinedPrimMetadata::StaticStruct()->FindPropertyByName(GET_MEMBER_NAME_CHECKED(FUsdCombinedPrimMetadata, PrimPathToMetadata)));
+			Chain->AddTail(
+				FUsdCombinedPrimMetadata::StaticStruct()->FindPropertyByName(GET_MEMBER_NAME_CHECKED(FUsdCombinedPrimMetadata, PrimPathToMetadata))
+			);
 			Chain->AddTail(FUsdPrimMetadata::StaticStruct()->FindPropertyByName(GET_MEMBER_NAME_CHECKED(FUsdPrimMetadata, Metadata)));
 			Chain->AddTail(StringifiedValueProp);
 			Chain->SetActivePropertyNode(StringifiedValueProp);
@@ -796,20 +809,16 @@ namespace UE::UsdConversionBlueprintLibraryImpl::Private
 		FPropertyChangedEvent PostChangeEvent{
 			FUsdMetadataValue::StaticStruct()->FindPropertyByName(GET_MEMBER_NAME_CHECKED(FUsdMetadataValue, StringifiedValue)),
 			ChangeType,
-			{TopLevelObject}
-		};
+			{TopLevelObject}};
 		PostChangeEvent.MemberProperty = UUsdAssetUserData::StaticClass()->FindPropertyByName(
 			GET_MEMBER_NAME_CHECKED(UUsdAssetUserData, StageIdentifierToMetadata)
 		);
 
-		FPropertyChangedChainEvent ChainEvent{
-			GetMetadataPropertyChain(),
-			PostChangeEvent
-		};
+		FPropertyChangedChainEvent ChainEvent{GetMetadataPropertyChain(), PostChangeEvent};
 
 		return ChainEvent;
 	}
-}
+}	 // namespace UE::UsdConversionBlueprintLibrary::Private
 
 bool UUsdConversionBlueprintLibrary::SetMetadataField(
 	UUsdAssetUserData* AssetUserData,
@@ -821,7 +830,7 @@ bool UUsdConversionBlueprintLibrary::SetMetadataField(
 	bool bTriggerPropertyChangeEvents
 )
 {
-	using namespace UE::UsdConversionBlueprintLibraryImpl::Private;
+	using namespace UE::UsdConversionBlueprintLibrary::Private;
 
 	if (!AssetUserData)
 	{
@@ -830,19 +839,26 @@ bool UUsdConversionBlueprintLibrary::SetMetadataField(
 
 	if (StageIdentifier.IsEmpty() && AssetUserData->StageIdentifierToMetadata.Num() != 1)
 	{
-		UE_LOG(LogUsd, Warning, TEXT("Failed to set metadata field '%s' on AssetUserData '%s': Please provide a valid value for the StageIdentifier parameter"),
+		UE_LOG(
+			LogUsd,
+			Warning,
+			TEXT("Failed to set metadata field '%s' on AssetUserData '%s': Please provide a valid value for the StageIdentifier parameter"),
 			*AssetUserData->GetPathName(),
 			*Key
 		);
 		return false;
 	}
 
-	const FString& StageIdentifierToUse = !StageIdentifier.IsEmpty() ? StageIdentifier : AssetUserData->StageIdentifierToMetadata.CreateIterator()->Key;
+	const FString& StageIdentifierToUse = !StageIdentifier.IsEmpty() ? StageIdentifier
+																	 : AssetUserData->StageIdentifierToMetadata.CreateIterator()->Key;
 	FUsdCombinedPrimMetadata& CombinedPrimMetadata = AssetUserData->StageIdentifierToMetadata.FindOrAdd(StageIdentifierToUse);
 
 	if (PrimPath.IsEmpty() && CombinedPrimMetadata.PrimPathToMetadata.Num() != 1)
 	{
-		UE_LOG(LogUsd, Warning, TEXT("Failed to set metadata field '%s' on AssetUserData '%s': Please provide a valid value for the PrimPath parameter"),
+		UE_LOG(
+			LogUsd,
+			Warning,
+			TEXT("Failed to set metadata field '%s' on AssetUserData '%s': Please provide a valid value for the PrimPath parameter"),
 			*AssetUserData->GetPathName(),
 			*Key
 		);
@@ -878,7 +894,7 @@ bool UUsdConversionBlueprintLibrary::ClearMetadataField(
 	bool bTriggerPropertyChangeEvents
 )
 {
-	using namespace UE::UsdConversionBlueprintLibraryImpl::Private;
+	using namespace UE::UsdConversionBlueprintLibrary::Private;
 
 	if (!AssetUserData)
 	{
@@ -887,19 +903,26 @@ bool UUsdConversionBlueprintLibrary::ClearMetadataField(
 
 	if (StageIdentifier.IsEmpty() && AssetUserData->StageIdentifierToMetadata.Num() != 1)
 	{
-		UE_LOG(LogUsd, Warning, TEXT("Failed to clear metadata field '%s' on AssetUserData '%s': Please provide a valid value for the StageIdentifier parameter"),
+		UE_LOG(
+			LogUsd,
+			Warning,
+			TEXT("Failed to clear metadata field '%s' on AssetUserData '%s': Please provide a valid value for the StageIdentifier parameter"),
 			*AssetUserData->GetPathName(),
 			*Key
 		);
 		return false;
 	}
 
-	const FString& StageIdentifierToUse = !StageIdentifier.IsEmpty() ? StageIdentifier : AssetUserData->StageIdentifierToMetadata.CreateIterator()->Key;
+	const FString& StageIdentifierToUse = !StageIdentifier.IsEmpty() ? StageIdentifier
+																	 : AssetUserData->StageIdentifierToMetadata.CreateIterator()->Key;
 	FUsdCombinedPrimMetadata& CombinedPrimMetadata = AssetUserData->StageIdentifierToMetadata.FindOrAdd(StageIdentifierToUse);
 
 	if (PrimPath.IsEmpty() && CombinedPrimMetadata.PrimPathToMetadata.Num() != 1)
 	{
-		UE_LOG(LogUsd, Warning, TEXT("Failed to clear metadata field '%s' on AssetUserData '%s': Please provide a valid value for the PrimPath parameter"),
+		UE_LOG(
+			LogUsd,
+			Warning,
+			TEXT("Failed to clear metadata field '%s' on AssetUserData '%s': Please provide a valid value for the PrimPath parameter"),
 			*AssetUserData->GetPathName(),
 			*Key
 		);
@@ -939,14 +962,18 @@ bool UUsdConversionBlueprintLibrary::HasMetadataField(
 
 	if (StageIdentifier.IsEmpty() && AssetUserData->StageIdentifierToMetadata.Num() != 1)
 	{
-		UE_LOG(LogUsd, Warning, TEXT("Failed to check for metadata field '%s' on AssetUserData '%s': Please provide a valid value for the StageIdentifier parameter"),
+		UE_LOG(
+			LogUsd,
+			Warning,
+			TEXT("Failed to check for metadata field '%s' on AssetUserData '%s': Please provide a valid value for the StageIdentifier parameter"),
 			*AssetUserData->GetPathName(),
 			*Key
 		);
 		return false;
 	}
 
-	const FString& StageIdentifierToUse = !StageIdentifier.IsEmpty() ? StageIdentifier : AssetUserData->StageIdentifierToMetadata.CreateIterator()->Key;
+	const FString& StageIdentifierToUse = !StageIdentifier.IsEmpty() ? StageIdentifier
+																	 : AssetUserData->StageIdentifierToMetadata.CreateIterator()->Key;
 	const FUsdCombinedPrimMetadata* CombinedPrimMetadata = AssetUserData->StageIdentifierToMetadata.Find(StageIdentifierToUse);
 	if (!CombinedPrimMetadata)
 	{
@@ -955,7 +982,10 @@ bool UUsdConversionBlueprintLibrary::HasMetadataField(
 
 	if (PrimPath.IsEmpty() && CombinedPrimMetadata->PrimPathToMetadata.Num() != 1)
 	{
-		UE_LOG(LogUsd, Warning, TEXT("Failed to clear metadata field '%s' on AssetUserData '%s': Please provide a valid value for the PrimPath parameter"),
+		UE_LOG(
+			LogUsd,
+			Warning,
+			TEXT("Failed to clear metadata field '%s' on AssetUserData '%s': Please provide a valid value for the PrimPath parameter"),
 			*AssetUserData->GetPathName(),
 			*Key
 		);
@@ -986,14 +1016,18 @@ FUsdMetadataValue UUsdConversionBlueprintLibrary::GetMetadataField(
 
 	if (StageIdentifier.IsEmpty() && AssetUserData->StageIdentifierToMetadata.Num() != 1)
 	{
-		UE_LOG(LogUsd, Warning, TEXT("Failed to get metadata field '%s' on AssetUserData '%s': Please provide a valid value for the StageIdentifier parameter"),
+		UE_LOG(
+			LogUsd,
+			Warning,
+			TEXT("Failed to get metadata field '%s' on AssetUserData '%s': Please provide a valid value for the StageIdentifier parameter"),
 			*AssetUserData->GetPathName(),
 			*Key
 		);
 		return {};
 	}
 
-	const FString& StageIdentifierToUse = !StageIdentifier.IsEmpty() ? StageIdentifier : AssetUserData->StageIdentifierToMetadata.CreateIterator()->Key;
+	const FString& StageIdentifierToUse = !StageIdentifier.IsEmpty() ? StageIdentifier
+																	 : AssetUserData->StageIdentifierToMetadata.CreateIterator()->Key;
 	const FUsdCombinedPrimMetadata* CombinedPrimMetadata = AssetUserData->StageIdentifierToMetadata.Find(StageIdentifierToUse);
 	if (!CombinedPrimMetadata)
 	{
@@ -1002,7 +1036,10 @@ FUsdMetadataValue UUsdConversionBlueprintLibrary::GetMetadataField(
 
 	if (PrimPath.IsEmpty() && CombinedPrimMetadata->PrimPathToMetadata.Num() != 1)
 	{
-		UE_LOG(LogUsd, Warning, TEXT("Failed to get metadata field '%s' on AssetUserData '%s': Please provide a valid value for the PrimPath parameter"),
+		UE_LOG(
+			LogUsd,
+			Warning,
+			TEXT("Failed to get metadata field '%s' on AssetUserData '%s': Please provide a valid value for the PrimPath parameter"),
 			*AssetUserData->GetPathName(),
 			*Key
 		);
@@ -1021,7 +1058,10 @@ FUsdMetadataValue UUsdConversionBlueprintLibrary::GetMetadataField(
 		return *FoundValue;
 	}
 
-	UE_LOG(LogUsd, Warning, TEXT("Failed to find a metadata entry in '%s' with stage identifier '%s', prim path '%s' and key '%s'"),
+	UE_LOG(
+		LogUsd,
+		Warning,
+		TEXT("Failed to find a metadata entry in '%s' with stage identifier '%s', prim path '%s' and key '%s'"),
 		*AssetUserData->GetPathName(),
 		*StageIdentifier,
 		*PrimPath,
@@ -1399,9 +1439,7 @@ int32 UUsdConversionBlueprintLibrary::UnstringifyAsInt(const FString& String)
 int32 UUsdConversionBlueprintLibrary::UnstringifyAsUInt(const FString& String)
 {
 	// Have to cast as there are no unsigned types higher than uint8 on blueprint...
-	return static_cast<int32>(
-		UE::UsdConversionLibrary::Private::UnstringifyOrDefault<uint32, UsdUtils::UnstringifyAsUInt>(String, TEXT("uint"))
-	);
+	return static_cast<int32>(UE::UsdConversionLibrary::Private::UnstringifyOrDefault<uint32, UsdUtils::UnstringifyAsUInt>(String, TEXT("uint")));
 }
 
 int64 UUsdConversionBlueprintLibrary::UnstringifyAsInt64(const FString& String)
@@ -1412,9 +1450,7 @@ int64 UUsdConversionBlueprintLibrary::UnstringifyAsInt64(const FString& String)
 int64 UUsdConversionBlueprintLibrary::UnstringifyAsUInt64(const FString& String)
 {
 	// Have to cast as there are no unsigned types higher than uint8 on blueprint...
-	return static_cast<int64>(
-		UE::UsdConversionLibrary::Private::UnstringifyOrDefault<uint64, UsdUtils::UnstringifyAsUInt64>(String, TEXT("uint64"))
-	);
+	return static_cast<int64>(UE::UsdConversionLibrary::Private::UnstringifyOrDefault<uint64, UsdUtils::UnstringifyAsUInt64>(String, TEXT("uint64")));
 }
 
 float UUsdConversionBlueprintLibrary::UnstringifyAsHalf(const FString& String)
@@ -1631,17 +1667,26 @@ TArray<FString> UUsdConversionBlueprintLibrary::UnstringifyAsAssetPathArray(cons
 
 TArray<FString> UUsdConversionBlueprintLibrary::UnstringifyAsListOpTokens(const FString& String)
 {
-	return UE::UsdConversionLibrary::Private::UnstringifyOrDefault<TArray<FString>, UsdUtils::UnstringifyAsListOpTokens>(String, TEXT("SdfListOp<Token>"));
+	return UE::UsdConversionLibrary::Private::UnstringifyOrDefault<TArray<FString>, UsdUtils::UnstringifyAsListOpTokens>(
+		String,
+		TEXT("SdfListOp<Token>")
+	);
 }
 
 TArray<FMatrix2D> UUsdConversionBlueprintLibrary::UnstringifyAsMatrix2dArray(const FString& String)
 {
-	return UE::UsdConversionLibrary::Private::UnstringifyOrDefault<TArray<FMatrix2D>, UsdUtils::UnstringifyAsMatrix2dArray>(String, TEXT("matrix2d[]"));
+	return UE::UsdConversionLibrary::Private::UnstringifyOrDefault<TArray<FMatrix2D>, UsdUtils::UnstringifyAsMatrix2dArray>(
+		String,
+		TEXT("matrix2d[]")
+	);
 }
 
 TArray<FMatrix3D> UUsdConversionBlueprintLibrary::UnstringifyAsMatrix3dArray(const FString& String)
 {
-	return UE::UsdConversionLibrary::Private::UnstringifyOrDefault<TArray<FMatrix3D>, UsdUtils::UnstringifyAsMatrix3dArray>(String, TEXT("matrix3d[]"));
+	return UE::UsdConversionLibrary::Private::UnstringifyOrDefault<TArray<FMatrix3D>, UsdUtils::UnstringifyAsMatrix3dArray>(
+		String,
+		TEXT("matrix3d[]")
+	);
 }
 
 TArray<FMatrix> UUsdConversionBlueprintLibrary::UnstringifyAsMatrix4dArray(const FString& String)
diff --git a/Engine/Plugins/Importers/USDImporter/Source/USDExporter/Private/USDExporterModule.cpp b/Engine/Plugins/Importers/USDImporter/Source/USDExporter/Private/USDExporterModule.cpp
--- a/Engine/Plugins/Importers/USDImporter/Source/USDExporter/Private/USDExporterModule.cpp
+++ b/Engine/Plugins/Importers/USDImporter/Source/USDExporter/Private/USDExporterModule.cpp
@@ -6,7 +6,6 @@
 #include "USDAssetOptions.h"
 #include "USDLog.h"
 #include "USDMemory.h"
-#include "USDMetadataExportOptions.h"
 
 #include "UsdWrappers/SdfLayer.h"
 
@@ -16,7 +15,6 @@
 #include "Modules/ModuleManager.h"
 #include "PropertyEditorModule.h"
 #include "Selection.h"
-#include "UObject/ObjectSaveContext.h"
 #include "Widgets/Notifications/SNotificationList.h"
 
 #define LOCTEXT_NAMESPACE "USDExporterModule"
@@ -47,38 +45,59 @@ public:
 	{
 		LLM_SCOPE_BYTAG(Usd);
 
-		FPropertyEditorModule& PropertyModule = FModuleManager::LoadModuleChecked< FPropertyEditorModule >( TEXT( "PropertyEditor" ) );
+		FPropertyEditorModule& PropertyModule = FModuleManager::LoadModuleChecked<FPropertyEditorModule>(TEXT("PropertyEditor"));
 
 		// Hide the "bExportComponentMetadata" property if we're just exporting assets
-		PropertyModule.RegisterCustomClassLayout(TEXT("AnimSequenceExporterUSDOptions"), FOnGetDetailCustomizationInstance::CreateStatic(&FHideExportComponentMetadataCustomization::MakeInstance));
-		PropertyModule.RegisterCustomClassLayout(TEXT("MaterialExporterUSDOptions"), FOnGetDetailCustomizationInstance::CreateStatic(&FHideExportComponentMetadataCustomization::MakeInstance));
-		PropertyModule.RegisterCustomClassLayout(TEXT("SkeletalMeshExporterUSDOptions"), FOnGetDetailCustomizationInstance::CreateStatic(&FHideExportComponentMetadataCustomization::MakeInstance));
-		PropertyModule.RegisterCustomClassLayout(TEXT("StaticMeshExporterUSDOptions"), FOnGetDetailCustomizationInstance::CreateStatic(&FHideExportComponentMetadataCustomization::MakeInstance));
-		PropertyModule.RegisterCustomClassLayout(TEXT("GeometryCacheExporterUSDOptions"), FOnGetDetailCustomizationInstance::CreateStatic(&FHideExportComponentMetadataCustomization::MakeInstance));
+		PropertyModule.RegisterCustomClassLayout(
+			TEXT("AnimSequenceExporterUSDOptions"),
+			FOnGetDetailCustomizationInstance::CreateStatic(&FHideExportComponentMetadataCustomization::MakeInstance)
+		);
+		PropertyModule.RegisterCustomClassLayout(
+			TEXT("MaterialExporterUSDOptions"),
+			FOnGetDetailCustomizationInstance::CreateStatic(&FHideExportComponentMetadataCustomization::MakeInstance)
+		);
+		PropertyModule.RegisterCustomClassLayout(
+			TEXT("SkeletalMeshExporterUSDOptions"),
+			FOnGetDetailCustomizationInstance::CreateStatic(&FHideExportComponentMetadataCustomization::MakeInstance)
+		);
+		PropertyModule.RegisterCustomClassLayout(
+			TEXT("StaticMeshExporterUSDOptions"),
+			FOnGetDetailCustomizationInstance::CreateStatic(&FHideExportComponentMetadataCustomization::MakeInstance)
+		);
+		PropertyModule.RegisterCustomClassLayout(
+			TEXT("GeometryCacheExporterUSDOptions"),
+			FOnGetDetailCustomizationInstance::CreateStatic(&FHideExportComponentMetadataCustomization::MakeInstance)
+		);
 
 		// We intentionally use the same customization for both of these
-		PropertyModule.RegisterCustomClassLayout( TEXT( "LevelExporterUSDOptions" ), FOnGetDetailCustomizationInstance::CreateStatic( &FLevelExporterUSDOptionsCustomization::MakeInstance ) );
-		PropertyModule.RegisterCustomClassLayout( TEXT( "LevelSequenceExporterUSDOptions" ), FOnGetDetailCustomizationInstance::CreateStatic( &FLevelExporterUSDOptionsCustomization::MakeInstance ) );
+		PropertyModule.RegisterCustomClassLayout(
+			TEXT("LevelExporterUSDOptions"),
+			FOnGetDetailCustomizationInstance::CreateStatic(&FLevelExporterUSDOptionsCustomization::MakeInstance)
+		);
+		PropertyModule.RegisterCustomClassLayout(
+			TEXT("LevelSequenceExporterUSDOptions"),
+			FOnGetDetailCustomizationInstance::CreateStatic(&FLevelExporterUSDOptionsCustomization::MakeInstance)
+		);
 
 		// Modify the static mesh LOD range to have the proper define value as the maximum.
 		// We have to do this the hard way here because we can't use the define within the meta tag itself
-		if ( UScriptStruct* ScriptStruct = FUsdMeshAssetOptions::StaticStruct() )
+		if (UScriptStruct* ScriptStruct = FUsdMeshAssetOptions::StaticStruct())
 		{
-			if ( FProperty* LowestMeshLODProperty = ScriptStruct->FindPropertyByName( GET_MEMBER_NAME_CHECKED( FUsdMeshAssetOptions, LowestMeshLOD ) ) )
+			if (FProperty* LowestMeshLODProperty = ScriptStruct->FindPropertyByName(GET_MEMBER_NAME_CHECKED(FUsdMeshAssetOptions, LowestMeshLOD)))
 			{
-				LowestMeshLODProperty->SetMetaData( TEXT( "ClampMax" ), LexToString( MAX_MESH_LOD_COUNT - 1 ) );
+				LowestMeshLODProperty->SetMetaData(TEXT("ClampMax"), LexToString(MAX_MESH_LOD_COUNT - 1));
 			}
 
-			if ( FProperty* HighestMeshLODProperty = ScriptStruct->FindPropertyByName( GET_MEMBER_NAME_CHECKED( FUsdMeshAssetOptions, HighestMeshLOD ) ) )
+			if (FProperty* HighestMeshLODProperty = ScriptStruct->FindPropertyByName(GET_MEMBER_NAME_CHECKED(FUsdMeshAssetOptions, HighestMeshLOD)))
 			{
-				HighestMeshLODProperty->SetMetaData( TEXT( "ClampMax" ), LexToString( MAX_MESH_LOD_COUNT - 1 ) );
+				HighestMeshLODProperty->SetMetaData(TEXT("ClampMax"), LexToString(MAX_MESH_LOD_COUNT - 1));
 			}
 		}
 	}
 
 	virtual void ShutdownModule() override
 	{
-		if ( FPropertyEditorModule* PropertyModule = FModuleManager::GetModulePtr< FPropertyEditorModule >( TEXT( "PropertyEditor" ) ) )
+		if (FPropertyEditorModule* PropertyModule = FModuleManager::GetModulePtr<FPropertyEditorModule>(TEXT("PropertyEditor")))
 		{
 			PropertyModule->UnregisterCustomClassLayout(TEXT("AnimSequenceExporterUSDOptions"));
 			PropertyModule->UnregisterCustomClassLayout(TEXT("MaterialExporterUSDOptions"));
@@ -86,15 +105,15 @@ public:
 			PropertyModule->UnregisterCustomClassLayout(TEXT("StaticMeshExporterUSDOptions"));
 			PropertyModule->UnregisterCustomClassLayout(TEXT("GeometryCacheExporterUSDOptions"));
 
-			PropertyModule->UnregisterCustomClassLayout( TEXT( "LevelExporterUSDOptions" ) );
-			PropertyModule->UnregisterCustomClassLayout( TEXT( "LevelSequenceExporterUSDOptions" ) );
+			PropertyModule->UnregisterCustomClassLayout(TEXT("LevelExporterUSDOptions"));
+			PropertyModule->UnregisterCustomClassLayout(TEXT("LevelSequenceExporterUSDOptions"));
 		}
 	}
 };
 
-void IUsdExporterModule::HashEditorSelection( FSHA1& HashToUpdate )
+void IUsdExporterModule::HashEditorSelection(FSHA1& HashToUpdate)
 {
-	if ( !GEditor )
+	if (!GEditor)
 	{
 		return;
 	}
@@ -103,51 +122,56 @@ void IUsdExporterModule::HashEditorSelection( FSHA1& HashToUpdate )
 
 	USelection* ComponentSelection = GEditor->GetSelectedComponents();
 	TArray<UActorComponent*> Components;
-	ComponentSelection->GetSelectedObjects( Components );
-	for ( const UActorComponent* Component : Components )
+	ComponentSelection->GetSelectedObjects(Components);
+	for (const UActorComponent* Component : Components)
 	{
-		if ( Component )
+		if (Component)
 		{
-			SelectedPaths.Add( Component->GetFullName() );
+			SelectedPaths.Add(Component->GetFullName());
 		}
 	}
 
 	USelection* ActorSelection = GEditor->GetSelectedActors();
 	TArray<AActor*> Actors;
-	ActorSelection->GetSelectedObjects( Actors );
-	for ( const AActor* Actor : Actors )
+	ActorSelection->GetSelectedObjects(Actors);
+	for (const AActor* Actor : Actors)
 	{
-		if ( Actor )
+		if (Actor)
 		{
-			SelectedPaths.Add( Actor->GetFullName() );
+			SelectedPaths.Add(Actor->GetFullName());
 		}
 	}
 
 	TArray<FString> SelectedPathsArray = SelectedPaths.Array();
 	SelectedPathsArray.Sort();
-	for ( const FString& Path : SelectedPathsArray )
+	for (const FString& Path : SelectedPathsArray)
 	{
-		HashToUpdate.UpdateWithString( *Path, Path.Len() );
+		HashToUpdate.UpdateWithString(*Path, Path.Len());
 	}
 }
 
-bool IUsdExporterModule::CanExportToLayer( const FString& TargetFilePath )
+bool IUsdExporterModule::CanExportToLayer(const FString& TargetFilePath)
 {
-	for ( const UE::FSdfLayerWeak& LoadedLayer : UE::FSdfLayerWeak::GetLoadedLayers() )
+	for (const UE::FSdfLayerWeak& LoadedLayer : UE::FSdfLayerWeak::GetLoadedLayers())
 	{
 		if (FPaths::IsSamePath(LoadedLayer.GetRealPath(), TargetFilePath))
 		{
 			FText ErrorMessage = FText::Format(
-				LOCTEXT( "FailedExportLayerOpenSubText", "Failed to export to the USD layer '{0}' as a layer with this identifier is already open in another USD Stage. Please use a different file path or try closing that other USD stage and exporting again." ),
-				FText::FromString( TargetFilePath ) );
-			UE_LOG( LogUsd, Error, TEXT( "%s" ), *ErrorMessage.ToString() );
-
-			FNotificationInfo ErrorToast( LOCTEXT( "FailedExportLayerOpenText", "USD: Export failure" ) );
+				LOCTEXT(
+					"FailedExportLayerOpenSubText",
+					"Failed to export to the USD layer '{0}' as a layer with this identifier is already open in another USD Stage. Please use a "
+					"different file path or try closing that other USD stage and exporting again."
+				),
+				FText::FromString(TargetFilePath)
+			);
+			UE_LOG(LogUsd, Error, TEXT("%s"), *ErrorMessage.ToString());
+
+			FNotificationInfo ErrorToast(LOCTEXT("FailedExportLayerOpenText", "USD: Export failure"));
 			ErrorToast.ExpireDuration = 10.0f;
 			ErrorToast.bFireAndForget = true;
-			ErrorToast.Image = FCoreStyle::Get().GetBrush( TEXT( "MessageLog.Error" ) );
+			ErrorToast.Image = FCoreStyle::Get().GetBrush(TEXT("MessageLog.Error"));
 			ErrorToast.SubText = ErrorMessage;
-			FSlateNotificationManager::Get().AddNotification( ErrorToast );
+			FSlateNotificationManager::Get().AddNotification(ErrorToast);
 
 			return false;
 		}
@@ -156,6 +180,6 @@ bool IUsdExporterModule::CanExportToLayer( const FString& TargetFilePath )
 	return true;
 }
 
-IMPLEMENT_MODULE_USD( FUsdExporterModule, USDExporter );
+IMPLEMENT_MODULE_USD(FUsdExporterModule, USDExporter);
 
-#undef LOCTEXT_NAMESPACE
\ No newline at end of file
+#undef LOCTEXT_NAMESPACE
diff --git a/Engine/Plugins/Importers/USDImporter/Source/USDExporter/Public/AnimSequenceExporterUSDOptions.h b/Engine/Plugins/Importers/USDImporter/Source/USDExporter/Public/AnimSequenceExporterUSDOptions.h
--- a/Engine/Plugins/Importers/USDImporter/Source/USDExporter/Public/AnimSequenceExporterUSDOptions.h
+++ b/Engine/Plugins/Importers/USDImporter/Source/USDExporter/Public/AnimSequenceExporterUSDOptions.h
@@ -6,8 +6,6 @@
 #include "USDMetadataExportOptions.h"
 #include "USDStageOptions.h"
 
-#include "Engine/EngineTypes.h"
-
 #include "AnimSequenceExporterUSDOptions.generated.h"
 
 struct FAnalyticsEventAttribute;
@@ -15,22 +13,28 @@ struct FAnalyticsEventAttribute;
 /**
  * Options for exporting skeletal mesh animations to USD format.
  */
-UCLASS( Config = Editor, Blueprintable )
+UCLASS(Config = Editor, Blueprintable)
 class USDEXPORTER_API UAnimSequenceExporterUSDOptions : public UObject
 {
 	GENERATED_BODY()
 
 public:
 	/** Export options to use for the layer where the animation is emitted */
-	UPROPERTY( EditAnywhere, config, BlueprintReadWrite, Category = "Stage options", meta = ( ShowOnlyInnerProperties ) )
+	UPROPERTY(EditAnywhere, config, BlueprintReadWrite, Category = "Stage options", meta = (ShowOnlyInnerProperties))
 	FUsdStageOptions StageOptions;
 
 	/** Whether to also export the skeletal mesh data of the preview mesh */
-	UPROPERTY( EditAnywhere, config, BlueprintReadWrite, Category = "Mesh options" )
+	UPROPERTY(EditAnywhere, config, BlueprintReadWrite, Category = "Mesh options")
 	bool bExportPreviewMesh;
 
 	/** Export options to use for the preview mesh, if enabled */
-	UPROPERTY( EditAnywhere, config, BlueprintReadWrite, Category = "Mesh options", meta = ( ShowOnlyInnerProperties, EditCondition = bExportPreviewMesh ) )
+	UPROPERTY(
+		EditAnywhere,
+		config,
+		BlueprintReadWrite,
+		Category = "Mesh options",
+		meta = (ShowOnlyInnerProperties, EditCondition = bExportPreviewMesh)
+	)
 	FUsdMeshAssetOptions PreviewMeshOptions;
 
 	/** Options to use when exporting UsdAssetUserData metadata to USD for all assets (including the preview mesh) */
@@ -41,19 +45,13 @@ public:
 	 * Whether to export any asset (StaticMesh, Material, etc.) even if the existing file already describes the same version of a compatible asset.
 	 * This is only checked when bReplaceIdentical is set on the asset export task. Otherwise we'll never overwrite files.
 	 */
-	UPROPERTY( EditAnywhere, config, BlueprintReadWrite, Category = "Collision", meta = ( DisplayName = "Re-export Identical Assets" ) )
+	UPROPERTY(EditAnywhere, config, BlueprintReadWrite, Category = "Collision", meta = (DisplayName = "Re-export Identical Assets"))
 	bool bReExportIdenticalAssets = false;
 };
 
 namespace UsdUtils
 {
-	USDEXPORTER_API void AddAnalyticsAttributes(
-		const UAnimSequenceExporterUSDOptions& Options,
-		TArray< FAnalyticsEventAttribute >& InOutAttributes
-	);
-
-	USDEXPORTER_API void HashForAnimSequenceExport(
-		const UAnimSequenceExporterUSDOptions& Options,
-		FSHA1& HashToUpdate
-	);
+	USDEXPORTER_API void AddAnalyticsAttributes(const UAnimSequenceExporterUSDOptions& Options, TArray<FAnalyticsEventAttribute>& InOutAttributes);
+
+	USDEXPORTER_API void HashForAnimSequenceExport(const UAnimSequenceExporterUSDOptions& Options, FSHA1& HashToUpdate);
 }
diff --git a/Engine/Plugins/Importers/USDImporter/Source/USDExporter/Public/LevelExporterUSDOptions.h b/Engine/Plugins/Importers/USDImporter/Source/USDExporter/Public/LevelExporterUSDOptions.h
--- a/Engine/Plugins/Importers/USDImporter/Source/USDExporter/Public/LevelExporterUSDOptions.h
+++ b/Engine/Plugins/Importers/USDImporter/Source/USDExporter/Public/LevelExporterUSDOptions.h
@@ -2,6 +2,7 @@
 
 #pragma once
 
+#include "CoreMinimal.h"
 #include "USDAssetOptions.h"
 #include "USDMetadataExportOptions.h"
 #include "USDStageOptions.h"
@@ -12,21 +13,22 @@
 
 struct FAnalyticsEventAttribute;
 
-USTRUCT( BlueprintType )
+USTRUCT(BlueprintType)
 struct USDEXPORTER_API FLevelExporterUSDOptionsInner
 {
 	GENERATED_BODY()
 
 	/** Whether to export only the selected actors, and assets used by them */
-    UPROPERTY( EditAnywhere, config, BlueprintReadWrite, Category = "Export options" )
-    bool bSelectionOnly = false;
+	UPROPERTY(EditAnywhere, config, BlueprintReadWrite, Category = "Export options")
+	bool bSelectionOnly = false;
 
 	/** Whether to use UE actor folders as empty prims */
-    UPROPERTY( EditAnywhere, config, BlueprintReadWrite, Category = "Export options" )
-    bool bExportActorFolders = false;
+	UPROPERTY(EditAnywhere, config, BlueprintReadWrite, Category = "Export options")
+	bool bExportActorFolders = false;
 
-	/** If true, and if we have a level sequence animating the level during export, it will revert any actor or component to its unanimated state before writing to USD */
-	UPROPERTY( EditAnywhere, config, BlueprintReadWrite, Category = "Export options" )
+	/** If true, and if we have a level sequence animating the level during export, it will revert any actor or component to its unanimated state
+	 * before writing to USD */
+	UPROPERTY(EditAnywhere, config, BlueprintReadWrite, Category = "Export options")
 	bool bIgnoreSequencerAnimations = false;
 
 	/**
@@ -34,15 +36,15 @@ struct USDEXPORTER_API FLevelExporterUSDOptionsInner
 	 * Enable this to instead export the foliage instances to the same layer as the foliage actor they belong to.
 	 * This is useful if those foliage instances were placed with the "Place In Current Level" option.
 	 */
-	UPROPERTY( EditAnywhere, config, BlueprintReadWrite, Category = "Export options" )
+	UPROPERTY(EditAnywhere, config, BlueprintReadWrite, Category = "Export options")
 	bool bExportFoliageOnActorsLayer = false;
 
 	/** Where to place all the generated asset files */
-	UPROPERTY( EditAnywhere, config, BlueprintReadWrite, Category = "Asset options" )
+	UPROPERTY(EditAnywhere, config, BlueprintReadWrite, Category = "Asset options")
 	FDirectoryPath AssetFolder;
 
 	/** Options to use for all exported assets when appropriate (e.g. static and skeletal meshes, materials, etc.) */
-	UPROPERTY( EditAnywhere, config, BlueprintReadWrite, Category = "Asset options", meta = ( ShowOnlyInnerProperties ) )
+	UPROPERTY(EditAnywhere, config, BlueprintReadWrite, Category = "Asset options", meta = (ShowOnlyInnerProperties))
 	FUsdMeshAssetOptions AssetOptions;
 
 	/** Options to use when exporting UsdAssetUserData metadata to USD for all assets (including the level itself) */
@@ -50,66 +52,67 @@ struct USDEXPORTER_API FLevelExporterUSDOptionsInner
 	FUsdMetadataExportOptions MetadataOptions;
 
 	/** Lowest of the LOD indices to export landscapes with (use 0 for full resolution) */
-	UPROPERTY( EditAnywhere, config, BlueprintReadWrite, Category = "Landscape options", meta = ( ClampMin = "0" ) )
+	UPROPERTY(EditAnywhere, config, BlueprintReadWrite, Category = "Landscape options", meta = (ClampMin = "0"))
 	int32 LowestLandscapeLOD = 0;
 
 	/**
 	 * Highest of the LOD indices to export landscapes with. Each value above 0 halves resolution.
 	 * The max value depends on the number of components and sections per component of each landscape, and may be clamped.
 	 */
-	UPROPERTY( EditAnywhere, config, BlueprintReadWrite, Category = "Landscape options", meta = ( ClampMin = "0" ) )
+	UPROPERTY(EditAnywhere, config, BlueprintReadWrite, Category = "Landscape options", meta = (ClampMin = "0"))
 	int32 HighestLandscapeLOD = 0;
 
 	/** Resolution to use when baking landscape materials into textures  */
-	UPROPERTY( EditAnywhere, config, BlueprintReadWrite, Category = "Landscape options", meta = ( ClampMin = "1" ) )
-	FIntPoint LandscapeBakeResolution = FIntPoint( 1024, 1024 );
+	UPROPERTY(EditAnywhere, config, BlueprintReadWrite, Category = "Landscape options", meta = (ClampMin = "1"))
+	FIntPoint LandscapeBakeResolution = FIntPoint(1024, 1024);
 
-	/** If true, will export sub-levels as separate layers (referenced as sublayers). If false, will collapse all sub-levels in a single exported root layer */
-    UPROPERTY( EditAnywhere, config, BlueprintReadWrite, Category = "Sublayers", meta = ( DisplayName = "Export Sublevels as Layers" ) )
-    bool bExportSublayers = false;
+	/** If true, will export sub-levels as separate layers (referenced as sublayers). If false, will collapse all sub-levels in a single exported root
+	 * layer */
+	UPROPERTY(EditAnywhere, config, BlueprintReadWrite, Category = "Sublayers", meta = (DisplayName = "Export Sublevels as Layers"))
+	bool bExportSublayers = false;
 
 	/** Names of levels that should be ignored when collecting actors to export (e.g. "Persistent Level", "Level1", "MySubLevel", etc.) */
-    UPROPERTY( EditAnywhere, config, BlueprintReadWrite, Category = "Sublayers" )
-    TSet<FString> LevelsToIgnore;
+	UPROPERTY(EditAnywhere, config, BlueprintReadWrite, Category = "Sublayers")
+	TSet<FString> LevelsToIgnore;
 };
 
 /**
  * Options for exporting levels to USD format.
  */
-UCLASS( Config = Editor, Blueprintable, HideCategories=Hidden )
+UCLASS(Config = Editor, Blueprintable, HideCategories = Hidden)
 class USDEXPORTER_API ULevelExporterUSDOptions : public UObject
 {
 	GENERATED_BODY()
 
 public:
 	/** Basic options about the stage to export */
-	UPROPERTY( EditAnywhere, config, BlueprintReadWrite, Category = "Stage options", meta = ( ShowOnlyInnerProperties ) )
+	UPROPERTY(EditAnywhere, config, BlueprintReadWrite, Category = "Stage options", meta = (ShowOnlyInnerProperties))
 	FUsdStageOptions StageOptions;
 
 	/** StartTimeCode to be used for all exported layers */
-    UPROPERTY( EditAnywhere, config, BlueprintReadWrite, Category = "Stage options" )
-    float StartTimeCode = 0.0f;
+	UPROPERTY(EditAnywhere, config, BlueprintReadWrite, Category = "Stage options")
+	float StartTimeCode = 0.0f;
 
 	/** EndTimeCode to be used for all exported layers */
-    UPROPERTY( EditAnywhere, config, BlueprintReadWrite, Category = "Stage options" )
-    float EndTimeCode = 0.0f;
+	UPROPERTY(EditAnywhere, config, BlueprintReadWrite, Category = "Stage options")
+	float EndTimeCode = 0.0f;
 
 	/** Inner struct that actually contains most of the export options */
-    UPROPERTY( EditAnywhere, config, BlueprintReadWrite, Category = "Export settings", meta = ( ShowOnlyInnerProperties ) )
-    FLevelExporterUSDOptionsInner Inner;
+	UPROPERTY(EditAnywhere, config, BlueprintReadWrite, Category = "Export settings", meta = (ShowOnlyInnerProperties))
+	FLevelExporterUSDOptionsInner Inner;
 
 	/**
 	 * Whether to export levels and LevelSequences even if the existing files already describe the same versions of compatible assets.
 	 * This is only checked when bReplaceIdentical is set on the asset export task. Otherwise we'll never overwrite files.
 	 */
-	UPROPERTY( EditAnywhere, config, BlueprintReadWrite, Category = "Collision", meta = ( DisplayName = "Re-export Identical Levels and Sequences" ) )
+	UPROPERTY(EditAnywhere, config, BlueprintReadWrite, Category = "Collision", meta = (DisplayName = "Re-export Identical Levels and Sequences"))
 	bool bReExportIdenticalLevelsAndSequences = false;
 
 	/**
 	 * Whether to export any asset (StaticMesh, Material, etc.) even if the existing file already describes the same version of a compatible asset.
 	 * This is only checked when bReplaceIdentical is set on the asset export task. Otherwise we'll never overwrite files.
 	 */
-	UPROPERTY( EditAnywhere, config, BlueprintReadWrite, Category = "Collision", meta = ( DisplayName = "Re-export Identical Assets" ) )
+	UPROPERTY(EditAnywhere, config, BlueprintReadWrite, Category = "Collision", meta = (DisplayName = "Re-export Identical Assets"))
 	bool bReExportIdenticalAssets = false;
 
 public:
@@ -117,7 +120,7 @@ public:
 	// the Python exporter, that does the actual level exporting.
 	// This is weak because we often use the CDO of this class directly, and we never want to
 	// permanently hold on to a particular export task
-	UPROPERTY( VisibleAnywhere, BlueprintReadOnly, Category = Hidden )
+	UPROPERTY(VisibleAnywhere, BlueprintReadOnly, Category = Hidden)
 	TWeakObjectPtr<UAssetExportTask> CurrentTask;
 
 private:
@@ -127,23 +130,11 @@ private:
 
 namespace UsdUtils
 {
-	USDEXPORTER_API void AddAnalyticsAttributes(
-		const FLevelExporterUSDOptionsInner& Options,
-		TArray< FAnalyticsEventAttribute >& InOutAttributes
-	);
-
-	USDEXPORTER_API void AddAnalyticsAttributes(
-		const ULevelExporterUSDOptions& Options,
-		TArray< FAnalyticsEventAttribute >& InOutAttributes
-	);
-
-	USDEXPORTER_API void HashForLevelExport(
-		const FLevelExporterUSDOptionsInner& Options,
-		FSHA1& HashToUpdate
-	);
-
-	USDEXPORTER_API void HashForLevelExport(
-		const ULevelExporterUSDOptions& Options,
-		FSHA1& HashToUpdate
-	);
-}
\ No newline at end of file
+	USDEXPORTER_API void AddAnalyticsAttributes(const FLevelExporterUSDOptionsInner& Options, TArray<FAnalyticsEventAttribute>& InOutAttributes);
+
+	USDEXPORTER_API void AddAnalyticsAttributes(const ULevelExporterUSDOptions& Options, TArray<FAnalyticsEventAttribute>& InOutAttributes);
+
+	USDEXPORTER_API void HashForLevelExport(const FLevelExporterUSDOptionsInner& Options, FSHA1& HashToUpdate);
+
+	USDEXPORTER_API void HashForLevelExport(const ULevelExporterUSDOptions& Options, FSHA1& HashToUpdate);
+}
diff --git a/Engine/Plugins/Importers/USDImporter/Source/USDExporter/Public/USDAssetOptions.h b/Engine/Plugins/Importers/USDImporter/Source/USDExporter/Public/USDAssetOptions.h
--- a/Engine/Plugins/Importers/USDImporter/Source/USDExporter/Public/USDAssetOptions.h
+++ b/Engine/Plugins/Importers/USDImporter/Source/USDExporter/Public/USDAssetOptions.h
@@ -10,88 +10,95 @@
 
 struct FAnalyticsEventAttribute;
 
-USTRUCT( BlueprintType )
+USTRUCT(BlueprintType)
 struct USDEXPORTER_API FUsdMaterialBakingOptions
 {
 	GENERATED_BODY()
 
 	/** Properties which are supposed to be baked out for the material */
-	UPROPERTY( EditAnywhere, config, BlueprintReadWrite, Category = "Material baking options" )
-	TArray<FPropertyEntry> Properties = { MP_BaseColor, MP_Metallic, MP_Specular, MP_Roughness, MP_Anisotropy, MP_EmissiveColor, MP_Opacity, MP_OpacityMask, MP_Normal, MP_Tangent, MP_SubsurfaceColor, MP_AmbientOcclusion };
+	UPROPERTY(EditAnywhere, config, BlueprintReadWrite, Category = "Material baking options")
+	TArray<FPropertyEntry> Properties = {
+		MP_BaseColor,
+		MP_Metallic,
+		MP_Specular,
+		MP_Roughness,
+		MP_Anisotropy,
+		MP_EmissiveColor,
+		MP_Opacity,
+		MP_OpacityMask,
+		MP_Normal,
+		MP_Tangent,
+		MP_SubsurfaceColor,
+		MP_AmbientOcclusion};
 
 	/** Size of the baked texture for all properties that don't have a CustomSize set */
-	UPROPERTY( EditAnywhere, config, BlueprintReadWrite, Category = "Material baking options", meta = ( ClampMin = "1", UIMin = "1" ) )
-	FIntPoint DefaultTextureSize = FIntPoint{ 128, 128 };
+	UPROPERTY(EditAnywhere, config, BlueprintReadWrite, Category = "Material baking options", meta = (ClampMin = "1", UIMin = "1"))
+	FIntPoint DefaultTextureSize = FIntPoint{128, 128};
 
-	/** When this is true and a baked texture contains a single flat color we will write out that color value directly on the USD layer and skip generating a texture */
-	UPROPERTY( EditAnywhere, BlueprintReadWrite, Category = "Material baking options" )
+	/** When this is true and a baked texture contains a single flat color we will write out that color value directly on the USD layer and skip
+	 * generating a texture */
+	UPROPERTY(EditAnywhere, BlueprintReadWrite, Category = "Material baking options")
 	bool bConstantColorAsSingleValue = true;
 
 	/** Where baked textures are placed. Intentionally not a config as it's heavily dependent on where the stage is */
-	UPROPERTY( EditAnywhere, BlueprintReadWrite, Category = "Material baking options" )
+	UPROPERTY(EditAnywhere, BlueprintReadWrite, Category = "Material baking options")
 	FDirectoryPath TexturesDir;
 };
 
-USTRUCT( BlueprintType )
+USTRUCT(BlueprintType)
 struct USDEXPORTER_API FUsdMeshAssetOptions
 {
 	GENERATED_BODY()
 
 	/** If true, the mesh data is exported to yet another "payload" file, and referenced via a payload composition arc */
-	UPROPERTY( EditAnywhere, config, BlueprintReadWrite, Category = "Mesh options" )
+	UPROPERTY(EditAnywhere, config, BlueprintReadWrite, Category = "Mesh options")
 	bool bUsePayload = false;
 
 	/** USD format to use for exported payload files */
-	UPROPERTY( EditAnywhere, config, BlueprintReadWrite, Category = "Mesh options", meta = ( EditCondition = bUsePayload, GetOptions = "USDExporter.LevelExporterUSDOptions.GetUsdExtensions" ) )
+	UPROPERTY(
+		EditAnywhere,
+		config,
+		BlueprintReadWrite,
+		Category = "Mesh options",
+		meta = (EditCondition = bUsePayload, GetOptions = "USDExporter.LevelExporterUSDOptions.GetUsdExtensions")
+	)
 	FString PayloadFormat;
 
 	/** Whether to bake the mesh's assigned material and export these as separate UsdPreviewSurface assets */
-	UPROPERTY( EditAnywhere, config, BlueprintReadWrite, Category = "Material options" )
+	UPROPERTY(EditAnywhere, config, BlueprintReadWrite, Category = "Material options")
 	bool bBakeMaterials = false;
 
 	/** Whether to remove the 'unrealMaterial' attribute after binding the corresponding baked material */
-	UE_DEPRECATED( 5.2, "This option is now deprecated as UE material assignments are only visible in the 'unreal' render context anyway" )
-	UPROPERTY( config, meta = ( EditCondition = bBakeMaterials ) )
+	UE_DEPRECATED(5.2, "This option is now deprecated as UE material assignments are only visible in the 'unreal' render context anyway")
+	UPROPERTY(config, meta = (EditCondition = bBakeMaterials))
 	bool bRemoveUnrealMaterials = false;
 
-	UPROPERTY( EditAnywhere, config, BlueprintReadWrite, Category = "Material options", meta = ( EditCondition = bBakeMaterials ) )
+	UPROPERTY(EditAnywhere, config, BlueprintReadWrite, Category = "Material options", meta = (EditCondition = bBakeMaterials))
 	FUsdMaterialBakingOptions MaterialBakingOptions;
 
 	/** Lowest of the LOD indices to export static and skeletal meshes with (use 0 for full detail) */
-	UPROPERTY( EditAnywhere, config, BlueprintReadWrite, Category = "Mesh options", meta = ( ClampMin = "0" ) )
+	UPROPERTY(EditAnywhere, config, BlueprintReadWrite, Category = "Mesh options", meta = (ClampMin = "0"))
 	int32 LowestMeshLOD = 0;
 
 	/** Highest of the LOD indices to export static and skeletal meshes with */
-	UPROPERTY( EditAnywhere, config, BlueprintReadWrite, Category = "Mesh options", meta = ( ClampMin = "0" ) )
+	UPROPERTY(EditAnywhere, config, BlueprintReadWrite, Category = "Mesh options", meta = (ClampMin = "0"))
 	int32 HighestMeshLOD = MAX_MESH_LOD_COUNT - 1;
 
 	// We need this as clang would otherwise emit some warnings when generating these due to the usages of bRemoveUnrealMaterials
 	PRAGMA_DISABLE_DEPRECATION_WARNINGS
 	FUsdMeshAssetOptions() = default;
-	FUsdMeshAssetOptions( const FUsdMeshAssetOptions& ) = default;
-	FUsdMeshAssetOptions& operator=( const FUsdMeshAssetOptions& ) = default;
+	FUsdMeshAssetOptions(const FUsdMeshAssetOptions&) = default;
+	FUsdMeshAssetOptions& operator=(const FUsdMeshAssetOptions&) = default;
 	PRAGMA_ENABLE_DEPRECATION_WARNINGS
 };
 
 namespace UsdUtils
 {
-	USDEXPORTER_API void AddAnalyticsAttributes(
-		const FUsdMaterialBakingOptions& Options,
-		TArray< FAnalyticsEventAttribute >& InOutAttributes
-	);
-
-	USDEXPORTER_API void AddAnalyticsAttributes(
-		const FUsdMeshAssetOptions& Options,
-		TArray< FAnalyticsEventAttribute >& InOutAttributes
-	);
-
-	USDEXPORTER_API void HashForMaterialExport(
-		const FUsdMaterialBakingOptions& Options,
-		FSHA1& HashToUpdate
-	);
-
-	USDEXPORTER_API void HashForMeshExport(
-		const FUsdMeshAssetOptions& Options,
-		FSHA1& HashToUpdate
-	);
-}
\ No newline at end of file
+	USDEXPORTER_API void AddAnalyticsAttributes(const FUsdMaterialBakingOptions& Options, TArray<FAnalyticsEventAttribute>& InOutAttributes);
+
+	USDEXPORTER_API void AddAnalyticsAttributes(const FUsdMeshAssetOptions& Options, TArray<FAnalyticsEventAttribute>& InOutAttributes);
+
+	USDEXPORTER_API void HashForMaterialExport(const FUsdMaterialBakingOptions& Options, FSHA1& HashToUpdate);
+
+	USDEXPORTER_API void HashForMeshExport(const FUsdMeshAssetOptions& Options, FSHA1& HashToUpdate);
+}
diff --git a/Engine/Plugins/Importers/USDImporter/Source/USDExporter/Public/USDConversionBlueprintContext.h b/Engine/Plugins/Importers/USDImporter/Source/USDExporter/Public/USDConversionBlueprintContext.h
--- a/Engine/Plugins/Importers/USDImporter/Source/USDExporter/Public/USDConversionBlueprintContext.h
+++ b/Engine/Plugins/Importers/USDImporter/Source/USDExporter/Public/USDConversionBlueprintContext.h
@@ -44,7 +44,7 @@ struct FUsdCombinedPrimMetadata;
  * C++ and e.g. Python are still referencing the same USD Stage in memory, so that we can e.g. use these functions to
  * convert data within stages created via Python.
  */
-UCLASS(meta=(ScriptName="UsdConversionContext"))
+UCLASS(meta = (ScriptName = "UsdConversionContext"))
 class USDEXPORTER_API UUsdConversionBlueprintContext : public UObject
 {
 	GENERATED_BODY()
@@ -53,7 +53,7 @@ class USDEXPORTER_API UUsdConversionBlueprintContext : public UObject
 
 private:
 	/** Stage to use when converting components */
-    UE::FUsdStage Stage;
+	UE::FUsdStage Stage;
 
 	/**
 	 * Whether we will erase our current stage from the stage cache when we Cleanup().
@@ -62,89 +62,108 @@ private:
 	bool bEraseFromStageCache = false;
 
 public:
-    /**
-     * Opens or creates a USD stage using `StageRootLayerPath` as root layer, creating the root layer if needed.
-     * All future conversions will fetch prims and get/set USD data to/from this stage.
+	/**
+	 * Opens or creates a USD stage using `StageRootLayerPath` as root layer, creating the root layer if needed.
+	 * All future conversions will fetch prims and get/set USD data to/from this stage.
 	 * Note: You must remember to call Cleanup() when done, or else this object will permanently hold a reference to the opened stage!
-     */
-	UFUNCTION( BlueprintCallable, Category = "Export context" )
-	void SetStageRootLayer( FFilePath StageRootLayerPath );
+	 */
+	UFUNCTION(BlueprintCallable, Category = "Export context")
+	void SetStageRootLayer(FFilePath StageRootLayerPath);
 
 	/**
 	 * Gets the file path of the root layer of our current stage
 	 */
-	UFUNCTION( BlueprintCallable, Category = "Export context" )
+	UFUNCTION(BlueprintCallable, Category = "Export context")
 	FFilePath GetStageRootLayer();
 
 	/**
-     * Sets the current edit target of our internal stage. When calling the conversion functions, prims and attributes
+	 * Sets the current edit target of our internal stage. When calling the conversion functions, prims and attributes
 	 * will be authored on this edit target only
-     */
-	UFUNCTION( BlueprintCallable, Category = "Export context" )
-	void SetEditTarget( FFilePath EditTargetLayerPath );
+	 */
+	UFUNCTION(BlueprintCallable, Category = "Export context")
+	void SetEditTarget(FFilePath EditTargetLayerPath);
 
 	/**
 	 * Gets the filepath of the current edit target layer of our internal stage
 	 */
-	UFUNCTION( BlueprintCallable, Category = "Export context" )
+	UFUNCTION(BlueprintCallable, Category = "Export context")
 	FFilePath GetEditTarget();
 
 	/**
 	 * Discards the currently opened stage. This is critical when using this class via scripting: The C++ destructor will
 	 * not be called when the python object runs out of scope, so we would otherwise keep a strong reference to the stage
 	 */
-	UFUNCTION( BlueprintCallable, Category = "Export context" )
+	UFUNCTION(BlueprintCallable, Category = "Export context")
 	void Cleanup();
 
 public:
-	// Note: We use FLT_MAX on these functions because Usd.TimeCode.Default().GetValue() is actually a nan, and nan arguments are automatically sanitized to 0.0f.
-	// We manually convert the FLT_MAX value into Usd.TimeCode.Default().GetValue() within the functions though, so if you want the Default timecode just omit the argument
-	// We are also forced to copypaste the FLT_MAX value (3.402823466e+38F) in here as the default arguments are parsed before the preprocessor replaces the defines
+	// Note: We use FLT_MAX on these functions because Usd.TimeCode.Default().GetValue() is actually a nan, and nan arguments are automatically
+	// sanitized to 0.0f. We manually convert the FLT_MAX value into Usd.TimeCode.Default().GetValue() within the functions though, so if you want the
+	// Default timecode just omit the argument We are also forced to copypaste the FLT_MAX value (3.402823466e+38F) in here as the default arguments
+	// are parsed before the preprocessor replaces the defines
 
-	UFUNCTION( BlueprintCallable, Category = "Component conversion" )
-	bool ConvertLightComponent( const ULightComponentBase* Component, const FString& PrimPath, float TimeCode = 3.402823466e+38F );
+	UFUNCTION(BlueprintCallable, Category = "Component conversion")
+	bool ConvertLightComponent(const ULightComponentBase* Component, const FString& PrimPath, float TimeCode = 3.402823466e+38F);
 
-	UFUNCTION( BlueprintCallable, Category = "Component conversion" )
-	bool ConvertDirectionalLightComponent( const UDirectionalLightComponent* Component, const FString& PrimPath, float TimeCode = 3.402823466e+38F );
+	UFUNCTION(BlueprintCallable, Category = "Component conversion")
+	bool ConvertDirectionalLightComponent(const UDirectionalLightComponent* Component, const FString& PrimPath, float TimeCode = 3.402823466e+38F);
 
-	UFUNCTION( BlueprintCallable, Category = "Component conversion" )
-	bool ConvertRectLightComponent( const URectLightComponent* Component, const FString& PrimPath, float TimeCode = 3.402823466e+38F );
+	UFUNCTION(BlueprintCallable, Category = "Component conversion")
+	bool ConvertRectLightComponent(const URectLightComponent* Component, const FString& PrimPath, float TimeCode = 3.402823466e+38F);
 
-	UFUNCTION( BlueprintCallable, Category = "Component conversion" )
-	bool ConvertPointLightComponent( const UPointLightComponent* Component, const FString& PrimPath, float TimeCode = 3.402823466e+38F );
+	UFUNCTION(BlueprintCallable, Category = "Component conversion")
+	bool ConvertPointLightComponent(const UPointLightComponent* Component, const FString& PrimPath, float TimeCode = 3.402823466e+38F);
 
-	UFUNCTION( BlueprintCallable, Category = "Component conversion" )
-	bool ConvertSkyLightComponent( const USkyLightComponent* Component, const FString& PrimPath, float TimeCode = 3.402823466e+38F );
+	UFUNCTION(BlueprintCallable, Category = "Component conversion")
+	bool ConvertSkyLightComponent(const USkyLightComponent* Component, const FString& PrimPath, float TimeCode = 3.402823466e+38F);
 
-	UFUNCTION( BlueprintCallable, Category = "Component conversion" )
-	bool ConvertSpotLightComponent( const USpotLightComponent* Component, const FString& PrimPath, float TimeCode = 3.402823466e+38F );
+	UFUNCTION(BlueprintCallable, Category = "Component conversion")
+	bool ConvertSpotLightComponent(const USpotLightComponent* Component, const FString& PrimPath, float TimeCode = 3.402823466e+38F);
 
 	UFUNCTION(BlueprintCallable, Category = "Component conversion")
 	bool ConvertBoundsComponent(const UUsdDrawModeComponent* Component, const FString& PrimPath, float TimeCode = 3.402823466e+38F);
 
-	UFUNCTION( BlueprintCallable, Category = "Component conversion" )
-	bool ConvertSceneComponent( const USceneComponent* Component, const FString& PrimPath );
+	UFUNCTION(BlueprintCallable, Category = "Component conversion")
+	bool ConvertSceneComponent(const USceneComponent* Component, const FString& PrimPath);
 
-	UFUNCTION( BlueprintCallable, Category = "Component conversion" )
-	bool ConvertHismComponent( const UHierarchicalInstancedStaticMeshComponent* Component, const FString& PrimPath, float TimeCode = 3.402823466e+38F );
+	UFUNCTION(BlueprintCallable, Category = "Component conversion")
+	bool ConvertHismComponent(const UHierarchicalInstancedStaticMeshComponent* Component, const FString& PrimPath, float TimeCode = 3.402823466e+38F);
 
-	UFUNCTION( BlueprintCallable, Category = "Component conversion" )
-	bool ConvertMeshComponent( const UMeshComponent* Component, const FString& PrimPath );
+	UFUNCTION(BlueprintCallable, Category = "Component conversion")
+	bool ConvertMeshComponent(const UMeshComponent* Component, const FString& PrimPath);
 
-	UFUNCTION( BlueprintCallable, Category = "Component conversion" )
-	bool ConvertCineCameraComponent( const UCineCameraComponent* Component, const FString& PrimPath, float TimeCode = 3.402823466e+38F );
+	UFUNCTION(BlueprintCallable, Category = "Component conversion")
+	bool ConvertCineCameraComponent(const UCineCameraComponent* Component, const FString& PrimPath, float TimeCode = 3.402823466e+38F);
 
-	UFUNCTION( BlueprintCallable, Category = "Component conversion" )
-	bool ConvertInstancedFoliageActor( const AInstancedFoliageActor* Actor, const FString& PrimPath, ULevel* InstancesLevel = nullptr, float TimeCode = 3.402823466e+38F );
+	UFUNCTION(BlueprintCallable, Category = "Component conversion")
+	bool ConvertInstancedFoliageActor(
+		const AInstancedFoliageActor* Actor,
+		const FString& PrimPath,
+		ULevel* InstancesLevel = nullptr,
+		float TimeCode = 3.402823466e+38F
+	);
 
-	UFUNCTION( BlueprintCallable, Category = "Component conversion" )
-	bool ConvertLandscapeProxyActorMesh( const ALandscapeProxy* Actor, const FString& PrimPath, int32 LowestLOD, int32 HighestLOD, float TimeCode = 3.402823466e+38F );
+	UFUNCTION(BlueprintCallable, Category = "Component conversion")
+	bool ConvertLandscapeProxyActorMesh(
+		const ALandscapeProxy* Actor,
+		const FString& PrimPath,
+		int32 LowestLOD,
+		int32 HighestLOD,
+		float TimeCode = 3.402823466e+38F
+	);
 
-	UFUNCTION( BlueprintCallable, Category = "Component conversion" )
-	bool ConvertLandscapeProxyActorMaterial( ALandscapeProxy* Actor, const FString& PrimPath, const TArray<FPropertyEntry>& PropertiesToBake, const FIntPoint& DefaultTextureSize, const FDirectoryPath& TexturesDir, float TimeCode = 3.402823466e+38F );
+	UFUNCTION(BlueprintCallable, Category = "Component conversion")
+	bool ConvertLandscapeProxyActorMaterial(
+		ALandscapeProxy* Actor,
+		const FString& PrimPath,
+		const TArray<FPropertyEntry>& PropertiesToBake,
+		const FIntPoint& DefaultTextureSize,
+		const FDirectoryPath& TexturesDir,
+		float TimeCode = 3.402823466e+38F
+	);
 
 public:
-	UFUNCTION( BlueprintCallable, Category = "Conversion utils" )
+	UFUNCTION(BlueprintCallable, Category = "Conversion utils")
 	bool ConvertMaterialOverrides(
 		const UObject* MeshAsset,
 		const TArray<UMaterialInterface*> MaterialOverrides,
@@ -164,8 +183,14 @@ public:
 	 *                      behaviors when authoring the material binding relationships.
 	 * @param bRemoveUnrealMaterials - No longer used.
 	 */
-	UFUNCTION( BlueprintCallable, Category = "Conversion utils" )
-	void ReplaceUnrealMaterialsWithBaked( const FFilePath& LayerToAuthorIn, const TMap<FString, FString>& BakedMaterials, bool bIsAssetLayer, bool bUsePayload, bool bRemoveUnrealMaterials = false );
+	UFUNCTION(BlueprintCallable, Category = "Conversion utils")
+	void ReplaceUnrealMaterialsWithBaked(
+		const FFilePath& LayerToAuthorIn,
+		const TMap<FString, FString>& BakedMaterials,
+		bool bIsAssetLayer,
+		bool bUsePayload,
+		bool bRemoveUnrealMaterials = false
+	);
 
 	/**
 	 * Clears any opinions for the 'unreal' render context surface output of MaterialPrim within LayerToAuthorIn.
@@ -175,31 +200,46 @@ public:
 	 * @param LayerToAuthorIn - Layer to clear the opinions in. Can be the empty string to clear opinions from all layers in the layer stack
 	 * @return Whether we successfully cleared the opinions or not
 	 */
-	UE_DEPRECATED( 5.2, "No longer used as UE material assignments are only visible in the 'unreal' render context anyway" )
-	UFUNCTION( BlueprintCallable, Category = "Conversion utils" )
-	bool RemoveUnrealSurfaceOutput( const FString& PrimPath, const FFilePath& LayerToAuthorIn );
+	UE_DEPRECATED(5.2, "No longer used as UE material assignments are only visible in the 'unreal' render context anyway")
+	UFUNCTION(BlueprintCallable, Category = "Conversion utils")
+	bool RemoveUnrealSurfaceOutput(const FString& PrimPath, const FFilePath& LayerToAuthorIn);
 
 	/** Returns how many frames of animation the stage would generate if imported */
-	UFUNCTION( BlueprintCallable, Category = "Conversion utils" )
+	UFUNCTION(BlueprintCallable, Category = "Conversion utils")
 	int32 GetUsdStageNumFrames();
 
 	/** Adds to Prim the assetInfo metadata the values described in Info */
-	UFUNCTION( BlueprintCallable, Category = "Conversion utils|Metadata" )
-	void SetPrimAssetInfo( const FString& PrimPath, const FUsdUnrealAssetInfo& Info );
+	UFUNCTION(BlueprintCallable, Category = "Conversion utils|Metadata")
+	void SetPrimAssetInfo(const FString& PrimPath, const FUsdUnrealAssetInfo& Info);
 
 	/** Retrieves from Prim the assetInfo metadata values that we use as export metadata, when exporting Unreal assets */
-	UFUNCTION( BlueprintCallable, Category = "Conversion utils|Metadata" )
-	FUsdUnrealAssetInfo GetPrimAssetInfo( const FString& PrimPath );
+	UFUNCTION(BlueprintCallable, Category = "Conversion utils|Metadata")
+	FUsdUnrealAssetInfo GetPrimAssetInfo(const FString& PrimPath);
 
 	/** Applies Metadata to the prim specified at PrimPath, using the provided filters */
 	UFUNCTION(BlueprintCallable, Category = "Conversion utils|Metadata")
-	void SetPrimMetadata(const FString& PrimPath, const FUsdCombinedPrimMetadata& Metadata, const TArray<FString>& BlockedPrefixFilter, bool bInvertFilter);
+	void SetPrimMetadata(
+		const FString& PrimPath,
+		const FUsdCombinedPrimMetadata& Metadata,
+		const TArray<FString>& BlockedPrefixFilter,
+		bool bInvertFilter
+	);
 
 	/** Applies all the metadata contained in UserData to the prim specified at PrimPath, using the provided filters */
 	UFUNCTION(BlueprintCallable, Category = "Conversion utils|Metadata")
-	void SetPrimMetadataFromUserData(const FString& PrimPath, const UUsdAssetUserData* UserData, const TArray<FString>& BlockedPrefixFilter, bool bInvertFilter);
+	void SetPrimMetadataFromUserData(
+		const FString& PrimPath,
+		const UUsdAssetUserData* UserData,
+		const TArray<FString>& BlockedPrefixFilter,
+		bool bInvertFilter
+	);
 
 	/** Extracts metadata from the prim at PrimPath using the provided filters and returns it */
 	UFUNCTION(BlueprintCallable, Category = "Conversion utils|Metadata")
-	FUsdCombinedPrimMetadata GetPrimMetadata(const FString& PrimPath, const TArray<FString>& BlockedPrefixFilter, bool bInvertFilter, bool bCollectFromEntireSubtrees);
+	FUsdCombinedPrimMetadata GetPrimMetadata(
+		const FString& PrimPath,
+		const TArray<FString>& BlockedPrefixFilter,
+		bool bInvertFilter,
+		bool bCollectFromEntireSubtrees
+	);
 };
