USD: Reformat USDSchemas module.

#rnx

diff --git a/Engine/Plugins/Importers/USDImporter/Source/USDSchemas/Private/MeshTranslationImpl.cpp b/Engine/Plugins/Importers/USDImporter/Source/USDSchemas/Private/MeshTranslationImpl.cpp
--- a/Engine/Plugins/Importers/USDImporter/Source/USDSchemas/Private/MeshTranslationImpl.cpp
+++ b/Engine/Plugins/Importers/USDImporter/Source/USDSchemas/Private/MeshTranslationImpl.cpp
@@ -38,7 +38,10 @@
 
 #define UNUSED_UV_INDEX USD_PREVIEW_SURFACE_MAX_UV_SETS
 
-static_assert(USD_PREVIEW_SURFACE_MAX_UV_SETS <= MAX_MESH_TEXTURE_COORDS_MD, "UsdPreviewSurface materials can only have up to as many UV sets as MeshDescription supports!");
+static_assert(
+	USD_PREVIEW_SURFACE_MAX_UV_SETS <= MAX_MESH_TEXTURE_COORDS_MD,
+	"UsdPreviewSurface materials can only have up to as many UV sets as MeshDescription supports!"
+);
 
 namespace UE::MeshTranslationImplInternal::Private
 {
@@ -86,8 +89,11 @@ namespace UE::MeshTranslationImplInternal::Private
 		// Just create an instance of the TwoSided version of the same reference material and copy parameter values.
 		if (GIsEditor && MIC && ReferenceMaterialTwoSided)
 		{
-			UMaterialInstanceConstant* TwoSidedMIC =
-				NewObject<UMaterialInstanceConstant>(GetTransientPackage(), NewInstanceName, OneSidedMat->GetFlags());
+			UMaterialInstanceConstant* TwoSidedMIC = NewObject<UMaterialInstanceConstant>(
+				GetTransientPackage(),
+				NewInstanceName,
+				OneSidedMat->GetFlags()
+			);
 
 			TwoSidedMIC->SetParentEditorOnly(ReferenceMaterialTwoSided);
 			TwoSidedMIC->CopyMaterialUniformParametersEditorOnly(OneSidedMat);
@@ -116,28 +122,28 @@ namespace UE::MeshTranslationImplInternal::Private
 		}
 
 		else
-#endif // WITH_EDITOR
+#endif	  // WITH_EDITOR
 
-		// At runtime all we can do is create another instance of our two-sided reference materials, we cannot set
-		// another override
-		if (MID && ReferenceMaterialTwoSided)
-		{
-			UMaterialInstanceDynamic* TwoSidedMID = UMaterialInstanceDynamic::Create(
-				ReferenceMaterialTwoSided,
-				GetTransientPackage(),
-				NewInstanceName
-			);
-			if (!ensure(TwoSidedMID))
+			// At runtime all we can do is create another instance of our two-sided reference materials, we cannot set
+			// another override
+			if (MID && ReferenceMaterialTwoSided)
 			{
-				return nullptr;
-			}
+				UMaterialInstanceDynamic* TwoSidedMID = UMaterialInstanceDynamic::Create(
+					ReferenceMaterialTwoSided,
+					GetTransientPackage(),
+					NewInstanceName
+				);
+				if (!ensure(TwoSidedMID))
+				{
+					return nullptr;
+				}
 
-			TwoSidedMID->CopyParameterOverrides(MID);
+				TwoSidedMID->CopyParameterOverrides(MID);
 
-			TwoSidedMat = TwoSidedMID;
-		}
+				TwoSidedMat = TwoSidedMID;
+			}
 
-		if(TwoSidedMat)
+		if (TwoSidedMat)
 		{
 			UUsdMaterialAssetUserData* UserData = DuplicateObject(OneSidedUserData, TwoSidedMat, TEXT("USDAssetUserData"));
 			TwoSidedMat->AddAssetUserData(UserData);
@@ -162,12 +168,20 @@ namespace UE::MeshTranslationImplInternal::Private
 	)
 	{
 		UUsdMaterialAssetUserData* MaterialAssetUserData = Material.GetAssetUserData<UUsdMaterialAssetUserData>();
-		if (!ensureMsgf(MaterialAssetUserData, TEXT("Expected material '%s' to have an UUsdMaterialAssetUserData at this point!"), *Material.GetPathName()))
+		if (!ensureMsgf(
+				MaterialAssetUserData,
+				TEXT("Expected material '%s' to have an UUsdMaterialAssetUserData at this point!"),
+				*Material.GetPathName()
+			))
 		{
 			return nullptr;
 		}
 
-		UE_LOG(LogUsd, Verbose, TEXT("Getting compatible material based on Material '%s' (Parameter to primvar: %s, primvar to UV index: %s) matching mesh primvar to UV index mapping '%s'"),
+		UE_LOG(
+			LogUsd,
+			Verbose,
+			TEXT("Getting compatible material based on Material '%s' (Parameter to primvar: %s, primvar to UV index: %s) matching mesh primvar to UV "
+				 "index mapping '%s'"),
 			*Material.GetPathName(),
 			*UsdUtils::StringifyMap(MaterialAssetUserData->ParameterToPrimvar),
 			*UsdUtils::StringifyMap(MaterialAssetUserData->PrimvarToUVIndex),
@@ -204,7 +218,10 @@ namespace UE::MeshTranslationImplInternal::Private
 			}
 			else
 			{
-				UE_LOG(LogUsd, Log, TEXT("Failed to find primvar '%s' needed by material '%s' on its assigned mesh. Available primvars and UV indices: %s"),
+				UE_LOG(
+					LogUsd,
+					Log,
+					TEXT("Failed to find primvar '%s' needed by material '%s' on its assigned mesh. Available primvars and UV indices: %s"),
 					*MaterialPrimvar,
 					*Material.GetPathName(),
 					*UsdUtils::StringifyMap(MeshPrimvarToUVIndex)
@@ -225,9 +242,7 @@ namespace UE::MeshTranslationImplInternal::Private
 		}
 		if (bCompatible)
 		{
-			UE_LOG(LogUsd, Verbose, TEXT("Material '%s' is compatible with provided primvar to UV index mapping"),
-				*Material.GetPathName()
-			);
+			UE_LOG(LogUsd, Verbose, TEXT("Material '%s' is compatible with provided primvar to UV index mapping"), *Material.GetPathName());
 			return &Material;
 		}
 
@@ -321,13 +336,12 @@ namespace UE::MeshTranslationImplInternal::Private
 				CompatiblePrimvarToUVIndex.Add(Pair);
 			}
 
-			const FName NewInstanceName = MakeUniqueObjectName(
-				GetTransientPackage(),
-				UMaterialInstance::StaticClass(),
-				Material.GetFName()
-			);
+			const FName NewInstanceName = MakeUniqueObjectName(GetTransientPackage(), UMaterialInstance::StaticClass(), Material.GetFName());
 
-			UE_LOG(LogUsd, Verbose, TEXT("Generating compatible version of Material '%s' (Parameter to primvar: %s, primvar to UV index: %s) with hash '%s'"),
+			UE_LOG(
+				LogUsd,
+				Verbose,
+				TEXT("Generating compatible version of Material '%s' (Parameter to primvar: %s, primvar to UV index: %s) with hash '%s'"),
 				*Material.GetPathName(),
 				*UsdUtils::StringifyMap(MaterialAssetUserData->ParameterToPrimvar),
 				*UsdUtils::StringifyMap(CompatiblePrimvarToUVIndex),
@@ -348,22 +362,14 @@ namespace UE::MeshTranslationImplInternal::Private
 				CompatibleMaterial = CompatibleMIC;
 			}
 			else
-#endif // WITH_EDITOR
+#endif	  // WITH_EDITOR
 			{
-				UMaterialInstanceDynamic* CompatibleMID = UMaterialInstanceDynamic::Create(
-					&Material,
-					GetTransientPackage(),
-					NewInstanceName
-				);
+				UMaterialInstanceDynamic* CompatibleMID = UMaterialInstanceDynamic::Create(&Material, GetTransientPackage(), NewInstanceName);
 
 				CompatibleMaterial = CompatibleMID;
 			}
 
-			UUsdMaterialAssetUserData* CompatibleUserData = DuplicateObject(
-				MaterialAssetUserData,
-				CompatibleMaterial,
-				TEXT("USDAssetUserData")
-			);
+			UUsdMaterialAssetUserData* CompatibleUserData = DuplicateObject(MaterialAssetUserData, CompatibleMaterial, TEXT("USDAssetUserData"));
 			CompatibleUserData->PrimvarToUVIndex = CompatiblePrimvarToUVIndex;
 			CompatibleMaterial->AddAssetUserData(CompatibleUserData);
 
@@ -381,11 +387,7 @@ namespace UE::MeshTranslationImplInternal::Private
 						// used to sample it with
 						if (*UVIndex == UNUSED_UV_INDEX)
 						{
-							UsdUtils::SetScalarParameterValue(
-								*CompatibleInstance,
-								*FString::Printf(TEXT("Use%sTexture"), *Parameter),
-								0.0f
-							);
+							UsdUtils::SetScalarParameterValue(*CompatibleInstance, *FString::Printf(TEXT("Use%sTexture"), *Parameter), 0.0f);
 						}
 						else
 						{
@@ -400,7 +402,7 @@ namespace UE::MeshTranslationImplInternal::Private
 
 #if WITH_EDITOR
 				CompatibleInstance->PostEditChange();
-#endif // WITH_EDITOR
+#endif	  // WITH_EDITOR
 			}
 		}
 
@@ -422,7 +424,7 @@ namespace UE::MeshTranslationImplInternal::Private
 
 		return CompatibleMaterial;
 	}
-};
+};	  // namespace UE::MeshTranslationImplInternal::Private
 
 TMap<const UsdUtils::FUsdPrimMaterialSlot*, UMaterialInterface*> MeshTranslationImpl::ResolveMaterialAssignmentInfo(
 	const pxr::UsdPrim& UsdPrim,
@@ -436,7 +438,7 @@ TMap<const UsdUtils::FUsdPrimMaterialSlot*, UMaterialInterface*> MeshTranslation
 	FScopedUnrealAllocs Allocs;
 
 	TMap<const UsdUtils::FUsdPrimMaterialSlot*, UMaterialInterface*> ResolvedMaterials;
-	if(AssignmentInfo.Num() == 0)
+	if (AssignmentInfo.Num() == 0)
 	{
 		return ResolvedMaterials;
 	}
@@ -450,7 +452,7 @@ TMap<const UsdUtils::FUsdPrimMaterialSlot*, UMaterialInterface*> MeshTranslation
 	uint32 GlobalResolvedMaterialIndex = 0;
 	for (int32 InfoIndex = 0; InfoIndex < AssignmentInfo.Num(); ++InfoIndex)
 	{
-		const TArray< UsdUtils::FUsdPrimMaterialSlot >& Slots = AssignmentInfo[InfoIndex].Slots;
+		const TArray<UsdUtils::FUsdPrimMaterialSlot>& Slots = AssignmentInfo[InfoIndex].Slots;
 
 		for (int32 SlotIndex = 0; SlotIndex < Slots.Num(); ++SlotIndex, ++GlobalResolvedMaterialIndex)
 		{
@@ -459,220 +461,228 @@ TMap<const UsdUtils::FUsdPrimMaterialSlot*, UMaterialInterface*> MeshTranslation
 
 			switch (Slot.AssignmentType)
 			{
-			case UsdUtils::EPrimAssignmentType::DisplayColor:
-			{
-				const FString PrefixedHash = UsdUtils::GetAssetHashPrefix(UsdPrim, bReuseIdenticalAssets) + Slot.MaterialSource;
-
-				// Try reusing an already created DisplayColor material
-				if (UMaterialInterface* ExistingMaterial = Cast<UMaterialInterface>(AssetCache.GetCachedAsset(PrefixedHash)))
+				case UsdUtils::EPrimAssignmentType::DisplayColor:
 				{
-					Material = ExistingMaterial;
-				}
+					const FString PrefixedHash = UsdUtils::GetAssetHashPrefix(UsdPrim, bReuseIdenticalAssets) + Slot.MaterialSource;
 
-				// Need to create a new DisplayColor material
-				if (Material == nullptr)
-				{
-					if (TOptional<IUsdClassesModule::FDisplayColorMaterial> DisplayColorDesc = IUsdClassesModule::FDisplayColorMaterial::FromString(
-							Slot.MaterialSource
-						))
+					// Try reusing an already created DisplayColor material
+					if (UMaterialInterface* ExistingMaterial = Cast<UMaterialInterface>(AssetCache.GetCachedAsset(PrefixedHash)))
 					{
-						UMaterialInstance* MaterialInstance = nullptr;
+						Material = ExistingMaterial;
+					}
 
-						if (GIsEditor)  // Editor, PIE => true; Standlone, packaged => false
+					// Need to create a new DisplayColor material
+					if (Material == nullptr)
+					{
+						if (TOptional<IUsdClassesModule::FDisplayColorMaterial>
+								DisplayColorDesc = IUsdClassesModule::FDisplayColorMaterial::FromString(Slot.MaterialSource))
 						{
-							MaterialInstance = IUsdClassesModule::CreateDisplayColorMaterialInstanceConstant(DisplayColorDesc.GetValue());
+							UMaterialInstance* MaterialInstance = nullptr;
+
+							if (GIsEditor)	  // Editor, PIE => true; Standlone, packaged => false
+							{
+								MaterialInstance = IUsdClassesModule::CreateDisplayColorMaterialInstanceConstant(DisplayColorDesc.GetValue());
+							}
+							else
+							{
+								MaterialInstance = IUsdClassesModule::CreateDisplayColorMaterialInstanceDynamic(DisplayColorDesc.GetValue());
+							}
+
+							if (MaterialInstance)
+							{
+								// Leave PrimPath as empty as it likely will be reused by many prims
+								UUsdAssetUserData* UserData = NewObject<UUsdAssetUserData>(MaterialInstance, TEXT("USDAssetUserData"));
+								MaterialInstance->AddAssetUserData(UserData);
+							}
+
+							// We can only cache transient assets
+							MaterialInstance->SetFlags(RF_Transient);
+
+							AssetCache.CacheAsset(PrefixedHash, MaterialInstance);
+							Material = MaterialInstance;
 						}
-						else
+					}
+
+					break;
+				}
+				case UsdUtils::EPrimAssignmentType::MaterialPrim:
+				{
+					UMaterialInterface* OneSidedMat = nullptr;
+					bool bOneSidedMatIsInstanceOfReferencePreviewSurface = false;
+
+					UE::FSdfPath MaterialPrimPath{*Slot.MaterialSource};
+
+					TArray<UMaterialInterface*> ExistingMaterials = InfoCache.GetAssetsForPrim<UMaterialInterface>(MaterialPrimPath);
+
+					for (UMaterialInterface* ExistingMaterial : ExistingMaterials)
+					{
+						const bool bExistingIsTwoSided = ExistingMaterial->IsTwoSided();
+
+						if (!bExistingIsTwoSided)
 						{
-							MaterialInstance = IUsdClassesModule::CreateDisplayColorMaterialInstanceDynamic(DisplayColorDesc.GetValue());
+							// Prefer sticking with a material instance that has as parent one of our reference materials.
+							// The idea here being that we have two approaches when making TwoSided and compatible
+							// materials: A) Make the material compatible first, and then a TwoSided version of the
+							// compatible; B) Make the material TwoSided first, and then a compatible version of the
+							// TwoSided; We're going to chose B), for the reason that at runtime we can only make a material
+							// TwoSided if it is an instance of our reference materials (as we can't manually change the
+							// material base property overrides at runtime)
+							UMaterialInstance* ExistingInstance = Cast<UMaterialInstance>(ExistingMaterial);
+							const bool bExistingIsInstanceOfReferencePreviewSurface = ExistingInstance
+																					  && MeshTranslationImpl::IsReferencePreviewSurfaceMaterial(
+																						  ExistingInstance->Parent
+																					  );
+							if (!OneSidedMat || (!bOneSidedMatIsInstanceOfReferencePreviewSurface && bExistingIsInstanceOfReferencePreviewSurface))
+							{
+								OneSidedMat = ExistingMaterial;
+								bOneSidedMatIsInstanceOfReferencePreviewSurface = bExistingIsInstanceOfReferencePreviewSurface;
+							}
 						}
 
-						if(MaterialInstance)
+						if (Slot.bMeshIsDoubleSided == bExistingIsTwoSided)
 						{
-							// Leave PrimPath as empty as it likely will be reused by many prims
-							UUsdAssetUserData* UserData = NewObject<UUsdAssetUserData>(MaterialInstance, TEXT("USDAssetUserData"));
-							MaterialInstance->AddAssetUserData(UserData);
+							Material = ExistingMaterial;
 						}
-
-						// We can only cache transient assets
-						MaterialInstance->SetFlags(RF_Transient);
-
-						AssetCache.CacheAsset(PrefixedHash, MaterialInstance);
-						Material = MaterialInstance;
 					}
-				}
 
-				break;
-			}
-			case UsdUtils::EPrimAssignmentType::MaterialPrim:
-			{
-				UMaterialInterface* OneSidedMat = nullptr;
-				bool bOneSidedMatIsInstanceOfReferencePreviewSurface = false;
+					FString PrefixedMaterialHash = Material ? AssetCache.GetHashForAsset(Material) : FString{};
+					FString HashPrefix = UsdUtils::GetAssetHashPrefix(UsdPrim.GetStage()->GetPrimAtPath(MaterialPrimPath), bReuseIdenticalAssets);
 
-				UE::FSdfPath MaterialPrimPath{*Slot.MaterialSource};
+					// Need to create a two-sided material on-demand, *before* we make it compatible:
+					// This because at runtime we can't just set the base property overrides, and just instead create a new
+					// MIC based on the TwoSided reference material, and the compatible material should be a MIC of that MIC
+					if (Slot.bMeshIsDoubleSided && !Material)
+					{
+						// By now we parsed all materials so we must have the single-sided version of this material
+						if (!OneSidedMat)
+						{
+							UE_LOG(
+								LogUsd,
+								Warning,
+								TEXT("Failed to generate a two-sided material from the material prim at path '%s' as no "
+									 "single-sided material was generated for it."),
+								*Slot.MaterialSource
+							);
+							continue;
+						}
 
-				TArray<UMaterialInterface*> ExistingMaterials = InfoCache.GetAssetsForPrim<UMaterialInterface>(MaterialPrimPath);
+						const FString PrefixedOneSidedHash = AssetCache.GetHashForAsset(OneSidedMat);
+						const FString PrefixedTwoSidedHash = PrefixedOneSidedHash + UnrealIdentifiers::TwoSidedMaterialSuffix;
 
-				for (UMaterialInterface* ExistingMaterial : ExistingMaterials)
-				{
-					const bool bExistingIsTwoSided = ExistingMaterial->IsTwoSided();
+						// Check if for some reason we already have a two-sided material ready due to a complex scenario
+						// related to the global cache
+						UMaterialInterface* TwoSidedMat = Cast<UMaterialInterface>(AssetCache.GetCachedAsset(PrefixedTwoSidedHash));
+						if (!TwoSidedMat)
+						{
+							TwoSidedMat = UE::MeshTranslationImplInternal::Private::CreateTwoSidedVersionOfMaterial(OneSidedMat);
+						}
 
-					if (!bExistingIsTwoSided)
-					{
-						// Prefer sticking with a material instance that has as parent one of our reference materials.
-						// The idea here being that we have two approaches when making TwoSided and compatible
-						// materials: A) Make the material compatible first, and then a TwoSided version of the
-						// compatible; B) Make the material TwoSided first, and then a compatible version of the
-						// TwoSided; We're going to chose B), for the reason that at runtime we can only make a material
-						// TwoSided if it is an instance of our reference materials (as we can't manually change the
-						// material base property overrides at runtime)
-						UMaterialInstance* ExistingInstance = Cast<UMaterialInstance>(ExistingMaterial);
-						const bool bExistingIsInstanceOfReferencePreviewSurface =
-							ExistingInstance
-							&& MeshTranslationImpl::IsReferencePreviewSurfaceMaterial(ExistingInstance->Parent);
-						if (!OneSidedMat
-							|| (!bOneSidedMatIsInstanceOfReferencePreviewSurface
-								&& bExistingIsInstanceOfReferencePreviewSurface))
+						if (TwoSidedMat)
 						{
-							OneSidedMat = ExistingMaterial;
-							bOneSidedMatIsInstanceOfReferencePreviewSurface = bExistingIsInstanceOfReferencePreviewSurface;
+							TwoSidedMat->SetFlags(RF_Transient);
+							Material = TwoSidedMat;
+							PrefixedMaterialHash = PrefixedTwoSidedHash;
+						}
+						else
+						{
+							UE_LOG(
+								LogUsd,
+								Warning,
+								TEXT("Failed to generate a two-sided material from the material prim at path '%s'. Falling "
+									 "back to using the single-sided material '%s' instead."),
+								*Slot.MaterialSource,
+								*OneSidedMat->GetPathName()
+							);
+							Material = OneSidedMat;
+							PrefixedMaterialHash = PrefixedOneSidedHash;
 						}
 					}
 
-					if (Slot.bMeshIsDoubleSided == bExistingIsTwoSided)
+					if (Material)
 					{
-						Material = ExistingMaterial;
-					}
-				}
+						// Cache the material to "ping it" as active, but also register two sided materials for the
+						// first time
+						AssetCache.CacheAsset(PrefixedMaterialHash, Material);
+						InfoCache.LinkAssetToPrim(UE::FSdfPath{*Slot.MaterialSource}, Material);
 
-				FString PrefixedMaterialHash = Material ? AssetCache.GetHashForAsset(Material) : FString{};
-				FString HashPrefix = UsdUtils::GetAssetHashPrefix(UsdPrim.GetStage()->GetPrimAtPath(MaterialPrimPath), bReuseIdenticalAssets);
+						// Finally, try to make our generated material primvar-compatible. We do this last because this will
+						// create another instance with the non-compatible material as reference material, which means we also
+						// need that reference to be cached and linked for the asset cache to be able to handle dependencies
+						// properly
+						if (UMaterialInterface* AlreadyHandledMaterial = MaterialToCompatibleMaterial.FindRef(Material))
+						{
+							Material = AlreadyHandledMaterial;
 
-				// Need to create a two-sided material on-demand, *before* we make it compatible:
-				// This because at runtime we can't just set the base property overrides, and just instead create a new
-				// MIC based on the TwoSided reference material, and the compatible material should be a MIC of that MIC
-				if (Slot.bMeshIsDoubleSided && !Material)
+							AssetCache.TouchAsset(Material);
+							InfoCache.LinkAssetToPrim(UE::FSdfPath{*Slot.MaterialSource}, Material);
+						}
+						else
+						{
+							UMaterialInterface*
+								CompatibleMaterial = UE::MeshTranslationImplInternal::Private::CreatePrimvarCompatibleVersionOfMaterial(
+									*Material,
+									MeshPrimvarToUVIndex,
+									&AssetCache,
+									&InfoCache,
+									HashPrefix,
+									bReuseIdenticalAssets
+								);
+
+							if (CompatibleMaterial)
+							{
+								MaterialToCompatibleMaterial.Add(Material, CompatibleMaterial);
+								Material = CompatibleMaterial;
+							}
+						}
+					}
+
+					break;
+				}
+				case UsdUtils::EPrimAssignmentType::UnrealMaterial:
 				{
-					// By now we parsed all materials so we must have the single-sided version of this material
-					if (!OneSidedMat)
+					UObject* Object = FSoftObjectPath(Slot.MaterialSource).TryLoad();
+					Material = Cast<UMaterialInterface>(Object);
+					if (!Object)
 					{
 						UE_LOG(
 							LogUsd,
 							Warning,
-							TEXT("Failed to generate a two-sided material from the material prim at path '%s' as no "
-								 "single-sided material was generated for it."),
-							*Slot.MaterialSource
+							TEXT("UE material '%s' for prim '%s' could not be loaded or was not found."),
+							*Slot.MaterialSource,
+							*UsdToUnreal::ConvertPath(UsdPrim.GetPrimPath())
 						);
-						continue;
 					}
-
-					const FString PrefixedOneSidedHash = AssetCache.GetHashForAsset(OneSidedMat);
-					const FString PrefixedTwoSidedHash = PrefixedOneSidedHash + UnrealIdentifiers::TwoSidedMaterialSuffix;
-
-					// Check if for some reason we already have a two-sided material ready due to a complex scenario
-					// related to the global cache
-					UMaterialInterface* TwoSidedMat = Cast<UMaterialInterface>(AssetCache.GetCachedAsset(PrefixedTwoSidedHash));
-					if (!TwoSidedMat)
-					{
-						TwoSidedMat =
-							UE::MeshTranslationImplInternal::Private::CreateTwoSidedVersionOfMaterial(OneSidedMat);
-					}
-
-					if (TwoSidedMat)
-					{
-						TwoSidedMat->SetFlags(RF_Transient);
-						Material = TwoSidedMat;
-						PrefixedMaterialHash = PrefixedTwoSidedHash;
-					}
-					else
+					else if (!Material)
 					{
 						UE_LOG(
 							LogUsd,
 							Warning,
-							TEXT("Failed to generate a two-sided material from the material prim at path '%s'. Falling "
-								 "back to using the single-sided material '%s' instead."),
+							TEXT("Object '%s' assigned as an Unreal Material for prim '%s' is not actually a material (but instead a '%s') and will "
+								 "not be used"),
 							*Slot.MaterialSource,
-							*OneSidedMat->GetPathName()
+							*UsdToUnreal::ConvertPath(UsdPrim.GetPrimPath()),
+							*Object->GetClass()->GetName()
 						);
-						Material = OneSidedMat;
-						PrefixedMaterialHash = PrefixedOneSidedHash;
-					}
-				}
-
-				if (Material)
-				{
-					// Cache the material to "ping it" as active, but also register two sided materials for the
-					// first time
-					AssetCache.CacheAsset(PrefixedMaterialHash, Material);
-					InfoCache.LinkAssetToPrim(UE::FSdfPath{*Slot.MaterialSource}, Material);
-
-					// Finally, try to make our generated material primvar-compatible. We do this last because this will
-					// create another instance with the non-compatible material as reference material, which means we also
-					// need that reference to be cached and linked for the asset cache to be able to handle dependencies
-					// properly
-					if (UMaterialInterface* AlreadyHandledMaterial = MaterialToCompatibleMaterial.FindRef(Material))
-					{
-						Material = AlreadyHandledMaterial;
-
-						AssetCache.TouchAsset(Material);
-						InfoCache.LinkAssetToPrim(UE::FSdfPath{*Slot.MaterialSource}, Material);
 					}
-					else
+					else if (!Material->IsTwoSided() && Slot.bMeshIsDoubleSided)
 					{
-						UMaterialInterface* CompatibleMaterial =
-							UE::MeshTranslationImplInternal::Private::CreatePrimvarCompatibleVersionOfMaterial(
-								*Material,
-								MeshPrimvarToUVIndex,
-								&AssetCache,
-								&InfoCache,
-								HashPrefix,
-								bReuseIdenticalAssets
-							);
-
-						if (CompatibleMaterial)
-						{
-							MaterialToCompatibleMaterial.Add(Material, CompatibleMaterial);
-							Material = CompatibleMaterial;
-						}
+						UE_LOG(
+							LogUsd,
+							Warning,
+							TEXT("Using one-sided UE material '%s' for doubleSided prim '%s'"),
+							*Slot.MaterialSource,
+							*UsdToUnreal::ConvertPath(UsdPrim.GetPrimPath())
+						);
 					}
-				}
 
-				break;
-			}
-			case UsdUtils::EPrimAssignmentType::UnrealMaterial:
-			{
-				UObject* Object = FSoftObjectPath(Slot.MaterialSource).TryLoad();
-				Material = Cast< UMaterialInterface >(Object);
-				if (!Object)
-				{
-					UE_LOG(LogUsd, Warning, TEXT("UE material '%s' for prim '%s' could not be loaded or was not found."),
-						*Slot.MaterialSource,
-						*UsdToUnreal::ConvertPath(UsdPrim.GetPrimPath()));
-				}
-				else if (!Material)
-				{
-					UE_LOG(LogUsd, Warning, TEXT("Object '%s' assigned as an Unreal Material for prim '%s' is not actually a material (but instead a '%s') and will not be used"),
-						*Slot.MaterialSource,
-						*UsdToUnreal::ConvertPath(UsdPrim.GetPrimPath()),
-						*Object->GetClass()->GetName()
-					);
+					break;
 				}
-				else if (!Material->IsTwoSided() && Slot.bMeshIsDoubleSided)
+				case UsdUtils::EPrimAssignmentType::None:
+				default:
 				{
-					UE_LOG(LogUsd, Warning, TEXT("Using one-sided UE material '%s' for doubleSided prim '%s'"),
-						*Slot.MaterialSource,
-						*UsdToUnreal::ConvertPath(UsdPrim.GetPrimPath())
-					);
+					ensure(false);
+					break;
 				}
-
-				break;
-			}
-			case UsdUtils::EPrimAssignmentType::None:
-			default:
-			{
-				ensure(false);
-				break;
-			}
 			}
 
 			ResolvedMaterials.Add(&Slot, Material);
@@ -714,7 +724,7 @@ void MeshTranslationImpl::SetMaterialOverrides(
 	}
 
 	TArray<UsdUtils::FUsdPrimMaterialAssignmentInfo> LODIndexToAssignments;
-	const bool bProvideMaterialIndices = false; // We have no use for material indices and it can be slow to retrieve, as it will iterate all faces
+	const bool bProvideMaterialIndices = false;	   // We have no use for material indices and it can be slow to retrieve, as it will iterate all faces
 
 	// Extract material assignment info from prim if it is a LOD mesh
 	bool bInterpretedLODs = false;
@@ -771,15 +781,13 @@ void MeshTranslationImpl::SetMaterialOverrides(
 		}
 		else
 		{
-			LODIndexToAssignments = {
-				UsdUtils::GetPrimMaterialAssignments(
-					ValidPrim,
-					pxr::UsdTimeCode(Time),
-					bProvideMaterialIndices,
-					RenderContextToken,
-					MaterialPurposeToken
-				)
-			};
+			LODIndexToAssignments = {UsdUtils::GetPrimMaterialAssignments(
+				ValidPrim,
+				pxr::UsdTimeCode(Time),
+				bProvideMaterialIndices,
+				RenderContextToken,
+				MaterialPurposeToken
+			)};
 		}
 	}
 
@@ -809,18 +817,22 @@ void MeshTranslationImpl::SetMaterialOverrides(
 	}
 	else
 	{
-		ensureMsgf(false, TEXT("Unexpected component class '%s' encountered when setting material overrides for prim '%s'!"),
+		ensureMsgf(
+			false,
+			TEXT("Unexpected component class '%s' encountered when setting material overrides for prim '%s'!"),
 			*MeshComponent.GetClass()->GetName(),
 			*UsdToUnreal::ConvertPath(Prim.GetPrimPath())
 		);
 	}
 
-	ensureMsgf(UserData, TEXT("Mesh assigned to component '%s' generated for prim '%s' should have an UUsdMeshAssetUserData at this point!"),
+	ensureMsgf(
+		UserData,
+		TEXT("Mesh assigned to component '%s' generated for prim '%s' should have an UUsdMeshAssetUserData at this point!"),
 		*MeshComponent.GetPathName(),
 		*UsdToUnreal::ConvertPath(Prim.GetPrimPath())
 	);
 
-	if(UserData && LODIndexToAssignments.Num() > 0)
+	if (UserData && LODIndexToAssignments.Num() > 0)
 	{
 		// Stash our PrimvarToUVIndex in here, as that's where ResolveMaterialAssignmentInfo will look for it
 		LODIndexToAssignments[0].PrimvarToUVIndex = UserData->PrimvarToUVIndex;
@@ -839,7 +851,7 @@ void MeshTranslationImpl::SetMaterialOverrides(
 	uint32 StaticMeshSlotIndex = 0;
 	for (int32 LODIndex = 0; LODIndex < LODIndexToAssignments.Num(); ++LODIndex)
 	{
-		const TArray< UsdUtils::FUsdPrimMaterialSlot >& LODSlots = LODIndexToAssignments[LODIndex].Slots;
+		const TArray<UsdUtils::FUsdPrimMaterialSlot>& LODSlots = LODIndexToAssignments[LODIndex].Slots;
 		for (int32 LODSlotIndex = 0; LODSlotIndex < LODSlots.Num(); ++LODSlotIndex, ++StaticMeshSlotIndex)
 		{
 			// If we don't even have as many existing assignments as we have overrides just stop here.
@@ -862,7 +874,14 @@ void MeshTranslationImpl::SetMaterialOverrides(
 			}
 			else
 			{
-				UE_LOG(LogUsd, Error, TEXT("Lost track of resolved material for slot '%d' of LOD '%d' for mesh '%s'"), LODSlotIndex, LODIndex, *UsdToUnreal::ConvertPath(Prim.GetPath()));
+				UE_LOG(
+					LogUsd,
+					Error,
+					TEXT("Lost track of resolved material for slot '%d' of LOD '%d' for mesh '%s'"),
+					LODSlotIndex,
+					LODIndex,
+					*UsdToUnreal::ConvertPath(Prim.GetPath())
+				);
 				continue;
 			}
 
@@ -971,7 +990,7 @@ UMaterialInterface* MeshTranslationImpl::GetReferencePreviewSurfaceMaterial(EUsd
 		return nullptr;
 	}
 
-	return Cast< UMaterialInterface >(TargetMaterialPath->TryLoad());
+	return Cast<UMaterialInterface>(TargetMaterialPath->TryLoad());
 }
 
 UMaterialInterface* MeshTranslationImpl::GetVTVersionOfReferencePreviewSurfaceMaterial(UMaterialInterface* ReferenceMaterial)
@@ -994,19 +1013,19 @@ UMaterialInterface* MeshTranslationImpl::GetVTVersionOfReferencePreviewSurfaceMa
 	}
 	else if (PathName == Settings->ReferencePreviewSurfaceMaterial)
 	{
-		return Cast< UMaterialInterface >(Settings->ReferencePreviewSurfaceVTMaterial.TryLoad());
+		return Cast<UMaterialInterface>(Settings->ReferencePreviewSurfaceVTMaterial.TryLoad());
 	}
 	else if (PathName == Settings->ReferencePreviewSurfaceTwoSidedMaterial)
 	{
-		return Cast< UMaterialInterface >(Settings->ReferencePreviewSurfaceTwoSidedVTMaterial.TryLoad());
+		return Cast<UMaterialInterface>(Settings->ReferencePreviewSurfaceTwoSidedVTMaterial.TryLoad());
 	}
 	else if (PathName == Settings->ReferencePreviewSurfaceTranslucentMaterial)
 	{
-		return Cast< UMaterialInterface >(Settings->ReferencePreviewSurfaceTranslucentVTMaterial.TryLoad());
+		return Cast<UMaterialInterface>(Settings->ReferencePreviewSurfaceTranslucentVTMaterial.TryLoad());
 	}
 	else if (PathName == Settings->ReferencePreviewSurfaceTranslucentTwoSidedMaterial)
 	{
-		return Cast< UMaterialInterface >(Settings->ReferencePreviewSurfaceTranslucentTwoSidedVTMaterial.TryLoad());
+		return Cast<UMaterialInterface>(Settings->ReferencePreviewSurfaceTranslucentTwoSidedVTMaterial.TryLoad());
 	}
 
 	// We should only ever call this function with a ReferenceMaterial that matches one of the above paths
@@ -1034,19 +1053,19 @@ UMaterialInterface* MeshTranslationImpl::GetTwoSidedVersionOfReferencePreviewSur
 	}
 	else if (PathName == Settings->ReferencePreviewSurfaceMaterial)
 	{
-		return Cast< UMaterialInterface >(Settings->ReferencePreviewSurfaceTwoSidedMaterial.TryLoad());
+		return Cast<UMaterialInterface>(Settings->ReferencePreviewSurfaceTwoSidedMaterial.TryLoad());
 	}
 	else if (PathName == Settings->ReferencePreviewSurfaceTranslucentMaterial)
 	{
-		return Cast< UMaterialInterface >(Settings->ReferencePreviewSurfaceTranslucentTwoSidedMaterial.TryLoad());
+		return Cast<UMaterialInterface>(Settings->ReferencePreviewSurfaceTranslucentTwoSidedMaterial.TryLoad());
 	}
 	else if (PathName == Settings->ReferencePreviewSurfaceVTMaterial)
 	{
-		return Cast< UMaterialInterface >(Settings->ReferencePreviewSurfaceTwoSidedVTMaterial.TryLoad());
+		return Cast<UMaterialInterface>(Settings->ReferencePreviewSurfaceTwoSidedVTMaterial.TryLoad());
 	}
 	else if (PathName == Settings->ReferencePreviewSurfaceTranslucentVTMaterial)
 	{
-		return Cast< UMaterialInterface >(Settings->ReferencePreviewSurfaceTranslucentTwoSidedVTMaterial.TryLoad());
+		return Cast<UMaterialInterface>(Settings->ReferencePreviewSurfaceTranslucentTwoSidedVTMaterial.TryLoad());
 	}
 
 	// We should only ever call this function with a ReferenceMaterial that matches one of the above paths
@@ -1077,10 +1096,9 @@ bool MeshTranslationImpl::IsReferencePreviewSurfaceMaterial(UMaterialInterface*
 		Settings->ReferencePreviewSurfaceVTMaterial,
 		Settings->ReferencePreviewSurfaceTranslucentVTMaterial,
 		Settings->ReferencePreviewSurfaceTwoSidedVTMaterial,
-		Settings->ReferencePreviewSurfaceTranslucentTwoSidedVTMaterial
-	};
+		Settings->ReferencePreviewSurfaceTranslucentTwoSidedVTMaterial};
 
 	return ReferenceMaterials.Contains(PathName);
 }
 
-#endif // #if USE_USD_SDK
\ No newline at end of file
+#endif	  // #if USE_USD_SDK
diff --git a/Engine/Plugins/Importers/USDImporter/Source/USDSchemas/Private/MeshTranslationImpl.h b/Engine/Plugins/Importers/USDImporter/Source/USDSchemas/Private/MeshTranslationImpl.h
--- a/Engine/Plugins/Importers/USDImporter/Source/USDSchemas/Private/MeshTranslationImpl.h
+++ b/Engine/Plugins/Importers/USDImporter/Source/USDSchemas/Private/MeshTranslationImpl.h
@@ -8,7 +8,7 @@
 #include "UObject/ObjectMacros.h"
 
 #include "USDIncludesStart.h"
-	#include "pxr/pxr.h"
+#include "pxr/pxr.h"
 #include "USDIncludesEnd.h"
 
 PXR_NAMESPACE_OPEN_SCOPE
@@ -69,11 +69,11 @@ namespace MeshTranslationImpl
 		VT = 2,
 		TwoSided = 4
 	};
-	ENUM_CLASS_FLAGS( EUsdReferenceMaterialProperties )
+	ENUM_CLASS_FLAGS(EUsdReferenceMaterialProperties)
 
 	// Returns one of the alternatives of the UsdPreviewSurface reference material depending on the material overrides
 	// provided, and nullptr otherwise
-	UMaterialInterface* GetReferencePreviewSurfaceMaterial( EUsdReferenceMaterialProperties ReferenceMaterialProperties );
+	UMaterialInterface* GetReferencePreviewSurfaceMaterial(EUsdReferenceMaterialProperties ReferenceMaterialProperties);
 
 	// Returns the VT version of the provided UsdPreviewSurface ReferenceMaterial. Returns the provided ReferenceMaterial back if
 	// it is already a VT-capable reference material, and returns nullptr if ReferenceMaterial isn't one of our reference material
@@ -81,7 +81,7 @@ namespace MeshTranslationImpl
 	// Example: Receives UsdPreviewSurfaceTwoSided -> Returns UsdPreviewSurfaceTwoSidedVT
 	// Example: Receives UsdPreviewSurfaceTwoSidedVT -> Returns UsdPreviewSurfaceTwoSidedVT
 	// Example: Receives SomeOtherReferenceMaterial -> Returns nullptr
-	UMaterialInterface* GetVTVersionOfReferencePreviewSurfaceMaterial( UMaterialInterface* ReferenceMaterial );
+	UMaterialInterface* GetVTVersionOfReferencePreviewSurfaceMaterial(UMaterialInterface* ReferenceMaterial);
 
 	// Returns the two-sided version of the provided UsdPreviewSurface ReferenceMaterial. Returns the provided ReferenceMaterial
 	// back if it is already a two-sided-capable reference material, and returns nullptr if ReferenceMaterial isn't one of our reference
@@ -89,11 +89,11 @@ namespace MeshTranslationImpl
 	// Example: Receives UsdPreviewSurfaceTranslucent -> Returns UsdPreviewSurfaceTwoSidedTranslucent
 	// Example: Receives UsdPreviewSurfaceTwoSidedTranslucent -> Returns UsdPreviewSurfaceTwoSidedTranslucent
 	// Example: Receives SomeOtherReferenceMaterial -> Returns nullptr
-	UMaterialInterface* GetTwoSidedVersionOfReferencePreviewSurfaceMaterial( UMaterialInterface* ReferenceMaterial );
+	UMaterialInterface* GetTwoSidedVersionOfReferencePreviewSurfaceMaterial(UMaterialInterface* ReferenceMaterial);
 
 	// Returns whether the Material is one of the UsdPreviewSurface reference materials (which can be reassigned by the
 	// user on a per project basis)
 	bool IsReferencePreviewSurfaceMaterial(UMaterialInterface* Material);
-}
+}	 // namespace MeshTranslationImpl
 
-#endif // #if USE_USD_SDK
\ No newline at end of file
+#endif	  // #if USE_USD_SDK
diff --git a/Engine/Plugins/Importers/USDImporter/Source/USDSchemas/Private/USDGeomCameraTranslator.cpp b/Engine/Plugins/Importers/USDImporter/Source/USDSchemas/Private/USDGeomCameraTranslator.cpp
--- a/Engine/Plugins/Importers/USDImporter/Source/USDSchemas/Private/USDGeomCameraTranslator.cpp
+++ b/Engine/Plugins/Importers/USDImporter/Source/USDSchemas/Private/USDGeomCameraTranslator.cpp
@@ -8,17 +8,16 @@
 #include "USDTypesConversion.h"
 
 #include "UsdWrappers/UsdPrim.h"
-#include "UsdWrappers/UsdStage.h"
 
 #include "CineCameraActor.h"
 #include "CineCameraComponent.h"
 
 #if USE_USD_SDK
 #include "USDIncludesStart.h"
-	#include "pxr/usd/usdGeom/camera.h"
+#include "pxr/usd/usd/attribute.h"
+#include "pxr/usd/usd/prim.h"
 #include "USDIncludesEnd.h"
 
-
 USceneComponent* FUsdGeomCameraTranslator::CreateComponents()
 {
 	USceneComponent* Component = nullptr;
@@ -37,30 +36,30 @@ USceneComponent* FUsdGeomCameraTranslator::CreateComponents()
 	{
 		FScopedUsdAllocs UsdAllocs;
 
-		if ( pxr::UsdPrim UsdPrim{ GetPrim() } )
+		if (pxr::UsdPrim UsdPrim{GetPrim()})
 		{
 			pxr::UsdPrim UsdParentPrim = UsdPrim.GetParent();
 
 			// Check if we're a child of the pseudoroot here because if we are USD will emit a warning if we try getting attributes from it...
-			if ( UsdParentPrim && !UsdParentPrim.IsPseudoRoot() )
+			if (UsdParentPrim && !UsdParentPrim.IsPseudoRoot())
 			{
 				bool bIsCineCameraActorMainComponent = false;
-				if ( pxr::UsdAttribute Attr = UsdParentPrim.GetAttribute( UnrealToUsd::ConvertToken( TEXT( "unrealCameraPrimName" ) ).Get() ) )
+				if (pxr::UsdAttribute Attr = UsdParentPrim.GetAttribute(UnrealToUsd::ConvertToken(TEXT("unrealCameraPrimName")).Get()))
 				{
 					pxr::TfToken CameraComponentPrimName;
-					if ( Attr.Get<pxr::TfToken>( &CameraComponentPrimName ) )
+					if (Attr.Get<pxr::TfToken>(&CameraComponentPrimName))
 					{
-						if ( CameraComponentPrimName == UsdPrim.GetName() )
+						if (CameraComponentPrimName == UsdPrim.GetName())
 						{
 							bIsCineCameraActorMainComponent = true;
 						}
 					}
 				}
 
-				if ( Context->ParentComponent && bIsCineCameraActorMainComponent )
+				if (Context->ParentComponent && bIsCineCameraActorMainComponent)
 				{
-					ACineCameraActor* ParentOwnerActor = Cast<ACineCameraActor>( Context->ParentComponent->GetOwner() );
-					if ( ParentOwnerActor && ParentOwnerActor->GetRootComponent() == Context->ParentComponent )
+					ACineCameraActor* ParentOwnerActor = Cast<ACineCameraActor>(Context->ParentComponent->GetOwner());
+					if (ParentOwnerActor && ParentOwnerActor->GetRootComponent() == Context->ParentComponent)
 					{
 						Component = ParentOwnerActor->GetCineCameraComponent();
 					}
@@ -82,41 +81,41 @@ USceneComponent* FUsdGeomCameraTranslator::CreateComponents()
 	return Component;
 }
 
-void FUsdGeomCameraTranslator::UpdateComponents( USceneComponent* SceneComponent )
+void FUsdGeomCameraTranslator::UpdateComponents(USceneComponent* SceneComponent)
 {
-	if ( !SceneComponent )
+	if (!SceneComponent)
 	{
 		return;
 	}
 
-	FUsdGeomXformableTranslator::UpdateComponents( SceneComponent );
+	FUsdGeomXformableTranslator::UpdateComponents(SceneComponent);
 
-	UCineCameraComponent* CameraComponent = Cast<UCineCameraComponent>( SceneComponent );
+	UCineCameraComponent* CameraComponent = Cast<UCineCameraComponent>(SceneComponent);
 
 	// If could be that we're a random Camera prim (and not an exported CineCameraComponent of an ACineCameraActor),
 	// and so we'll have spawned a brand new ACineCameraActor for this prim alone and SceneComponent will be pointing
 	// at its root component (which is just an USceneComponent)
-	if ( !CameraComponent )
+	if (!CameraComponent)
 	{
-		if ( ACineCameraActor* CameraActor = Cast<ACineCameraActor>( SceneComponent->GetOwner() ) )
+		if (ACineCameraActor* CameraActor = Cast<ACineCameraActor>(SceneComponent->GetOwner()))
 		{
-			if ( SceneComponent == CameraActor->GetRootComponent() )
+			if (SceneComponent == CameraActor->GetRootComponent())
 			{
 				CameraComponent = CameraActor->GetCineCameraComponent();
 			}
 		}
 	}
 
-	if ( CameraComponent )
+	if (CameraComponent)
 	{
 		CameraComponent->Modify();
 
 		FScopedUsdAllocs UsdAllocs;
-		UsdToUnreal::ConvertGeomCamera( GetPrim(), *CameraComponent, Context->Time );
+		UsdToUnreal::ConvertGeomCamera(GetPrim(), *CameraComponent, Context->Time);
 	}
 }
 
-bool FUsdGeomCameraTranslator::CollapsesChildren( ECollapsingType CollapsingType ) const
+bool FUsdGeomCameraTranslator::CollapsesChildren(ECollapsingType CollapsingType) const
 {
 	// If we have a custom draw mode, it means we should draw bounds/cards/etc. instead
 	// of our entire subtree, which is basically the same thing as collapsing
@@ -129,9 +128,9 @@ bool FUsdGeomCameraTranslator::CollapsesChildren( ECollapsingType CollapsingType
 	return false;
 }
 
-bool FUsdGeomCameraTranslator::CanBeCollapsed( ECollapsingType CollapsingType ) const
+bool FUsdGeomCameraTranslator::CanBeCollapsed(ECollapsingType CollapsingType) const
 {
 	return false;
 }
 
-#endif // #if USE_USD_SDK
+#endif	  // #if USE_USD_SDK
diff --git a/Engine/Plugins/Importers/USDImporter/Source/USDSchemas/Private/USDGeomMeshTranslator.cpp b/Engine/Plugins/Importers/USDImporter/Source/USDSchemas/Private/USDGeomMeshTranslator.cpp
--- a/Engine/Plugins/Importers/USDImporter/Source/USDSchemas/Private/USDGeomMeshTranslator.cpp
+++ b/Engine/Plugins/Importers/USDImporter/Source/USDSchemas/Private/USDGeomMeshTranslator.cpp
@@ -7,9 +7,9 @@
 #include "MeshTranslationImpl.h"
 #include "UnrealUSDWrapper.h"
 #include "USDAssetUserData.h"
-#include "USDDrawModeComponent.h"
 #include "USDClassesModule.h"
 #include "USDConversionUtils.h"
+#include "USDDrawModeComponent.h"
 #include "USDGeomMeshConversion.h"
 #include "USDInfoCache.h"
 #include "USDLog.h"
@@ -21,7 +21,6 @@
 #include "UsdWrappers/UsdStage.h"
 
 #include "CompGeom/FitKDOP3.h"
-#include "Components/StaticMeshComponent.h"
 #include "Engine/CollisionProfile.h"
 #include "Engine/Level.h"
 #include "Engine/StaticMesh.h"
@@ -45,32 +44,34 @@
 #include "ConvexDecompTool.h"
 #include "IMeshBuilderModule.h"
 #include "MeshBudgetProjectSettings.h"
-#endif // WITH_EDITOR
+#endif	  // WITH_EDITOR
 
 #include "USDIncludesStart.h"
-	#include "pxr/usd/usd/prim.h"
-	#include "pxr/usd/usd/stage.h"
-	#include "pxr/usd/usd/typed.h"
-	#include "pxr/usd/usdGeom/mesh.h"
-	#include "pxr/usd/usdGeom/subset.h"
-	#include "pxr/usd/usdSkel/bindingAPI.h"
-	#include "pxr/usd/usdSkel/root.h"
-	#include "pxr/usd/usdPhysics/collisionAPI.h"
-	#include "pxr/usd/usdPhysics/meshCollisionAPI.h"
-	#include "pxr/usd/usdPhysics/tokens.h"
+#include "pxr/usd/usd/prim.h"
+#include "pxr/usd/usd/stage.h"
+#include "pxr/usd/usd/typed.h"
+#include "pxr/usd/usdGeom/mesh.h"
+#include "pxr/usd/usdGeom/subset.h"
+#include "pxr/usd/usdPhysics/collisionAPI.h"
+#include "pxr/usd/usdPhysics/meshCollisionAPI.h"
+#include "pxr/usd/usdPhysics/tokens.h"
+#include "pxr/usd/usdSkel/bindingAPI.h"
 #include "USDIncludesEnd.h"
 
 static float GMeshNormalRepairThreshold = 0.05f;
 static FAutoConsoleVariableRef CVarMeshNormalRepairThreshold(
 	TEXT("USD.MeshNormalRepairThreshold"),
 	GMeshNormalRepairThreshold,
-	TEXT("We will try repairing up to this fraction of a Mesh's normals when invalid. If a Mesh has more invalid normals than this, we will recompute all of them. Defaults to 0.05 (5% of all normals)."));
+	TEXT("We will try repairing up to this fraction of a Mesh's normals when invalid. If a Mesh has more invalid normals than this, we will "
+		 "recompute all of them. Defaults to 0.05 (5% of all normals).")
+);
 
 static bool GSkipMeshTangentComputation = false;
 static FAutoConsoleVariableRef CVarSkipMeshTangentComputation(
 	TEXT("USD.SkipMeshTangentComputation"),
 	GSkipMeshTangentComputation,
-	TEXT("Skip computing tangents for meshes. With meshes with a huge numer of vertices, it can take a very long time to compute them."));
+	TEXT("Skip computing tangents for meshes. With meshes with a huge numer of vertices, it can take a very long time to compute them.")
+);
 
 namespace UsdGeomMeshTranslatorImpl
 {
@@ -107,16 +108,21 @@ namespace UsdGeomMeshTranslatorImpl
 				if (OverrideValue == UnrealIdentifiers::UnrealNaniteOverrideEnable)
 				{
 					bHasNaniteOverrideEnabled = true;
-					UE_LOG(LogUsd, Log, TEXT("Trying to enable Nanite for mesh generated for prim '%s' as the '%s' attribute is set to '%s'"),
+					UE_LOG(
+						LogUsd,
+						Log,
+						TEXT("Trying to enable Nanite for mesh generated for prim '%s' as the '%s' attribute is set to '%s'"),
 						*PrimPath.GetString(),
 						*UsdToUnreal::ConvertToken(UnrealIdentifiers::UnrealNaniteOverride),
 						*UsdToUnreal::ConvertToken(UnrealIdentifiers::UnrealNaniteOverrideEnable)
 					);
-
 				}
 				else if (OverrideValue == UnrealIdentifiers::UnrealNaniteOverrideDisable)
 				{
-					UE_LOG(LogUsd, Log, TEXT("Not enabling Nanite for mesh generated for prim '%s' as the '%s' attribute is set to '%s'"),
+					UE_LOG(
+						LogUsd,
+						Log,
+						TEXT("Not enabling Nanite for mesh generated for prim '%s' as the '%s' attribute is set to '%s'"),
 						*PrimPath.GetString(),
 						*UsdToUnreal::ConvertToken(UnrealIdentifiers::UnrealNaniteOverride),
 						*UsdToUnreal::ConvertToken(UnrealIdentifiers::UnrealNaniteOverrideDisable)
@@ -132,7 +138,10 @@ namespace UsdGeomMeshTranslatorImpl
 			const int32 NumTriangles = LODIndexToMeshDescription[0].Triangles().Num();
 			if (NumTriangles >= Context.NaniteTriangleThreshold)
 			{
-				UE_LOG(LogUsd, Verbose, TEXT("Trying to enable Nanite for mesh generated for prim '%s' as it has '%d' triangles, and the threshold is '%d'"),
+				UE_LOG(
+					LogUsd,
+					Verbose,
+					TEXT("Trying to enable Nanite for mesh generated for prim '%s' as it has '%d' triangles, and the threshold is '%d'"),
 					*PrimPath.GetString(),
 					NumTriangles,
 					Context.NaniteTriangleThreshold
@@ -140,7 +149,10 @@ namespace UsdGeomMeshTranslatorImpl
 			}
 			else
 			{
-				UE_LOG(LogUsd, Verbose, TEXT("Not enabling Nanite for mesh generated for prim '%s' as it has '%d' triangles, and the threshold is '%d'"),
+				UE_LOG(
+					LogUsd,
+					Verbose,
+					TEXT("Not enabling Nanite for mesh generated for prim '%s' as it has '%d' triangles, and the threshold is '%d'"),
 					*PrimPath.GetString(),
 					NumTriangles,
 					Context.NaniteTriangleThreshold
@@ -156,7 +168,11 @@ namespace UsdGeomMeshTranslatorImpl
 		// else other than LOD.
 		if (LODIndexToMeshDescription.Num() > 1)
 		{
-			UE_LOG(LogUsd, Warning, TEXT("Not enabling Nanite for mesh generated for prim '%s' as it has more than one generated LOD (and so came from a LOD variant set setup)"),
+			UE_LOG(
+				LogUsd,
+				Warning,
+				TEXT("Not enabling Nanite for mesh generated for prim '%s' as it has more than one generated LOD (and so came from a LOD variant set "
+					 "setup)"),
 				*PrimPath.GetString()
 			);
 			return false;
@@ -166,10 +182,14 @@ namespace UsdGeomMeshTranslatorImpl
 		{
 			TOptional<uint64> SubtreeSectionCount = Context.InfoCache->GetSubtreeMaterialSlotCount(PrimPath);
 
-			const int32 MaxNumSections = 64; // There is no define for this, but it's checked for on NaniteBuilder.cpp, FBuilderModule::Build
+			const int32 MaxNumSections = 64;	// There is no define for this, but it's checked for on NaniteBuilder.cpp, FBuilderModule::Build
 			if (!SubtreeSectionCount.IsSet() || SubtreeSectionCount.GetValue() > MaxNumSections)
 			{
-				UE_LOG(LogUsd, Warning, TEXT("Not enabling Nanite for mesh generated for prim '%s' as LOD0 has '%d' material slots, which is above the Nanite limit of '%d'"),
+				UE_LOG(
+					LogUsd,
+					Warning,
+					TEXT("Not enabling Nanite for mesh generated for prim '%s' as LOD0 has '%d' material slots, which is above the Nanite limit of "
+						 "'%d'"),
 					*PrimPath.GetString(),
 					SubtreeSectionCount.GetValue(),
 					MaxNumSections
@@ -179,7 +199,10 @@ namespace UsdGeomMeshTranslatorImpl
 		}
 
 #if !WITH_EDITOR
-		UE_LOG(LogUsd, Warning, TEXT("Not enabling Nanite for mesh generated for prim '%s' as we can't setup Nanite during runtime"),
+		UE_LOG(
+			LogUsd,
+			Warning,
+			TEXT("Not enabling Nanite for mesh generated for prim '%s' as we can't setup Nanite during runtime"),
 			*PrimPath.GetString()
 		);
 		return false;
@@ -225,7 +248,7 @@ namespace UsdGeomMeshTranslatorImpl
 		uint32 StaticMeshSlotIndex = 0;
 		for (int32 LODIndex = 0; LODIndex < LODIndexToMaterialInfo.Num(); ++LODIndex)
 		{
-			const TArray< UsdUtils::FUsdPrimMaterialSlot >& LODSlots = LODIndexToMaterialInfo[LODIndex].Slots;
+			const TArray<UsdUtils::FUsdPrimMaterialSlot>& LODSlots = LODIndexToMaterialInfo[LODIndex].Slots;
 
 			for (int32 LODSlotIndex = 0; LODSlotIndex < LODSlots.Num(); ++LODSlotIndex, ++StaticMeshSlotIndex)
 			{
@@ -238,7 +261,15 @@ namespace UsdGeomMeshTranslatorImpl
 				}
 				else
 				{
-					UE_LOG(LogUsd, Error, TEXT("Failed to resolve material '%s' for slot '%d' of LOD '%d' for mesh '%s'"), *Slot.MaterialSource, LODSlotIndex, LODIndex, *UsdToUnreal::ConvertPath(UsdPrim.GetPath()));
+					UE_LOG(
+						LogUsd,
+						Error,
+						TEXT("Failed to resolve material '%s' for slot '%d' of LOD '%d' for mesh '%s'"),
+						*Slot.MaterialSource,
+						LODSlotIndex,
+						LODIndex,
+						*UsdToUnreal::ConvertPath(UsdPrim.GetPath())
+					);
 					continue;
 				}
 
@@ -279,21 +310,21 @@ namespace UsdGeomMeshTranslatorImpl
 
 					bMaterialAssignementsHaveChanged = true;
 				}
-#endif // WITH_EDITOR
+#endif	  // WITH_EDITOR
 			}
 		}
 
 #if WITH_EDITOR
 		StaticMesh.GetOriginalSectionInfoMap().CopyFrom(StaticMesh.GetSectionInfoMap());
-#endif // WITH_EDITOR
+#endif	  // WITH_EDITOR
 
 		return bMaterialAssignementsHaveChanged;
 	}
 
 	// If UsdMesh is a LOD, will parse it and all of the other LODs, and and place them in OutLODIndexToMeshDescription and OutLODIndexToMaterialInfo.
-	// Note that these other LODs will be hidden in other variants, and won't show up on traversal unless we actively switch the variants (which we do here).
-	// We use a separate function for this because there is a very specific set of conditions where we successfully can do this, and we
-	// want to fall back to just parsing UsdMesh as a simple single-LOD mesh if we fail.
+	// Note that these other LODs will be hidden in other variants, and won't show up on traversal unless we actively switch the variants (which we do
+	// here). We use a separate function for this because there is a very specific set of conditions where we successfully can do this, and we want to
+	// fall back to just parsing UsdMesh as a simple single-LOD mesh if we fail.
 	bool TryLoadingMultipleLODs(
 		const UE::FUsdPrim& MeshPrim,
 		TArray<FMeshDescription>& OutLODIndexToMeshDescription,
@@ -303,7 +334,7 @@ namespace UsdGeomMeshTranslatorImpl
 		const FUsdMetadataImportOptions& MetadataOptions
 	)
 	{
-		if(!MeshPrim)
+		if (!MeshPrim)
 		{
 			return false;
 		}
@@ -321,30 +352,26 @@ namespace UsdGeomMeshTranslatorImpl
 		// Here we are choosing to preemptively traverse all LODs to fetch the combined primvars that we'll use for
 		// each UV index for the static mesh as a whole
 		TFunction<bool(const pxr::UsdGeomMesh&, int32)> CombinePrimvars =
-			[&AllPrimvars, &PreferredPrimvars]
-			(const pxr::UsdGeomMesh& LODMesh, int32 LODIndex)
-			{
-				TArray<TUsdStore<pxr::UsdGeomPrimvar>> MeshPrimvars = UsdUtils::GetUVSetPrimvars(
-					LODMesh.GetPrim(),
-					TNumericLimits<int32>::Max()
-				);
+			[&AllPrimvars, &PreferredPrimvars](const pxr::UsdGeomMesh& LODMesh, int32 LODIndex)
+		{
+			TArray<TUsdStore<pxr::UsdGeomPrimvar>> MeshPrimvars = UsdUtils::GetUVSetPrimvars(LODMesh.GetPrim(), TNumericLimits<int32>::Max());
 
-				for (const TUsdStore<pxr::UsdGeomPrimvar>& MeshPrimvar : MeshPrimvars)
-				{
-					FString PrimvarName = UsdToUnreal::ConvertToken(MeshPrimvar.Get().GetName());
-					PrimvarName.RemoveFromStart(TEXT("primvars:"));
+			for (const TUsdStore<pxr::UsdGeomPrimvar>& MeshPrimvar : MeshPrimvars)
+			{
+				FString PrimvarName = UsdToUnreal::ConvertToken(MeshPrimvar.Get().GetName());
+				PrimvarName.RemoveFromStart(TEXT("primvars:"));
 
-					AllPrimvars.Add(PrimvarName);
+				AllPrimvars.Add(PrimvarName);
 
-					// Keep track of which primvars are texCoord2f as we always want to prefer these over other float2s
-					if (MeshPrimvar.Get().GetTypeName().GetRole() == pxr::SdfValueTypeNames->TexCoord2f.GetRole())
-					{
-						PreferredPrimvars.Add(PrimvarName);
-					}
+				// Keep track of which primvars are texCoord2f as we always want to prefer these over other float2s
+				if (MeshPrimvar.Get().GetTypeName().GetRole() == pxr::SdfValueTypeNames->TexCoord2f.GetRole())
+				{
+					PreferredPrimvars.Add(PrimvarName);
 				}
+			}
 
-				return true;
-			};
+			return true;
+		};
 		if (!UsdUtils::IterateLODMeshes(ParentPrim, CombinePrimvars))
 		{
 			return false;
@@ -497,17 +524,12 @@ namespace UsdGeomMeshTranslatorImpl
 			{
 				UsdToUnreal::FUsdMeshConversionOptions OptionsCopy = Options;
 				OptionsCopy.AdditionalTransform = MeshTransform * Options.AdditionalTransform;
-				OptionsCopy.bMergeIdenticalMaterialSlots = false;  // We only merge for collapsed meshes
+				OptionsCopy.bMergeIdenticalMaterialSlots = false;	 // We only merge for collapsed meshes
 
 				FScopedUsdAllocs Allocs;
 				pxr::UsdGeomMesh UsdMesh{MeshPrim};
 
-				bSuccess &= UsdToUnreal::ConvertGeomMesh(
-					UsdMesh,
-					TempMeshDescription,
-					TempMaterialInfo,
-					OptionsCopy
-				);
+				bSuccess &= UsdToUnreal::ConvertGeomMesh(UsdMesh, TempMeshDescription, TempMaterialInfo, OptionsCopy);
 			}
 
 			if (bSuccess)
@@ -560,11 +582,17 @@ namespace UsdGeomMeshTranslatorImpl
 		// Make sure our normals can be rebuilt from MeshDescription::InitializeAutoGeneratedAttributes in case some tool needs them.
 		// Always force-compute tangents here as we never have them anyway. If we don't force them to be recomputed we'll get
 		// the worst of both worlds as some of these will be arbitrarily recomputed anyway, and some will be left invalid
-		EComputeNTBsFlags Options = GSkipMeshTangentComputation ? EComputeNTBsFlags::None : EComputeNTBsFlags::UseMikkTSpace | EComputeNTBsFlags::Tangents;
+		EComputeNTBsFlags Options = GSkipMeshTangentComputation ? EComputeNTBsFlags::None
+																: EComputeNTBsFlags::UseMikkTSpace | EComputeNTBsFlags::Tangents;
 		if (InvalidNormalFraction >= GMeshNormalRepairThreshold)
 		{
 			Options |= EComputeNTBsFlags::Normals;
-			UE_LOG(LogUsd, Log, TEXT("%f%% of the normals from Mesh prim '%s' are invalid or unusable. This is at or above the threshold of '%f%%' (configurable via the cvar '%s'), so normals will be discarded and fully recomputed. Note that when the cvar 'USD.Subdiv.IgnoreNormalsWhenSubdividing' is true it is expected for subdivision meshes to have their normals discarded."),
+			UE_LOG(
+				LogUsd,
+				Log,
+				TEXT("%f%% of the normals from Mesh prim '%s' are invalid or unusable. This is at or above the threshold of '%f%%' (configurable via "
+					 "the cvar '%s'), so normals will be discarded and fully recomputed. Note that when the cvar "
+					 "'USD.Subdiv.IgnoreNormalsWhenSubdividing' is true it is expected for subdivision meshes to have their normals discarded."),
 				InvalidNormalFraction * 100.0f,
 				*PrimPath,
 				GMeshNormalRepairThreshold * 100.0f,
@@ -573,7 +601,11 @@ namespace UsdGeomMeshTranslatorImpl
 		}
 		else if (InvalidNormalFraction > 0)
 		{
-			UE_LOG(LogUsd, Log, TEXT("%f%% of the normals from Mesh prim '%s' are invalid or unusable. This is below the threshold of '%f%%' (configurable via the cvar '%s'), so the invalid normals will be repaired."),
+			UE_LOG(
+				LogUsd,
+				Log,
+				TEXT("%f%% of the normals from Mesh prim '%s' are invalid or unusable. This is below the threshold of '%f%%' (configurable via the "
+					 "cvar '%s'), so the invalid normals will be repaired."),
 				InvalidNormalFraction * 100.0f,
 				*PrimPath,
 				GMeshNormalRepairThreshold * 100.0f,
@@ -614,18 +646,18 @@ namespace UsdGeomMeshTranslatorImpl
 
 			// Put whether we want Nanite or not within the hash, so that the user could have one instance of the mesh without Nanite and another
 			// with Nanite if they want to (using the override parameters). This also nicely handles a couple of edge cases:
-			//	- What if we change a mesh from having Nanite disabled to enabled, or vice-versa (e.g. by changing the threshold)? We'd reuse the mesh from the asset cache
+			//	- What if we change a mesh from having Nanite disabled to enabled, or vice-versa (e.g. by changing the threshold)? We'd reuse the mesh
+			// from the asset cache
 			//    in that case, so we'd need to rebuild it;
 			//  - What if multiple meshes on the scene hash the same, but only one of them has a Nanite override attribute?
-			// If we always enabled Nanite when either the mesh from the asset cache or the new prim wanted it, we wouldn't be able to turn Nanite off from
-			// a single mesh that once had it enabled: It would always find the old Nanite-enabled mesh on the cache and leave it enabled.
-			// If we always set Nanite to whatever the current prim wants, we could handle a single mesh turning Nanite on/off alright, but then we can't handle
-			// the case where multiple meshes on the scene hash the same and only one of them has the override: The last prim would win, and they'd all be randomly either
-			// enabled or disabled.
-			// Note that we could also fix these problems by trying to check if a mesh is reused due to being in the cache from an old instance of the stage, or due to being used by
-			// another prim, but that doesn't seem like a good path to go down
-			// Additionally, hashing this bool also prevents us from having to force-rebuild a mesh to switch its Nanite flag, which could be tricky to do since some of
-			// these build steps are async/thread-pool based.
+			// If we always enabled Nanite when either the mesh from the asset cache or the new prim wanted it, we wouldn't be able to turn Nanite off
+			// from a single mesh that once had it enabled: It would always find the old Nanite-enabled mesh on the cache and leave it enabled. If we
+			// always set Nanite to whatever the current prim wants, we could handle a single mesh turning Nanite on/off alright, but then we can't
+			// handle the case where multiple meshes on the scene hash the same and only one of them has the override: The last prim would win, and
+			// they'd all be randomly either enabled or disabled. Note that we could also fix these problems by trying to check if a mesh is reused
+			// due to being in the cache from an old instance of the stage, or due to being used by another prim, but that doesn't seem like a good
+			// path to go down Additionally, hashing this bool also prevents us from having to force-rebuild a mesh to switch its Nanite flag, which
+			// could be tricky to do since some of these build steps are async/thread-pool based.
 			SHA1.Update(reinterpret_cast<const uint8*>(&bShouldEnableNanite), sizeof(bShouldEnableNanite));
 
 			// Hash the threshhold so that if we update it and reload we'll regenerate static meshes
@@ -639,7 +671,7 @@ namespace UsdGeomMeshTranslatorImpl
 
 		if (Context.AssetCache)
 		{
-			StaticMesh = Cast< UStaticMesh >(Context.AssetCache->GetCachedAsset(PrefixedAssetHash));
+			StaticMesh = Cast<UStaticMesh>(Context.AssetCache->GetCachedAsset(PrefixedAssetHash));
 		}
 
 		if (!StaticMesh && bHasValidMeshDescription)
@@ -651,7 +683,11 @@ namespace UsdGeomMeshTranslatorImpl
 				UStaticMesh::StaticClass(),
 				*IUsdClassesModule::SanitizeObjectName(FPaths::GetBaseFilename(MeshName))
 			);
-			StaticMesh = NewObject< UStaticMesh >(GetTransientPackage(), AssetName, Context.ObjectFlags | EObjectFlags::RF_Public | EObjectFlags::RF_Transient);
+			StaticMesh = NewObject<UStaticMesh>(
+				GetTransientPackage(),
+				AssetName,
+				Context.ObjectFlags | EObjectFlags::RF_Public | EObjectFlags::RF_Transient
+			);
 
 #if WITH_EDITOR
 			for (int32 LODIndex = 0; LODIndex < LODIndexToMeshDescription.Num(); ++LODIndex)
@@ -673,7 +709,7 @@ namespace UsdGeomMeshTranslatorImpl
 
 			FMeshBudgetProjectSettingsUtils::SetLodGroupForStaticMesh(StaticMesh);
 
-#endif // WITH_EDITOR
+#endif	  // WITH_EDITOR
 
 			StaticMesh->SetLightingGuid();
 
@@ -684,7 +720,7 @@ namespace UsdGeomMeshTranslatorImpl
 		}
 		else
 		{
-			//FPlatformMisc::LowLevelOutputDebugStringf( TEXT("Mesh found in cache %s\n"), *StaticMesh->GetName() );
+			// FPlatformMisc::LowLevelOutputDebugStringf( TEXT("Mesh found in cache %s\n"), *StaticMesh->GetName() );
 			bOutIsNew = false;
 		}
 
@@ -701,9 +737,9 @@ namespace UsdGeomMeshTranslatorImpl
 			StaticMesh.ReleaseResourcesFence.Wait();
 		}
 
-		StaticMesh.SetRenderData(MakeUnique< FStaticMeshRenderData >());
+		StaticMesh.SetRenderData(MakeUnique<FStaticMeshRenderData>());
 		StaticMesh.CreateBodySetup();
-		StaticMesh.MarkAsNotHavingNavigationData(); // Needed or else it will warn if we try cooking with body setup
+		StaticMesh.MarkAsNotHavingNavigationData();	   // Needed or else it will warn if we try cooking with body setup
 	}
 
 	bool BuildStaticMesh(UStaticMesh& StaticMesh, const FStaticFeatureLevel& FeatureLevel, TArray<FMeshDescription>& LODIndexToMeshDescription)
@@ -731,7 +767,8 @@ namespace UsdGeomMeshTranslatorImpl
 			FStaticMeshLODResources& LODResources = StaticMesh.GetRenderData()->LODResources[LODIndex];
 
 			FMeshDescription& MeshDescription = LODIndexToMeshDescription[LODIndex];
-			TVertexInstanceAttributesConstRef< FVector4f > MeshDescriptionColors = MeshDescription.VertexInstanceAttributes().GetAttributesRef<FVector4f>(MeshAttribute::VertexInstance::Color);
+			TVertexInstanceAttributesConstRef<FVector4f>
+				MeshDescriptionColors = MeshDescription.VertexInstanceAttributes().GetAttributesRef<FVector4f>(MeshAttribute::VertexInstance::Color);
 
 			// Compute normals here if necessary because they're not going to be computed via the regular static mesh build pipeline at runtime
 			// (i.e. StaticMeshBuilder is not available at runtime)
@@ -746,7 +783,7 @@ namespace UsdGeomMeshTranslatorImpl
 			StaticMesh.BuildFromMeshDescription(MeshDescription, LODResources);
 		}
 
-#endif // WITH_EDITOR
+#endif	  // WITH_EDITOR
 
 		return true;
 	}
@@ -767,15 +804,15 @@ namespace UsdGeomMeshTranslatorImpl
 				StaticMesh.GetRenderData()->Bounds = MeshDescription->GetBounds();
 			}
 			StaticMesh.CalculateExtendedBounds();
-			StaticMesh.ClearMeshDescriptions(); // Clear mesh descriptions to reduce memory usage, they are kept only in bulk data form
+			StaticMesh.ClearMeshDescriptions();	   // Clear mesh descriptions to reduce memory usage, they are kept only in bulk data form
 #else
 			// Fetch the MeshDescription from the imported LODIndexToMeshDescription as StaticMesh.GetMeshDescription is editor-only
 			StaticMesh.GetRenderData()->Bounds = LODIndexToMeshDescription[0].GetBounds();
 			StaticMesh.CalculateExtendedBounds();
-#endif // WITH_EDITOR
+#endif											   // WITH_EDITOR
 		}
 	}
-}
+}	 // namespace UsdGeomMeshTranslatorImpl
 
 namespace UE::UsdCollision::Private
 {
@@ -789,23 +826,43 @@ namespace UE::UsdCollision::Private
 		virtual int32 GetVertexIndex(int32 Index) const = 0;
 		virtual FBox GetBoundingBox() const = 0;
 		virtual bool IsValid() const = 0;
-		virtual ~IMeshData() {}
+		virtual ~IMeshData()
+		{
+		}
 	};
 
 	class FMeshDescriptionWrapper : public IMeshData
 	{
 	public:
 		FMeshDescriptionWrapper(const FMeshDescription& InMeshDescription)
-		: MeshDescription(InMeshDescription)
+			: MeshDescription(InMeshDescription)
 		{
 		}
 
-		int32 GetNumVertices() const override { return MeshDescription.Vertices().Num(); }
-		int32 GetNumIndices() const override { return MeshDescription.VertexInstances().Num(); }
-		FVector3f GetVertexPosition(int32 Index) const override { return MeshDescription.GetVertexPositions()[Index]; }
-		int32 GetVertexIndex(int32 Index) const override { return MeshDescription.GetVertexInstanceVertex(Index); }
-		FBox GetBoundingBox() const override { return MeshDescription.ComputeBoundingBox(); }
-		bool IsValid() const override { return MeshDescription.Vertices().Num() > 0; };
+		int32 GetNumVertices() const override
+		{
+			return MeshDescription.Vertices().Num();
+		}
+		int32 GetNumIndices() const override
+		{
+			return MeshDescription.VertexInstances().Num();
+		}
+		FVector3f GetVertexPosition(int32 Index) const override
+		{
+			return MeshDescription.GetVertexPositions()[Index];
+		}
+		int32 GetVertexIndex(int32 Index) const override
+		{
+			return MeshDescription.GetVertexInstanceVertex(Index);
+		}
+		FBox GetBoundingBox() const override
+		{
+			return MeshDescription.ComputeBoundingBox();
+		}
+		bool IsValid() const override
+		{
+			return MeshDescription.Vertices().Num() > 0;
+		};
 
 	private:
 		const FMeshDescription& MeshDescription;
@@ -815,16 +872,31 @@ namespace UE::UsdCollision::Private
 	{
 	public:
 		FRenderDataWrapper(const FStaticMeshLODResources& InRenderData)
-		: RenderData(InRenderData)
+			: RenderData(InRenderData)
 		{
 		}
 
-		int32 GetNumVertices() const override { return RenderData.GetNumVertices(); }
-		int32 GetNumIndices() const override { return RenderData.IndexBuffer.GetNumIndices(); }
-		FVector3f GetVertexPosition(int32 Index) const override { return RenderData.VertexBuffers.PositionVertexBuffer.VertexPosition(Index); }
-		int32 GetVertexIndex(int32 Index) const override { return RenderData.IndexBuffer.GetIndex(Index); }
-		bool IsValid() const override { return RenderData.GetNumVertices() > 0; };
-		FBox GetBoundingBox() const override 
+		int32 GetNumVertices() const override
+		{
+			return RenderData.GetNumVertices();
+		}
+		int32 GetNumIndices() const override
+		{
+			return RenderData.IndexBuffer.GetNumIndices();
+		}
+		FVector3f GetVertexPosition(int32 Index) const override
+		{
+			return RenderData.VertexBuffers.PositionVertexBuffer.VertexPosition(Index);
+		}
+		int32 GetVertexIndex(int32 Index) const override
+		{
+			return RenderData.IndexBuffer.GetIndex(Index);
+		}
+		bool IsValid() const override
+		{
+			return RenderData.GetNumVertices() > 0;
+		};
+		FBox GetBoundingBox() const override
 		{
 			if (!BoundingBox.IsValid)
 			{
@@ -845,7 +917,7 @@ namespace UE::UsdCollision::Private
 	{
 	public:
 		FMeshData(const UStaticMesh& StaticMesh)
-		: MeshData(nullptr)
+			: MeshData(nullptr)
 		{
 #if WITH_EDITORONLY_DATA
 			FMeshDescription* MeshDescription = StaticMesh.GetMeshDescription(0);
@@ -864,7 +936,7 @@ namespace UE::UsdCollision::Private
 		}
 
 		FMeshData(const FMeshDescription& MeshDescription)
-		: MeshData(nullptr)
+			: MeshData(nullptr)
 		{
 			if (MeshDescription.Vertices().Num() > 0)
 			{
@@ -877,12 +949,30 @@ namespace UE::UsdCollision::Private
 			delete MeshData;
 		}
 
-		int32 GetNumVertices() const override { return IsValid() ? MeshData->GetNumVertices() : 0; }
-		int32 GetNumIndices() const override { return IsValid() ? MeshData->GetNumIndices() : 0; }
-		FVector3f GetVertexPosition(int32 Index) const override { return IsValid() ? MeshData->GetVertexPosition(Index) : FVector3f::ZeroVector; }
-		int32 GetVertexIndex(int32 Index) const override { return IsValid() ? MeshData->GetVertexIndex(Index) : 0; }
-		FBox GetBoundingBox() const override { return IsValid() ? MeshData->GetBoundingBox() : FBox(); }
-		bool IsValid() const override { return MeshData ? MeshData->IsValid() : false; }
+		int32 GetNumVertices() const override
+		{
+			return IsValid() ? MeshData->GetNumVertices() : 0;
+		}
+		int32 GetNumIndices() const override
+		{
+			return IsValid() ? MeshData->GetNumIndices() : 0;
+		}
+		FVector3f GetVertexPosition(int32 Index) const override
+		{
+			return IsValid() ? MeshData->GetVertexPosition(Index) : FVector3f::ZeroVector;
+		}
+		int32 GetVertexIndex(int32 Index) const override
+		{
+			return IsValid() ? MeshData->GetVertexIndex(Index) : 0;
+		}
+		FBox GetBoundingBox() const override
+		{
+			return IsValid() ? MeshData->GetBoundingBox() : FBox();
+		}
+		bool IsValid() const override
+		{
+			return MeshData ? MeshData->IsValid() : false;
+		}
 
 	private:
 		IMeshData* MeshData;
@@ -893,40 +983,38 @@ namespace UE::UsdCollision::Private
 	// k-DOP (k-Discrete Oriented Polytopes) Direction Vectors
 	constexpr float RCP_SQRT2 = 0.70710678118654752440084436210485f;
 
-	TArray<FVector> KDopDir18 = 
-	{
-		FVector( 1.f, 0.f, 0.f),
+	TArray<FVector> KDopDir18 = {
+		FVector(1.f, 0.f, 0.f),
 		FVector(-1.f, 0.f, 0.f),
-		FVector( 0.f, 1.f, 0.f),
-		FVector( 0.f,-1.f, 0.f),
-		FVector( 0.f, 0.f, 1.f),
-		FVector( 0.f, 0.f,-1.f),
-		FVector( 0.f, RCP_SQRT2,  RCP_SQRT2),
-		FVector( 0.f,-RCP_SQRT2, -RCP_SQRT2),
-		FVector( 0.f, RCP_SQRT2, -RCP_SQRT2),
-		FVector( 0.f,-RCP_SQRT2,  RCP_SQRT2),
-		FVector( RCP_SQRT2, 0.f,  RCP_SQRT2),
+		FVector(0.f, 1.f, 0.f),
+		FVector(0.f, -1.f, 0.f),
+		FVector(0.f, 0.f, 1.f),
+		FVector(0.f, 0.f, -1.f),
+		FVector(0.f, RCP_SQRT2, RCP_SQRT2),
+		FVector(0.f, -RCP_SQRT2, -RCP_SQRT2),
+		FVector(0.f, RCP_SQRT2, -RCP_SQRT2),
+		FVector(0.f, -RCP_SQRT2, RCP_SQRT2),
+		FVector(RCP_SQRT2, 0.f, RCP_SQRT2),
 		FVector(-RCP_SQRT2, 0.f, -RCP_SQRT2),
-		FVector( RCP_SQRT2, 0.f, -RCP_SQRT2),
-		FVector(-RCP_SQRT2, 0.f,  RCP_SQRT2),
-		FVector( RCP_SQRT2,  RCP_SQRT2, 0.f),
+		FVector(RCP_SQRT2, 0.f, -RCP_SQRT2),
+		FVector(-RCP_SQRT2, 0.f, RCP_SQRT2),
+		FVector(RCP_SQRT2, RCP_SQRT2, 0.f),
 		FVector(-RCP_SQRT2, -RCP_SQRT2, 0.f),
-		FVector( RCP_SQRT2, -RCP_SQRT2, 0.f),
-		FVector(-RCP_SQRT2,  RCP_SQRT2, 0.f)
-	};
+		FVector(RCP_SQRT2, -RCP_SQRT2, 0.f),
+		FVector(-RCP_SQRT2, RCP_SQRT2, 0.f)};
 
-	void GenerateKDopAsSimpleCollision(const FMeshData& MeshData, const TArray<FVector> &Dirs, FKAggregateGeom& CollisionShapes)
+	void GenerateKDopAsSimpleCollision(const FMeshData& MeshData, const TArray<FVector>& Dirs, FKAggregateGeom& CollisionShapes)
 	{
 		TArray<FVector> HullVertices;
 		UE::Geometry::FitKDOPVertices3<double>(
-			Dirs, 
+			Dirs,
 			MeshData.GetNumVertices(),
-			[&MeshData](int32 VertexId) 
+			[&MeshData](int32 VertexId)
 			{
 				return static_cast<FVector>(MeshData.GetVertexPosition(VertexId));
 			},
-			HullVertices);
-
+			HullVertices
+		);
 
 		FKConvexElem ConvexElem;
 		ConvexElem.VertexData = HullVertices;
@@ -1000,7 +1088,7 @@ namespace UE::UsdCollision::Private
 			}
 		}
 
-		const FVector Extremes[3] = { (MaxIx[0] - MinIx[0]), (MaxIx[1] - MinIx[1]), (MaxIx[2] - MinIx[2]) };
+		const FVector Extremes[3] = {(MaxIx[0] - MinIx[0]), (MaxIx[1] - MinIx[1]), (MaxIx[2] - MinIx[2])};
 
 		// Now find extreme points furthest apart, and initial center and radius of sphere.
 		float MaxDist2 = 0.f;
@@ -1068,13 +1156,13 @@ namespace UE::UsdCollision::Private
 		CalcBoundingSphere(MeshData, Sphere);
 		CalcBoundingSphere2(MeshData, Sphere2);
 
-		if(Sphere.W < Sphere2.W)
+		if (Sphere.W < Sphere2.W)
 			BestSphere = Sphere;
 		else
 			BestSphere = Sphere2;
 
 		// Don't use if radius is zero.
-		if(BestSphere.W <= 0.f)
+		if (BestSphere.W <= 0.f)
 		{
 			return;
 		}
@@ -1134,7 +1222,7 @@ namespace UE::UsdCollision::Private
 			FVector CenterToPoint = static_cast<FVector>(MeshData.GetVertexPosition(VertexID)) - Sphere.Center;
 			CenterToPoint = Rotation.UnrotateVector(CenterToPoint);
 
-			const float PointRadius2 = CenterToPoint.SizeSquared2D();	// Ignore Z here...
+			const float PointRadius2 = CenterToPoint.SizeSquared2D();	 // Ignore Z here...
 
 			// If this point is outside our current bounding sphere's radius
 			if (PointRadius2 > Radius2)
@@ -1169,7 +1257,8 @@ namespace UE::UsdCollision::Private
 					FVector ClosestPoint;
 					FMath::SphereDistToLine(Origin, Radius, CenterToPoint, (bFlip ? FVector(0.f, 0.f, 1.f) : FVector(0.f, 0.f, -1.f)), ClosestPoint);
 
-					// Don't accept zero as a valid diff when we know it's outside the sphere (saves needless retest on further iterations of like points)
+					// Don't accept zero as a valid diff when we know it's outside the sphere (saves needless retest on further iterations of like
+					// points)
 					HalfLength += FMath::Max<float>(FMath::Abs(CenterToPoint.Z - ClosestPoint.Z), 1.e-6f);
 				}
 			}
@@ -1239,7 +1328,7 @@ namespace UE::UsdCollision::Private
 	void CollectCustomCollisionShapes(const pxr::UsdPrim& UsdPrim, FKAggregateGeom& CollisionShapes)
 	{
 		UsdToUnreal::FUsdMeshConversionOptions Options;
-		Options.PurposesToLoad = EUsdPurpose::Guide; // custom collision mesh must have guide purpose
+		Options.PurposesToLoad = EUsdPurpose::Guide;	// custom collision mesh must have guide purpose
 		Options.bMergeIdenticalMaterialSlots = false;
 
 		FTransform ParentTransform;
@@ -1334,7 +1423,7 @@ namespace UE::UsdCollision::Private
 		{
 			BodySetup->Modify();
 			BodySetup->RemoveSimpleCollision();
-			BodySetup->bNeverNeedsCookedCollisionData = true; // this will prevent the complex collision mesh from being generated
+			BodySetup->bNeverNeedsCookedCollisionData = true;	 // this will prevent the complex collision mesh from being generated
 			BodySetup->DefaultInstance.SetCollisionEnabled(ECollisionEnabled::NoCollision);
 			BodySetup->DefaultInstance.SetCollisionProfileName(UCollisionProfile::NoCollision_ProfileName);
 
@@ -1358,7 +1447,7 @@ namespace UE::UsdCollision::Private
 				{
 					MeshCollisionAPI.GetApproximationAttr().Get(&ApproximationAttr);
 				}
-				
+
 				if (ApproximationAttr == pxr::UsdPhysicsTokens->convexDecomposition)
 				{
 #if WITH_EDITOR
@@ -1463,7 +1552,7 @@ namespace UE::UsdCollision::Private
 				}
 				break;
 			}
-#endif // WITH_EDITOR
+#endif	  // WITH_EDITOR
 			case EUsdCollisionType::ConvexHull:
 			{
 				GenerateKDopAsSimpleCollision(MeshData, KDopDir18, BodySetup->AggGeom);
@@ -1502,10 +1591,10 @@ namespace UE::UsdCollision::Private
 		const bool bIsUpdate = true;
 		StaticMesh.CreateNavCollision(bIsUpdate);
 	}
-}
+}	 // namespace UE::UsdCollision::Private
 
 FBuildStaticMeshTaskChain::FBuildStaticMeshTaskChain(
-	const TSharedRef< FUsdSchemaTranslationContext >& InContext,
+	const TSharedRef<FUsdSchemaTranslationContext>& InContext,
 	const UE::FSdfPath& InPrimPath,
 	const TOptional<UE::FSdfPath>& InAlternativePrimToLinkAssetsTo
 )
@@ -1525,126 +1614,130 @@ void FBuildStaticMeshTaskChain::SetupTasks()
 
 	// Create static mesh (Main thread)
 	Do(ESchemaTranslationLaunchPolicy::Sync,
-		[this]()
-		{
-			// Force load MeshBuilderModule so that it's ready for the async tasks
+	   [this]()
+	   {
+	// Force load MeshBuilderModule so that it's ready for the async tasks
 #if WITH_EDITOR
-			FModuleManager::LoadModuleChecked< IMeshBuilderModule >(TEXT("MeshBuilder"));
-#endif // WITH_EDITOR
-
-			const FString PrimPathString = PrimPath.GetString();
-
-			const bool bParsedLODs = Context->bAllowInterpretingLODs && UsdUtils::IsGeomMeshALOD(GetPrim());
-
-			// It's useful to have the LOD Mesh prims be named "LOD0", "LOD1", etc. within the LOD variants so that we
-			// can easily tell which Mesh is actually meant to be the LOD mesh (in case there are more Meshes in each
-			// variant or other Meshes outside of the variant), but it's not ideal to have all the generated assets end
-			// up imported as "SM_LOD0_22", "SM_LOD0_23", etc. So here we fetch the parent prim name in case we're a LOD
-			FString MeshName;
-			if (bParsedLODs)
-			{
-				MeshName = PrimPath.GetParentPath().GetString();
-			}
-			else
-			{
-				MeshName = PrimPathString;
-			}
-
-			bool bIsNew = true;
-			const bool bShouldEnableNanite = UsdGeomMeshTranslatorImpl::ShouldEnableNanite(LODIndexToMeshDescription, LODIndexToMaterialInfo, *Context, GetPrim());
-			StaticMesh = UsdGeomMeshTranslatorImpl::CreateStaticMesh(
-				GetPrim(),
-				LODIndexToMeshDescription,
-				*Context,
-				MeshName,
-				bShouldEnableNanite,
-				bIsNew
-			);
-
-			if (StaticMesh)
-			{
-				if (Context->InfoCache)
-				{
-					const UE::FSdfPath& TargetPath = AlternativePrimToLinkAssetsTo.IsSet()
-						? AlternativePrimToLinkAssetsTo.GetValue()
-						: PrimPath;
-					Context->InfoCache->LinkAssetToPrim(TargetPath, StaticMesh);
-				}
+		   FModuleManager::LoadModuleChecked<IMeshBuilderModule>(TEXT("MeshBuilder"));
+#endif	  // WITH_EDITOR
+
+		   const FString PrimPathString = PrimPath.GetString();
+
+		   const bool bParsedLODs = Context->bAllowInterpretingLODs && UsdUtils::IsGeomMeshALOD(GetPrim());
+
+		   // It's useful to have the LOD Mesh prims be named "LOD0", "LOD1", etc. within the LOD variants so that we
+		   // can easily tell which Mesh is actually meant to be the LOD mesh (in case there are more Meshes in each
+		   // variant or other Meshes outside of the variant), but it's not ideal to have all the generated assets end
+		   // up imported as "SM_LOD0_22", "SM_LOD0_23", etc. So here we fetch the parent prim name in case we're a LOD
+		   FString MeshName;
+		   if (bParsedLODs)
+		   {
+			   MeshName = PrimPath.GetParentPath().GetString();
+		   }
+		   else
+		   {
+			   MeshName = PrimPathString;
+		   }
+
+		   bool bIsNew = true;
+		   const bool bShouldEnableNanite = UsdGeomMeshTranslatorImpl::ShouldEnableNanite(
+			   LODIndexToMeshDescription,
+			   LODIndexToMaterialInfo,
+			   *Context,
+			   GetPrim()
+		   );
+		   StaticMesh = UsdGeomMeshTranslatorImpl::CreateStaticMesh(
+			   GetPrim(),
+			   LODIndexToMeshDescription,
+			   *Context,
+			   MeshName,
+			   bShouldEnableNanite,
+			   bIsNew
+		   );
+
+		   if (StaticMesh)
+		   {
+			   if (Context->InfoCache)
+			   {
+				   const UE::FSdfPath& TargetPath = AlternativePrimToLinkAssetsTo.IsSet() ? AlternativePrimToLinkAssetsTo.GetValue() : PrimPath;
+				   Context->InfoCache->LinkAssetToPrim(TargetPath, StaticMesh);
+			   }
 
 #if WITH_EDITOR
-				StaticMesh->NaniteSettings.bEnabled = bShouldEnableNanite;
-#endif // WITH_EDITOR
-
-				if (UUsdMeshAssetUserData* UserData = UsdUtils::GetOrCreateAssetUserData<UUsdMeshAssetUserData>(StaticMesh))
-				{
-					UserData->PrimvarToUVIndex = LODIndexToMaterialInfo[0].PrimvarToUVIndex;	// We use the same primvar mapping for all LODs
-					UserData->PrimPaths.AddUnique(PrimPathString);
-
-					if (Context->MetadataOptions.bCollectMetadata)
-					{
-						// If we already have collected metadata just stash it into our UserData directly (the GeomMeshTranslator task
-						// chain itself does this as it parses lods, but the regular FBuildStaticMeshTaskChain won't do this on
-						// its own, and is reused for other task chains)
-						if (bCollectedMetadata)
-						{
-							UserData->StageIdentifierToMetadata.Add(GetPrim().GetStage().GetRootLayer().GetIdentifier(), LODMetadata);
-						}
-						else
-						{
-							UsdToUnreal::ConvertMetadata(
-								bParsedLODs ? GetPrim().GetParent() : GetPrim(),
-								UserData,
-								Context->MetadataOptions.BlockedPrefixFilters,
-								Context->MetadataOptions.bInvertFilters,
-								Context->MetadataOptions.bCollectFromEntireSubtrees
-							);
-						}
-					}
-
-					MeshTranslationImpl::RecordSourcePrimsForMaterialSlots(LODIndexToMaterialInfo, UserData);
-				}
-
-				// Only the original creator of the prim at creation time gets to set the material assignments
-				// directly on the mesh, all others prims ensure their materials via material overrides on the
-				// components
-				if (bIsNew)
-				{
-					UsdGeomMeshTranslatorImpl::ProcessStaticMeshMaterials(
-						GetPrim(),
-						LODIndexToMaterialInfo,
-						*StaticMesh,
-						*Context->AssetCache.Get(),
-						Context->InfoCache.Get(),
-						Context->Time,
-						Context->ObjectFlags,
-						Context->bReuseIdenticalAssets
-					);
+			   StaticMesh->NaniteSettings.bEnabled = bShouldEnableNanite;
+#endif	  // WITH_EDITOR
+
+			   if (UUsdMeshAssetUserData* UserData = UsdUtils::GetOrCreateAssetUserData<UUsdMeshAssetUserData>(StaticMesh))
+			   {
+				   UserData->PrimvarToUVIndex = LODIndexToMaterialInfo[0].PrimvarToUVIndex;	   // We use the same primvar mapping for all LODs
+				   UserData->PrimPaths.AddUnique(PrimPathString);
+
+				   if (Context->MetadataOptions.bCollectMetadata)
+				   {
+					   // If we already have collected metadata just stash it into our UserData directly (the GeomMeshTranslator task
+					   // chain itself does this as it parses lods, but the regular FBuildStaticMeshTaskChain won't do this on
+					   // its own, and is reused for other task chains)
+					   if (bCollectedMetadata)
+					   {
+						   UserData->StageIdentifierToMetadata.Add(GetPrim().GetStage().GetRootLayer().GetIdentifier(), LODMetadata);
+					   }
+					   else
+					   {
+						   UsdToUnreal::ConvertMetadata(
+							   bParsedLODs ? GetPrim().GetParent() : GetPrim(),
+							   UserData,
+							   Context->MetadataOptions.BlockedPrefixFilters,
+							   Context->MetadataOptions.bInvertFilters,
+							   Context->MetadataOptions.bCollectFromEntireSubtrees
+						   );
+					   }
+				   }
+
+				   MeshTranslationImpl::RecordSourcePrimsForMaterialSlots(LODIndexToMaterialInfo, UserData);
+			   }
+
+			   // Only the original creator of the prim at creation time gets to set the material assignments
+			   // directly on the mesh, all others prims ensure their materials via material overrides on the
+			   // components
+			   if (bIsNew)
+			   {
+				   UsdGeomMeshTranslatorImpl::ProcessStaticMeshMaterials(
+					   GetPrim(),
+					   LODIndexToMaterialInfo,
+					   *StaticMesh,
+					   *Context->AssetCache.Get(),
+					   Context->InfoCache.Get(),
+					   Context->Time,
+					   Context->ObjectFlags,
+					   Context->bReuseIdenticalAssets
+				   );
 
 #if WITH_EDITOR
-					const bool bRebuildAll = true;
-					StaticMesh->UpdateUVChannelData(bRebuildAll);
+				   const bool bRebuildAll = true;
+				   StaticMesh->UpdateUVChannelData(bRebuildAll);
 #else
 					// UpdateUVChannelData doesn't do anything without the editor
 					for (FStaticMaterial& Material : StaticMesh->GetStaticMaterials())
 					{
 						Material.UVChannelData.bInitialized = true;
 					}
-#endif // WITH_EDITOR
-				}
-				else
-				{
-					// Setup collision on existing mesh in case the collision settings have changed
-					UE::UsdCollision::Private::SetupSimpleCollision(GetPrim(), *StaticMesh);
-				}
-			}
-
-			// Only need to continue building the mesh if we just created it
-			return bIsNew;
-		});
+#endif	  // WITH_EDITOR
+			   }
+			   else
+			   {
+				   // Setup collision on existing mesh in case the collision settings have changed
+				   UE::UsdCollision::Private::SetupSimpleCollision(GetPrim(), *StaticMesh);
+			   }
+		   }
+
+		   // Only need to continue building the mesh if we just created it
+		   return bIsNew;
+	   });
 
 #if WITH_EDITOR
 	// Commit mesh description (Async)
-	Then(ESchemaTranslationLaunchPolicy::Async,
+	Then(
+		ESchemaTranslationLaunchPolicy::Async,
 		[this]()
 		{
 			UStaticMesh::FCommitMeshDescriptionParams Params;
@@ -1657,11 +1750,13 @@ void FBuildStaticMeshTaskChain::SetupTasks()
 			}
 
 			return true;
-		});
-#endif // WITH_EDITOR
+		}
+	);
+#endif	  // WITH_EDITOR
 
 	// PreBuild static mesh (Main thread)
-	Then(ESchemaTranslationLaunchPolicy::Sync,
+	Then(
+		ESchemaTranslationLaunchPolicy::Sync,
 		[this]()
 		{
 			RecreateRenderStateContextPtr = MakeShared<FStaticMeshComponentRecreateRenderStateContext>(StaticMesh, true, true);
@@ -1669,10 +1764,12 @@ void FBuildStaticMeshTaskChain::SetupTasks()
 			UsdGeomMeshTranslatorImpl::PreBuildStaticMesh(*StaticMesh);
 
 			return true;
-		});
+		}
+	);
 
 	// Build static mesh (Async)
-	Then(ESchemaTranslationLaunchPolicy::Async,
+	Then(
+		ESchemaTranslationLaunchPolicy::Async,
 		[this]() mutable
 		{
 			FStaticFeatureLevel FeatureLevel = GMaxRHIFeatureLevel;
@@ -1698,10 +1795,12 @@ void FBuildStaticMeshTaskChain::SetupTasks()
 			UE::UsdCollision::Private::SetupSimpleCollision(GetPrim(), *StaticMesh);
 
 			return true;
-		});
+		}
+	);
 
 	// PostBuild static mesh (Main thread)
-	Then(ESchemaTranslationLaunchPolicy::Sync,
+	Then(
+		ESchemaTranslationLaunchPolicy::Sync,
 		[this]()
 		{
 			UsdGeomMeshTranslatorImpl::PostBuildStaticMesh(*StaticMesh, LODIndexToMeshDescription);
@@ -1709,11 +1808,12 @@ void FBuildStaticMeshTaskChain::SetupTasks()
 			RecreateRenderStateContextPtr.Reset();
 
 			return true;
-		});
+		}
+	);
 }
 
 FGeomMeshCreateAssetsTaskChain::FGeomMeshCreateAssetsTaskChain(
-	const TSharedRef< FUsdSchemaTranslationContext >& InContext,
+	const TSharedRef<FUsdSchemaTranslationContext>& InContext,
 	const UE::FSdfPath& InPrimPath,
 	const TOptional<UE::FSdfPath>& AlternativePrimToLinkAssetsTo,
 	const FTransform& InAdditionalTransform
@@ -1739,65 +1839,65 @@ void FGeomMeshCreateAssetsTaskChain::SetupTasks()
 
 	// Create mesh descriptions (Async or ExclusiveSync)
 	Do(LaunchPolicy,
-		[this, bParseLODs]() -> bool
-		{
-			pxr::TfToken RenderContextToken = pxr::UsdShadeTokens->universalRenderContext;
-			if (!Context->RenderContext.IsNone())
-			{
-				RenderContextToken = UnrealToUsd::ConvertToken(*Context->RenderContext.ToString()).Get();
-			}
-
-			pxr::TfToken MaterialPurposeToken = pxr::UsdShadeTokens->allPurpose;
-			if (!Context->MaterialPurpose.IsNone())
-			{
-				MaterialPurposeToken = UnrealToUsd::ConvertToken(*Context->MaterialPurpose.ToString()).Get();
-			}
-
-			UsdToUnreal::FUsdMeshConversionOptions Options;
-			Options.TimeCode = Context->Time;
-			Options.PurposesToLoad = Context->PurposesToLoad;
-			Options.RenderContext = RenderContextToken;
-			Options.MaterialPurpose = MaterialPurposeToken;
-			Options.bMergeIdenticalMaterialSlots = Context->bMergeIdenticalMaterialSlots;
-			Options.AdditionalTransform = AdditionalTransform;
-			Options.SubdivisionLevel = Context->SubdivisionLevel;
-
-			UsdGeomMeshTranslatorImpl::LoadMeshDescriptions(
-				GetPrim(),
-				LODIndexToMeshDescription,
-				LODIndexToMaterialInfo,
-				LODMetadata,
-				Options,
-				Context->MetadataOptions,
-				Context->bAllowInterpretingLODs && bParseLODs
-			);
-
-			// If we're parsing LODs, LoadMeshDescriptions will have already collected metadata from the Mesh
-			// prims directly, but we still have to collect stuff from the actual LOD root prim, which is kind
-			// of absorbed into the asset
-			if (bParseLODs && Context->MetadataOptions.bCollectMetadata)
-			{
-				const bool bCollectMetadataFromSubtree = false;
-				UsdToUnreal::ConvertMetadata(
-					GetPrim().GetParent(),
-					LODMetadata,
-					Context->MetadataOptions.BlockedPrefixFilters,
-					Context->MetadataOptions.bInvertFilters,
-					bCollectMetadataFromSubtree
-				);
-			}
-			bCollectedMetadata = true;
-
-			// If we have at least one valid LOD, we should keep going
-			for (const FMeshDescription& MeshDescription : LODIndexToMeshDescription)
-			{
-				if (!MeshDescription.IsEmpty())
-				{
-					return true;
-				}
-			}
-			return false;
-		});
+	   [this, bParseLODs]() -> bool
+	   {
+		   pxr::TfToken RenderContextToken = pxr::UsdShadeTokens->universalRenderContext;
+		   if (!Context->RenderContext.IsNone())
+		   {
+			   RenderContextToken = UnrealToUsd::ConvertToken(*Context->RenderContext.ToString()).Get();
+		   }
+
+		   pxr::TfToken MaterialPurposeToken = pxr::UsdShadeTokens->allPurpose;
+		   if (!Context->MaterialPurpose.IsNone())
+		   {
+			   MaterialPurposeToken = UnrealToUsd::ConvertToken(*Context->MaterialPurpose.ToString()).Get();
+		   }
+
+		   UsdToUnreal::FUsdMeshConversionOptions Options;
+		   Options.TimeCode = Context->Time;
+		   Options.PurposesToLoad = Context->PurposesToLoad;
+		   Options.RenderContext = RenderContextToken;
+		   Options.MaterialPurpose = MaterialPurposeToken;
+		   Options.bMergeIdenticalMaterialSlots = Context->bMergeIdenticalMaterialSlots;
+		   Options.AdditionalTransform = AdditionalTransform;
+		   Options.SubdivisionLevel = Context->SubdivisionLevel;
+
+		   UsdGeomMeshTranslatorImpl::LoadMeshDescriptions(
+			   GetPrim(),
+			   LODIndexToMeshDescription,
+			   LODIndexToMaterialInfo,
+			   LODMetadata,
+			   Options,
+			   Context->MetadataOptions,
+			   Context->bAllowInterpretingLODs && bParseLODs
+		   );
+
+		   // If we're parsing LODs, LoadMeshDescriptions will have already collected metadata from the Mesh
+		   // prims directly, but we still have to collect stuff from the actual LOD root prim, which is kind
+		   // of absorbed into the asset
+		   if (bParseLODs && Context->MetadataOptions.bCollectMetadata)
+		   {
+			   const bool bCollectMetadataFromSubtree = false;
+			   UsdToUnreal::ConvertMetadata(
+				   GetPrim().GetParent(),
+				   LODMetadata,
+				   Context->MetadataOptions.BlockedPrefixFilters,
+				   Context->MetadataOptions.bInvertFilters,
+				   bCollectMetadataFromSubtree
+			   );
+		   }
+		   bCollectedMetadata = true;
+
+		   // If we have at least one valid LOD, we should keep going
+		   for (const FMeshDescription& MeshDescription : LODIndexToMeshDescription)
+		   {
+			   if (!MeshDescription.IsEmpty())
+			   {
+				   return true;
+			   }
+		   }
+		   return false;
+	   });
 
 	FBuildStaticMeshTaskChain::SetupTasks();
 }
@@ -1829,7 +1929,7 @@ void FUsdGeomMeshTranslator::CreateAssets()
 		return;
 	}
 
-	TSharedRef< FGeomMeshCreateAssetsTaskChain > AssetsTaskChain = MakeShared< FGeomMeshCreateAssetsTaskChain >(Context, PrimPath);
+	TSharedRef<FGeomMeshCreateAssetsTaskChain> AssetsTaskChain = MakeShared<FGeomMeshCreateAssetsTaskChain>(Context, PrimPath);
 
 	Context->TranslatorTasks.Add(MoveTemp(AssetsTaskChain));
 }
@@ -1939,9 +2039,7 @@ bool FUsdGeomMeshTranslator::CanBeCollapsed(ECollapsingType CollapsingType) cons
 	UE::FUsdPrim Prim = GetPrim();
 
 	// Don't collapse if our final UStaticMesh would have multiple LODs
-	if (Context->bAllowInterpretingLODs &&
-		CollapsingType == ECollapsingType::Assets &&
-		UsdUtils::IsGeomMeshALOD(Prim))
+	if (Context->bAllowInterpretingLODs && CollapsingType == ECollapsingType::Assets && UsdUtils::IsGeomMeshALOD(Prim))
 	{
 		return false;
 	}
@@ -2000,10 +2098,7 @@ TSet<UE::FSdfPath> FUsdGeomMeshTranslator::CollectAuxiliaryPrims() const
 
 		// The UsdGeomSubset prims are used to express multiple material assignments per mesh. A change in them could
 		// mean a change in triangle to material slot mapping
-		TArray<TUsdStore<pxr::UsdPrim>> ChildPrims = UsdUtils::GetAllPrimsOfType(
-			Prim,
-			pxr::TfType::Find<pxr::UsdGeomSubset>()
-		);
+		TArray<TUsdStore<pxr::UsdPrim>> ChildPrims = UsdUtils::GetAllPrimsOfType(Prim, pxr::TfType::Find<pxr::UsdGeomSubset>());
 
 		Result.Reserve(ChildPrims.Num());
 		for (const TUsdStore<pxr::UsdPrim>& ChildPrim : ChildPrims)
@@ -2116,4 +2211,4 @@ bool FUsdGeomMeshTranslator::ShouldSkipSkinnablePrim(bool bCheckForComponent) co
 	return false;
 }
 
-#endif // #if USE_USD_SDK
+#endif	  // #if USE_USD_SDK
