USD: Reformat USDSchemas module.

diff --git a/Engine/Plugins/Importers/USDImporter/Source/USDSchemas/Private/USDGeomPointInstancerTranslator.cpp b/Engine/Plugins/Importers/USDImporter/Source/USDSchemas/Private/USDGeomPointInstancerTranslator.cpp
--- a/Engine/Plugins/Importers/USDImporter/Source/USDSchemas/Private/USDGeomPointInstancerTranslator.cpp
+++ b/Engine/Plugins/Importers/USDImporter/Source/USDSchemas/Private/USDGeomPointInstancerTranslator.cpp
@@ -5,8 +5,8 @@
 #include "MeshTranslationImpl.h"
 #include "USDAssetCache.h"
 #include "USDAssetUserData.h"
-#include "USDDrawModeComponent.h"
 #include "USDConversionUtils.h"
+#include "USDDrawModeComponent.h"
 #include "USDGeomMeshConversion.h"
 #include "USDGeomMeshTranslator.h"
 #include "USDLog.h"
@@ -18,7 +18,7 @@
 
 #if WITH_EDITOR
 #include "Editor.h"
-#endif // WITH_EDITOR
+#endif	  // WITH_EDITOR
 
 #if USE_USD_SDK
 
@@ -29,74 +29,72 @@
 #include "Modules/ModuleManager.h"
 
 #include "USDIncludesStart.h"
-	#include "pxr/usd/usd/prim.h"
-	#include "pxr/usd/usd/primRange.h"
-	#include "pxr/usd/usd/stage.h"
-	#include "pxr/usd/usdGeom/camera.h"
-	#include "pxr/usd/usdGeom/mesh.h"
-	#include "pxr/usd/usdGeom/pointInstancer.h"
-	#include "pxr/usd/usdGeom/xform.h"
-	#include "pxr/usd/usdGeom/xformable.h"
+#include "pxr/usd/usd/prim.h"
+#include "pxr/usd/usd/primRange.h"
+#include "pxr/usd/usd/stage.h"
+#include "pxr/usd/usdGeom/camera.h"
+#include "pxr/usd/usdGeom/mesh.h"
+#include "pxr/usd/usdGeom/pointInstancer.h"
+#include "pxr/usd/usdGeom/xform.h"
+#include "pxr/usd/usdGeom/xformable.h"
 #include "USDIncludesEnd.h"
 
 #define LOCTEXT_NAMESPACE "USDGeomPointInstancer"
 
 static bool GCollapseTopLevelPointInstancers = false;
 static FAutoConsoleVariableRef CVarCollapseTopLevelPointInstancers(
-	TEXT( "USD.CollapseTopLevelPointInstancers" ),
+	TEXT("USD.CollapseTopLevelPointInstancers"),
 	GCollapseTopLevelPointInstancers,
-	TEXT( "If this is true will cause any point instancer to be collapsed to a single static mesh. Point instancers that are used as prototypes for other point instancers will always be collapsed." ) );
+	TEXT("If this is true will cause any point instancer to be collapsed to a single static mesh. Point instancers that are used as prototypes for "
+		 "other point instancers will always be collapsed.")
+);
 
 namespace UsdGeomPointInstancerTranslatorImpl
 {
-	void ApplyPointInstanceTransforms( UInstancedStaticMeshComponent* Component, TArray<FTransform>& InstanceTransforms )
+	void ApplyPointInstanceTransforms(UInstancedStaticMeshComponent* Component, TArray<FTransform>& InstanceTransforms)
 	{
-		if ( Component )
+		if (Component)
 		{
 			TRACE_CPUPROFILER_EVENT_SCOPE(ApplyPointInstanceTransforms);
 
 			Component->ClearInstances();
-			Component->AddInstances( InstanceTransforms, false );
+			Component->AddInstances(InstanceTransforms, false);
 		}
 	}
 
-	void SetStaticMesh( UStaticMesh* StaticMesh, UStaticMeshComponent& MeshComponent )
+	void SetStaticMesh(UStaticMesh* StaticMesh, UStaticMeshComponent& MeshComponent)
 	{
-		if ( StaticMesh == MeshComponent.GetStaticMesh() )
+		if (StaticMesh == MeshComponent.GetStaticMesh())
 		{
 			return;
 		}
 
 		MeshComponent.Modify();
 
-		if ( MeshComponent.IsRegistered() )
+		if (MeshComponent.IsRegistered())
 		{
 			MeshComponent.UnregisterComponent();
 		}
 
-		if ( StaticMesh )
+		if (StaticMesh)
 		{
-			StaticMesh->CreateBodySetup(); // BodySetup is required for HISM component
-			StaticMesh->MarkAsNotHavingNavigationData(); // Needed or else it will warn if we try cooking with body setup
+			StaticMesh->CreateBodySetup();					// BodySetup is required for HISM component
+			StaticMesh->MarkAsNotHavingNavigationData();	// Needed or else it will warn if we try cooking with body setup
 		}
 
-		MeshComponent.SetStaticMesh( StaticMesh );
+		MeshComponent.SetStaticMesh(StaticMesh);
 
 		MeshComponent.RegisterComponent();
 	}
-}
+}	 // namespace UsdGeomPointInstancerTranslatorImpl
 
 FUsdGeomPointInstancerCreateAssetsTaskChain::FUsdGeomPointInstancerCreateAssetsTaskChain(
-	const TSharedRef< FUsdSchemaTranslationContext >& InContext,
+	const TSharedRef<FUsdSchemaTranslationContext>& InContext,
 	const UE::FSdfPath& InPrimPath,
 	bool bInIgnoreTopLevelTransformAndVisibility,
 	const TOptional<UE::FSdfPath>& InAlternativePrimToLinkAssetsTo
 )
-	: FBuildStaticMeshTaskChain(
-		InContext,
-		InPrimPath,
-		InAlternativePrimToLinkAssetsTo
-	)
+	: FBuildStaticMeshTaskChain(InContext, InPrimPath, InAlternativePrimToLinkAssetsTo)
 	, bIgnoreTopLevelTransformAndVisibility(bInIgnoreTopLevelTransformAndVisibility)
 {
 	SetupTasks();
@@ -107,52 +105,46 @@ void FUsdGeomPointInstancerCreateAssetsTaskChain::SetupTasks()
 	FScopedUnrealAllocs UnrealAllocs;
 
 	// Create mesh description (Async)
-	Do( ESchemaTranslationLaunchPolicy::Async,
-		[this]() -> bool
-		{
-			LODIndexToMeshDescription.Reset( 1 );
-			LODIndexToMaterialInfo.Reset( 1 );
-
-			FMeshDescription& AddedMeshDescription = LODIndexToMeshDescription.Emplace_GetRef();
-			UsdUtils::FUsdPrimMaterialAssignmentInfo& AssignmentInfo = LODIndexToMaterialInfo.Emplace_GetRef();
-
-			pxr::TfToken RenderContextToken = pxr::UsdShadeTokens->universalRenderContext;
-			if ( !Context->RenderContext.IsNone() )
-			{
-				RenderContextToken = UnrealToUsd::ConvertToken( *Context->RenderContext.ToString() ).Get();
-			}
-
-			pxr::TfToken MaterialPurposeToken = pxr::UsdShadeTokens->allPurpose;
-			if (!Context->MaterialPurpose.IsNone())
-			{
-				MaterialPurposeToken = UnrealToUsd::ConvertToken(*Context->MaterialPurpose.ToString()).Get();
-			}
-
-			UsdToUnreal::FUsdMeshConversionOptions Options;
-			Options.TimeCode = Context->Time;
-			Options.PurposesToLoad = Context->PurposesToLoad;
-			Options.RenderContext = RenderContextToken;
-			Options.MaterialPurpose = MaterialPurposeToken;
-			Options.bMergeIdenticalMaterialSlots = Context->bMergeIdenticalMaterialSlots;
-			Options.SubdivisionLevel = Context->SubdivisionLevel;
-
-			UsdToUnreal::ConvertGeomMeshHierarchy(
-				GetPrim(),
-				AddedMeshDescription,
-				AssignmentInfo,
-				Options,
-				bIgnoreTopLevelTransformAndVisibility
-			);
-
-			return !AddedMeshDescription.IsEmpty();
-		});
+	Do(ESchemaTranslationLaunchPolicy::Async,
+	   [this]() -> bool
+	   {
+		   LODIndexToMeshDescription.Reset(1);
+		   LODIndexToMaterialInfo.Reset(1);
+
+		   FMeshDescription& AddedMeshDescription = LODIndexToMeshDescription.Emplace_GetRef();
+		   UsdUtils::FUsdPrimMaterialAssignmentInfo& AssignmentInfo = LODIndexToMaterialInfo.Emplace_GetRef();
+
+		   pxr::TfToken RenderContextToken = pxr::UsdShadeTokens->universalRenderContext;
+		   if (!Context->RenderContext.IsNone())
+		   {
+			   RenderContextToken = UnrealToUsd::ConvertToken(*Context->RenderContext.ToString()).Get();
+		   }
+
+		   pxr::TfToken MaterialPurposeToken = pxr::UsdShadeTokens->allPurpose;
+		   if (!Context->MaterialPurpose.IsNone())
+		   {
+			   MaterialPurposeToken = UnrealToUsd::ConvertToken(*Context->MaterialPurpose.ToString()).Get();
+		   }
+
+		   UsdToUnreal::FUsdMeshConversionOptions Options;
+		   Options.TimeCode = Context->Time;
+		   Options.PurposesToLoad = Context->PurposesToLoad;
+		   Options.RenderContext = RenderContextToken;
+		   Options.MaterialPurpose = MaterialPurposeToken;
+		   Options.bMergeIdenticalMaterialSlots = Context->bMergeIdenticalMaterialSlots;
+		   Options.SubdivisionLevel = Context->SubdivisionLevel;
+
+		   UsdToUnreal::ConvertGeomMeshHierarchy(GetPrim(), AddedMeshDescription, AssignmentInfo, Options, bIgnoreTopLevelTransformAndVisibility);
+
+		   return !AddedMeshDescription.IsEmpty();
+	   });
 
 	FBuildStaticMeshTaskChain::SetupTasks();
 }
 
 void FUsdGeomPointInstancerTranslator::CreateAssets()
 {
-	TRACE_CPUPROFILER_EVENT_SCOPE( FUsdGeomPointInstancerTranslator::CreateAssets );
+	TRACE_CPUPROFILER_EVENT_SCOPE(FUsdGeomPointInstancerTranslator::CreateAssets);
 
 	// Don't bother generating assets if we're going to just draw some bounds for this prim instead
 	EUsdDrawMode DrawMode = UsdUtils::GetAppliedDrawMode(GetPrim());
@@ -165,8 +157,8 @@ void FUsdGeomPointInstancerTranslator::CreateAssets()
 	FScopedUsdAllocs Allocs;
 
 	pxr::UsdPrim Prim = GetPrim();
-	pxr::UsdGeomPointInstancer PointInstancer( Prim );
-	if ( !PointInstancer )
+	pxr::UsdGeomPointInstancer PointInstancer(Prim);
+	if (!PointInstancer)
 	{
 		return;
 	}
@@ -174,15 +166,13 @@ void FUsdGeomPointInstancerTranslator::CreateAssets()
 	// If another FUsdGeomXformableTranslator is collapsing the point instancer prim, it will do so by calling
 	// UsdToUnreal::ConvertGeomMeshHierarchy which will consume the prim directly.
 	// This case right here is if we're collapsing *ourselves*, where we'll essentially pretend we're a single static mesh.
-	if ( GCollapseTopLevelPointInstancers )
+	if (GCollapseTopLevelPointInstancers)
 	{
-		// Don't bake our actual point instancer's transform or visibility into the mesh as its nice to have these on the static mesh component instead
+		// Don't bake our actual point instancer's transform or visibility into the mesh as its nice to have these on the static mesh component
+		// instead
 		const bool bIgnoreTopLevelTransformAndVisibility = true;
-		Context->TranslatorTasks.Add(MakeShared<FUsdGeomPointInstancerCreateAssetsTaskChain>(
-			Context,
-			PrimPath,
-			bIgnoreTopLevelTransformAndVisibility
-		));
+		Context->TranslatorTasks.Add(MakeShared<FUsdGeomPointInstancerCreateAssetsTaskChain>(Context, PrimPath, bIgnoreTopLevelTransformAndVisibility)
+		);
 	}
 	// Otherwise we're just going to spawn HISM components instead
 	else
@@ -190,38 +180,44 @@ void FUsdGeomPointInstancerTranslator::CreateAssets()
 		const pxr::UsdRelationship& Prototypes = PointInstancer.GetPrototypesRel();
 
 		pxr::SdfPathVector PrototypePaths;
-		if ( !Prototypes.GetTargets( &PrototypePaths ) )
+		if (!Prototypes.GetTargets(&PrototypePaths))
 		{
 			return;
 		}
 
-		for ( const pxr::SdfPath& PrototypePath : PrototypePaths )
+		for (const pxr::SdfPath& PrototypePath : PrototypePaths)
 		{
 			// Note how we will spawn a task chain for the prototype *regardless of where it is*. This prototype
 			// could even be external to the point instancer itself, and so will already have been handled by
 			// another translator. Unfortunately we need to do this because we need to generate a task chain for
 			// it in case it is another point instancer itself
 
-			pxr::UsdPrim PrototypeUsdPrim = Prim.GetStage()->GetPrimAtPath( PrototypePath );
-			UE::FSdfPath UEPrototypePath{ PrototypePath };
+			pxr::UsdPrim PrototypeUsdPrim = Prim.GetStage()->GetPrimAtPath(PrototypePath);
+			UE::FSdfPath UEPrototypePath{PrototypePath};
 
-			if ( !PrototypeUsdPrim )
+			if (!PrototypeUsdPrim)
 			{
-				UE_LOG( LogUsd, Warning, TEXT( "Failed to find prototype '%s' for PointInstancer '%s' when collapsing assets" ), *UEPrototypePath.GetString(), *PrimPath.GetString() );
+				UE_LOG(
+					LogUsd,
+					Warning,
+					TEXT("Failed to find prototype '%s' for PointInstancer '%s' when collapsing assets"),
+					*UEPrototypePath.GetString(),
+					*PrimPath.GetString()
+				);
 				continue;
 			}
 
-			if ( Context->bAllowInterpretingLODs && UsdUtils::DoesPrimContainMeshLODs( PrototypeUsdPrim ) )
+			if (Context->bAllowInterpretingLODs && UsdUtils::DoesPrimContainMeshLODs(PrototypeUsdPrim))
 			{
 				// We have to provide one of the LOD meshes to the task chain, so find the path to one
 				UE::FSdfPath ChildMeshPath;
 				pxr::UsdPrimSiblingRange PrimRange = PrototypeUsdPrim.GetChildren();
-				for ( pxr::UsdPrimSiblingRange::iterator PrimRangeIt = PrimRange.begin(); PrimRangeIt != PrimRange.end(); ++PrimRangeIt )
+				for (pxr::UsdPrimSiblingRange::iterator PrimRangeIt = PrimRange.begin(); PrimRangeIt != PrimRange.end(); ++PrimRangeIt)
 				{
 					const pxr::UsdPrim& Child = *PrimRangeIt;
-					if ( pxr::UsdGeomMesh ChildMesh{ Child } )
+					if (pxr::UsdGeomMesh ChildMesh{Child})
 					{
-						ChildMeshPath = UE::FSdfPath{ Child.GetPrimPath() };
+						ChildMeshPath = UE::FSdfPath{Child.GetPrimPath()};
 						break;
 					}
 				}
@@ -230,38 +226,31 @@ void FUsdGeomPointInstancerTranslator::CreateAssets()
 				// otherwise wouldn't end up anywhere else. Note that in the default/export case 'Prim' (the prim that actually contains the LOD
 				// variant set) is schema-less (and so not an Xform), but this is just in case the user manually made it an Xform instead
 				FTransform AdditionalUESpaceTransform = FTransform::Identity;
-				if ( pxr::UsdGeomXform ParentXform{ PrototypeUsdPrim } )
+				if (pxr::UsdGeomXform ParentXform{PrototypeUsdPrim})
 				{
 					// Skip this LOD mesh if its invisible
 					pxr::TfToken Visibility;
 					pxr::UsdAttribute VisibilityAttr = ParentXform.GetVisibilityAttr();
-					if ( VisibilityAttr && VisibilityAttr.Get( &Visibility, Context->Time ) && Visibility == pxr::UsdGeomTokens->invisible )
+					if (VisibilityAttr && VisibilityAttr.Get(&Visibility, Context->Time) && Visibility == pxr::UsdGeomTokens->invisible)
 					{
 						continue;
 					}
 
 					// TODO: Handle the resetXformStack op for LOD parents
 					bool bOutResetTransformStack = false;
-					UsdToUnreal::ConvertXformable( Prim.GetStage(), ParentXform, AdditionalUESpaceTransform, Context->Time, &bOutResetTransformStack );
+					UsdToUnreal::ConvertXformable(Prim.GetStage(), ParentXform, AdditionalUESpaceTransform, Context->Time, &bOutResetTransformStack);
 				}
 
-				Context->TranslatorTasks.Add(MakeShared<FGeomMeshCreateAssetsTaskChain>(
-					Context,
-					ChildMeshPath,
-					PrimPath,
-					AdditionalUESpaceTransform
-				));
+				Context->TranslatorTasks.Add(MakeShared<FGeomMeshCreateAssetsTaskChain>(Context, ChildMeshPath, PrimPath, AdditionalUESpaceTransform)
+				);
 			}
 			else
 			{
 				// Fully bake the prototype top level transform and visibility into its own static mesh
 				const bool bIgnoreTopLevelTransformAndVisibility = false;
-				Context->TranslatorTasks.Add(MakeShared<FUsdGeomPointInstancerCreateAssetsTaskChain>(
-					Context,
-					UEPrototypePath,
-					bIgnoreTopLevelTransformAndVisibility,
-					PrimPath
-				));
+				Context->TranslatorTasks.Add(
+					MakeShared<FUsdGeomPointInstancerCreateAssetsTaskChain>(Context, UEPrototypePath, bIgnoreTopLevelTransformAndVisibility, PrimPath)
+				);
 			}
 		}
 	}
@@ -269,10 +258,10 @@ void FUsdGeomPointInstancerTranslator::CreateAssets()
 
 USceneComponent* FUsdGeomPointInstancerTranslator::CreateComponents()
 {
-	TRACE_CPUPROFILER_EVENT_SCOPE( FUsdGeomPointInstancerTranslator::CreateComponents );
+	TRACE_CPUPROFILER_EVENT_SCOPE(FUsdGeomPointInstancerTranslator::CreateComponents);
 
 	// If we're collapsing ourselves, we're really just a collapsed Xform prim, so let that translator handle it
-	if ( GCollapseTopLevelPointInstancers )
+	if (GCollapseTopLevelPointInstancers)
 	{
 		return FUsdGeomXformableTranslator::CreateComponents();
 	}
@@ -321,10 +310,13 @@ USceneComponent* FUsdGeomPointInstancerTranslator::CreateComponents()
 
 		// Lets pretend ParentComponent is pointing to the parent USceneComponent while we create the child HISMs, so they get
 		// automatically attached to it as children
-		TGuardValue< USceneComponent* > ParentComponentGuard{Context->ParentComponent, MainSceneComponent};
+		TGuardValue<USceneComponent*> ParentComponentGuard{Context->ParentComponent, MainSceneComponent};
 
 		TArray<TFuture<TTuple<UHierarchicalInstancedStaticMeshComponent*, TArray<FTransform>>>> Tasks;
-		FScopedSlowTask PrototypePathsSlowTask((float)PrototypePaths.size(), LOCTEXT("GeomPointCreateComponents", "Creating HierarchicalInstancedStaticMeshComponents for point instancers"));
+		FScopedSlowTask PrototypePathsSlowTask(
+			(float)PrototypePaths.size(),
+			LOCTEXT("GeomPointCreateComponents", "Creating HierarchicalInstancedStaticMeshComponents for point instancers")
+		);
 		for (int32 PrototypeIndex = 0; PrototypeIndex < PrototypePaths.size(); ++PrototypeIndex)
 		{
 			PrototypePathsSlowTask.EnterProgressFrame();
@@ -335,7 +327,13 @@ USceneComponent* FUsdGeomPointInstancerTranslator::CreateComponents()
 			pxr::UsdPrim PrototypeUsdPrim = Prim.GetStage()->GetPrimAtPath(PrototypePath);
 			if (!PrototypeUsdPrim)
 			{
-				UE_LOG(LogUsd, Warning, TEXT("Failed to find prototype '%s' for PointInstancer '%s' when creating components"), *PrototypePathStr, *PrimPath.GetString());
+				UE_LOG(
+					LogUsd,
+					Warning,
+					TEXT("Failed to find prototype '%s' for PointInstancer '%s' when creating components"),
+					*PrototypePathStr,
+					*PrimPath.GetString()
+				);
 				continue;
 			}
 
@@ -361,16 +359,23 @@ USceneComponent* FUsdGeomPointInstancerTranslator::CreateComponents()
 			FUsdGeomXformableTranslator XformableTranslator{UHISMComponent::StaticClass(), Context, UE::FUsdTyped(PrototypeUsdPrim)};
 
 			const bool bNeedsActor = false;
-			UHISMComponent* HISMComponent = Cast<UHISMComponent>(XformableTranslator.CreateComponentsEx({UHISMComponent::StaticClass()}, bNeedsActor));
+			UHISMComponent* HISMComponent = Cast<UHISMComponent>(XformableTranslator.CreateComponentsEx({UHISMComponent::StaticClass()}, bNeedsActor)
+			);
 			if (!HISMComponent)
 			{
-				UE_LOG(LogUsd, Warning, TEXT("Failed to generate HISM component for prototype '%s' for PointInstancer '%s'"), *PrototypePathStr, *PrimPath.GetString());
+				UE_LOG(
+					LogUsd,
+					Warning,
+					TEXT("Failed to generate HISM component for prototype '%s' for PointInstancer '%s'"),
+					*PrototypePathStr,
+					*PrimPath.GetString()
+				);
 				continue;
 			}
 		}
 	}
 
-	UpdateComponents( MainSceneComponent );
+	UpdateComponents(MainSceneComponent);
 
 	return MainSceneComponent;
 }
@@ -388,14 +393,14 @@ void FUsdGeomPointInstancerTranslator::UpdateComponents(USceneComponent* PointIn
 	if (PointInstancerRootComponent->GetClass() == USceneComponent::StaticClass())
 	{
 		pxr::UsdPrim Prim = GetPrim();
-		pxr::UsdGeomPointInstancer PointInstancer( Prim );
-		if ( !PointInstancer )
+		pxr::UsdGeomPointInstancer PointInstancer(Prim);
+		if (!PointInstancer)
 		{
 			return;
 		}
 
 		TUsdStore<pxr::SdfPathVector> PrototypePaths;
-		if ( !PointInstancer.GetPrototypesRel().GetTargets( &PrototypePaths.Get() ) )
+		if (!PointInstancer.GetPrototypesRel().GetTargets(&PrototypePaths.Get()))
 		{
 			return;
 		}
@@ -409,7 +414,7 @@ void FUsdGeomPointInstancerTranslator::UpdateComponents(USceneComponent* PointIn
 
 		// Lets pretend ParentComponent is pointing to the parent USceneComponent while we create the child HISMs, so they get
 		// automatically attached to it as children
-		TGuardValue< USceneComponent* > ParentComponentGuard{Context->ParentComponent, PointInstancerRootComponent};
+		TGuardValue<USceneComponent*> ParentComponentGuard{Context->ParentComponent, PointInstancerRootComponent};
 
 		const TArray<TObjectPtr<USceneComponent>>& AttachedChildren = PointInstancerRootComponent->GetAttachChildren();
 		TArray<UHierarchicalInstancedStaticMeshComponent*> AttachedHISMs;
@@ -443,15 +448,18 @@ void FUsdGeomPointInstancerTranslator::UpdateComponents(USceneComponent* PointIn
 
 		TArray<TFuture<TTuple<UHierarchicalInstancedStaticMeshComponent*, TArray<FTransform>>>> Tasks;
 		size_t NumPrototypePaths = PrototypePaths.Get().size();
-		FScopedSlowTask PrototypePathsSlowTask( ( float ) NumPrototypePaths, LOCTEXT( "GeomPointUpdateComponents", "Updating HierarchicalInstancedStaticMeshComponents for point instancers" ) );
-		for ( uint32 PrototypeIndex = 0; PrototypeIndex < NumPrototypePaths; ++PrototypeIndex )
+		FScopedSlowTask PrototypePathsSlowTask(
+			(float)NumPrototypePaths,
+			LOCTEXT("GeomPointUpdateComponents", "Updating HierarchicalInstancedStaticMeshComponents for point instancers")
+		);
+		for (uint32 PrototypeIndex = 0; PrototypeIndex < NumPrototypePaths; ++PrototypeIndex)
 		{
 			PrototypePathsSlowTask.EnterProgressFrame();
 
 			pxr::SdfPath PrototypePath = PrototypePaths.Get()[PrototypeIndex];
 
-			TUsdStore<pxr::UsdPrim> PrototypeUsdPrim = Prim.GetStage()->GetPrimAtPath( PrototypePath );
-			if ( !PrototypeUsdPrim.Get() )
+			TUsdStore<pxr::UsdPrim> PrototypeUsdPrim = Prim.GetStage()->GetPrimAtPath(PrototypePath);
+			if (!PrototypeUsdPrim.Get())
 			{
 				UE_LOG(
 					LogUsd,
@@ -481,13 +489,13 @@ void FUsdGeomPointInstancerTranslator::UpdateComponents(USceneComponent* PointIn
 
 			// If our prototype was a LOD mesh we will have used the path of one of the actual LOD meshes to start the FGeomMeshCreateAssetsTaskChain,
 			// so we have to look for our resulting mesh with the same path
-			if ( Context->bAllowInterpretingLODs && UsdUtils::DoesPrimContainMeshLODs( PrototypeUsdPrim.Get() ) )
+			if (Context->bAllowInterpretingLODs && UsdUtils::DoesPrimContainMeshLODs(PrototypeUsdPrim.Get()))
 			{
 				pxr::UsdPrimSiblingRange PrimRange = PrototypeUsdPrim.Get().GetChildren();
-				for ( pxr::UsdPrimSiblingRange::iterator PrimRangeIt = PrimRange.begin(); PrimRangeIt != PrimRange.end(); ++PrimRangeIt )
+				for (pxr::UsdPrimSiblingRange::iterator PrimRangeIt = PrimRange.begin(); PrimRangeIt != PrimRange.end(); ++PrimRangeIt)
 				{
 					const pxr::UsdPrim& Child = *PrimRangeIt;
-					if ( pxr::UsdGeomMesh ChildMesh{ Child } )
+					if (pxr::UsdGeomMesh ChildMesh{Child})
 					{
 						PrototypeUsdPrim = Child;
 						PrototypePath = Child.GetPrimPath();
@@ -503,31 +511,29 @@ void FUsdGeomPointInstancerTranslator::UpdateComponents(USceneComponent* PointIn
 			{
 				StaticMesh = Iter->second;
 			}
-			UsdGeomPointInstancerTranslatorImpl::SetStaticMesh( StaticMesh, *HISMComponent );
+			UsdGeomPointInstancerTranslatorImpl::SetStaticMesh(StaticMesh, *HISMComponent);
 
 			// Evaluating point instancer can take a long time and is thread-safe. Move to async task while we work on something else.
-			pxr::UsdTimeCode TimeCode{ Context->Time };
-			FUsdStageInfo StageInfo{ Prim.GetStage() };
-			Tasks.Emplace(
-				Async(
-					EAsyncExecution::ThreadPool,
-					[ TimeCode, StageInfo, PointInstancer, PrototypeIndex, HISMComponent ]()
-					{
-						TArray<FTransform> InstanceTransforms;
-						UsdUtils::GetPointInstancerTransforms( StageInfo, PointInstancer, PrototypeIndex, TimeCode, InstanceTransforms );
+			pxr::UsdTimeCode TimeCode{Context->Time};
+			FUsdStageInfo StageInfo{Prim.GetStage()};
+			Tasks.Emplace(Async(
+				EAsyncExecution::ThreadPool,
+				[TimeCode, StageInfo, PointInstancer, PrototypeIndex, HISMComponent]()
+				{
+					TArray<FTransform> InstanceTransforms;
+					UsdUtils::GetPointInstancerTransforms(StageInfo, PointInstancer, PrototypeIndex, TimeCode, InstanceTransforms);
 
-						return MakeTuple( HISMComponent, MoveTemp( InstanceTransforms ) );
-					}
-				)
-			);
+					return MakeTuple(HISMComponent, MoveTemp(InstanceTransforms));
+				}
+			));
 
 			// Handle material overrides
-			if ( StaticMesh )
+			if (StaticMesh)
 			{
 				TArray<UMaterialInterface*> ExistingAssignments;
-				for ( FStaticMaterial& StaticMaterial : StaticMesh->GetStaticMaterials() )
+				for (FStaticMaterial& StaticMaterial : StaticMesh->GetStaticMaterials())
 				{
-					ExistingAssignments.Add( StaticMaterial.MaterialInterface );
+					ExistingAssignments.Add(StaticMaterial.MaterialInterface);
 				}
 
 				MeshTranslationImpl::SetMaterialOverrides(
@@ -547,10 +553,10 @@ void FUsdGeomPointInstancerTranslator::UpdateComponents(USceneComponent* PointIn
 		}
 
 		// Wait on and assign results of the point instancer.
-		for ( auto& Future : Tasks )
+		for (auto& Future : Tasks)
 		{
-			TTuple<UHierarchicalInstancedStaticMeshComponent*, TArray<FTransform>> Result{ Future.Get() };
-			UsdGeomPointInstancerTranslatorImpl::ApplyPointInstanceTransforms( Result.Key, Result.Value );
+			TTuple<UHierarchicalInstancedStaticMeshComponent*, TArray<FTransform>> Result{Future.Get()};
+			UsdGeomPointInstancerTranslatorImpl::ApplyPointInstanceTransforms(Result.Key, Result.Value);
 		}
 	}
 
@@ -562,7 +568,7 @@ bool FUsdGeomPointInstancerTranslator::CollapsesChildren(ECollapsingType Collaps
 	return true;
 }
 
-bool FUsdGeomPointInstancerTranslator::CanBeCollapsed( ECollapsingType CollapsingType ) const
+bool FUsdGeomPointInstancerTranslator::CanBeCollapsed(ECollapsingType CollapsingType) const
 {
 	return GCollapseTopLevelPointInstancers;
 }
@@ -589,7 +595,7 @@ TSet<UE::FSdfPath> FUsdGeomPointInstancerTranslator::CollectAuxiliaryPrims() con
 		return {};
 	}
 
-	IUsdSchemasModule& UsdSchemasModule = FModuleManager::Get().LoadModuleChecked< IUsdSchemasModule >(TEXT("USDSchemas"));
+	IUsdSchemasModule& UsdSchemasModule = FModuleManager::Get().LoadModuleChecked<IUsdSchemasModule>(TEXT("USDSchemas"));
 
 	TSet<UE::FSdfPath> Result;
 	Result.Reserve(PrototypePaths.size());
@@ -604,16 +610,13 @@ TSet<UE::FSdfPath> FUsdGeomPointInstancerTranslator::CollectAuxiliaryPrims() con
 		// We must depend on all prims of the prototype subtree, because we're in charge of collapsing it
 		if (PrototypePath.HasPrefix(PrimPath))
 		{
-			TArray<TUsdStore<pxr::UsdPrim>> ChildPrims = UsdUtils::GetAllPrimsOfType(
-				PrototypePrim,
-				pxr::TfType::Find<pxr::UsdGeomImageable>()
-			);
+			TArray<TUsdStore<pxr::UsdPrim>> ChildPrims = UsdUtils::GetAllPrimsOfType(PrototypePrim, pxr::TfType::Find<pxr::UsdGeomImageable>());
 
 			for (const TUsdStore<pxr::UsdPrim>& ChildPrim : ChildPrims)
 			{
 				Result.Add(UE::FSdfPath{ChildPrim.Get().GetPrimPath()});
 
-				TSharedPtr< FUsdSchemaTranslator > ChildSchemaTranslator = UsdSchemasModule.GetTranslatorRegistry().CreateTranslatorForSchema(
+				TSharedPtr<FUsdSchemaTranslator> ChildSchemaTranslator = UsdSchemasModule.GetTranslatorRegistry().CreateTranslatorForSchema(
 					Context,
 					UE::FUsdTyped(ChildPrim.Get())
 				);
@@ -639,7 +642,7 @@ TSet<UE::FSdfPath> FUsdGeomPointInstancerTranslator::CollectAuxiliaryPrims() con
 			{
 				Result.Add(UE::FSdfPath{PrimRangeIt->GetPrimPath()});
 
-				TSharedPtr< FUsdSchemaTranslator > SchemaTranslator = UsdSchemasModule.GetTranslatorRegistry().CreateTranslatorForSchema(
+				TSharedPtr<FUsdSchemaTranslator> SchemaTranslator = UsdSchemasModule.GetTranslatorRegistry().CreateTranslatorForSchema(
 					Context,
 					UE::FUsdTyped(*PrimRangeIt)
 				);
@@ -656,5 +659,4 @@ TSet<UE::FSdfPath> FUsdGeomPointInstancerTranslator::CollectAuxiliaryPrims() con
 
 #undef LOCTEXT_NAMESPACE
 
-#endif // #if USE_USD_SDK
-
+#endif	  // #if USE_USD_SDK
diff --git a/Engine/Plugins/Importers/USDImporter/Source/USDSchemas/Private/USDGeomPrimitiveTranslator.cpp b/Engine/Plugins/Importers/USDImporter/Source/USDSchemas/Private/USDGeomPrimitiveTranslator.cpp
--- a/Engine/Plugins/Importers/USDImporter/Source/USDSchemas/Private/USDGeomPrimitiveTranslator.cpp
+++ b/Engine/Plugins/Importers/USDImporter/Source/USDSchemas/Private/USDGeomPrimitiveTranslator.cpp
@@ -5,13 +5,11 @@
 #if USE_USD_SDK
 
 #include "UnrealUSDWrapper.h"
-#include "USDAssetImportData.h"
 #include "USDConversionUtils.h"
 #include "USDDrawModeComponent.h"
 #include "USDGeomMeshConversion.h"
 #include "USDGeomMeshTranslator.h"
 #include "USDInfoCache.h"
-#include "USDLog.h"
 #include "USDTypesConversion.h"
 
 #include "UsdWrappers/SdfPath.h"
@@ -24,7 +22,7 @@
 #include "StaticMeshAttributes.h"
 
 #include "USDIncludesStart.h"
-	#include "pxr/usd/usd/prim.h"
+#include "pxr/usd/usd/prim.h"
 #include "USDIncludesEnd.h"
 
 namespace UsdGeomPrimitiveTranslatorImpl
@@ -97,12 +95,7 @@ namespace UsdGeomPrimitiveTranslatorImpl
 				   Options.bMergeIdenticalMaterialSlots = Context->bMergeIdenticalMaterialSlots;
 				   Options.SubdivisionLevel = Context->SubdivisionLevel;
 
-				   UsdGeomPrimitiveTranslatorImpl::LoadMeshDescriptions(
-					   GetPrim(),
-					   LODIndexToMeshDescription,
-					   LODIndexToMaterialInfo,
-					   Options
-				   );
+				   UsdGeomPrimitiveTranslatorImpl::LoadMeshDescriptions(GetPrim(), LODIndexToMeshDescription, LODIndexToMaterialInfo, Options);
 
 				   // If we have at least one valid LOD, we should keep going
 				   for (const FMeshDescription& MeshDescription : LODIndexToMeshDescription)
@@ -159,9 +152,7 @@ USceneComponent* FUsdGeomPrimitiveTranslator::CreateComponents()
 	{
 		if (Context->InfoCache)
 		{
-			if (UStaticMesh* StaticMesh = Context->InfoCache->GetSingleAssetForPrim<UStaticMesh>(
-				PrimPath
-			))
+			if (UStaticMesh* StaticMesh = Context->InfoCache->GetSingleAssetForPrim<UStaticMesh>(PrimPath))
 			{
 				TArray<UMaterialInterface*> ExistingAssignments;
 				for (FStaticMaterial& StaticMaterial : StaticMesh->GetStaticMaterials())
diff --git a/Engine/Plugins/Importers/USDImporter/Source/USDSchemas/Private/USDGeomXformableTranslator.cpp b/Engine/Plugins/Importers/USDImporter/Source/USDSchemas/Private/USDGeomXformableTranslator.cpp
--- a/Engine/Plugins/Importers/USDImporter/Source/USDSchemas/Private/USDGeomXformableTranslator.cpp
+++ b/Engine/Plugins/Importers/USDImporter/Source/USDSchemas/Private/USDGeomXformableTranslator.cpp
@@ -4,12 +4,10 @@
 
 #include "MeshTranslationImpl.h"
 #include "UnrealUSDWrapper.h"
-#include "USDAssetCache.h"
 #include "USDAssetUserData.h"
-#include "USDDrawModeComponent.h"
 #include "USDClassesModule.h"
 #include "USDConversionUtils.h"
-#include "USDErrorUtils.h"
+#include "USDDrawModeComponent.h"
 #include "USDGeomMeshConversion.h"
 #include "USDGeomMeshTranslator.h"
 #include "USDIntegrationUtils.h"
@@ -20,7 +18,6 @@
 #include "USDShadeConversion.h"
 #include "USDTypesConversion.h"
 
-#include "Components/HierarchicalInstancedStaticMeshComponent.h"
 #include "Components/LightComponentBase.h"
 #include "Components/SceneComponent.h"
 #include "Components/SkeletalMeshComponent.h"
@@ -35,7 +32,6 @@
 #include "Modules/ModuleManager.h"
 #include "PhysicsEngine/BodySetup.h"
 #include "Roles/LiveLinkTransformRole.h"
-#include "StaticMeshAttributes.h"
 
 #include "UsdWrappers/SdfPath.h"
 #include "UsdWrappers/UsdGeomBBoxCache.h"
@@ -46,49 +42,45 @@
 #if USE_USD_SDK
 
 #include "USDIncludesStart.h"
-	#include "pxr/usd/kind/registry.h"
-	#include "pxr/usd/usd/modelAPI.h"
-	#include "pxr/usd/usd/prim.h"
-	#include "pxr/usd/usd/primRange.h"
-	#include "pxr/usd/usd/variantSets.h"
-	#include "pxr/usd/usdGeom/mesh.h"
-	#include "pxr/usd/usdGeom/modelAPI.h"
-	#include "pxr/usd/usdGeom/pointInstancer.h"
-	#include "pxr/usd/usdGeom/subset.h"
-	#include "pxr/usd/usdGeom/xformable.h"
-	#include "pxr/usd/usdShade/materialBindingAPI.h"
-	#include "pxr/usd/usdSkel/root.h"
+#include "pxr/usd/usd/modelAPI.h"
+#include "pxr/usd/usd/prim.h"
+#include "pxr/usd/usdGeom/modelAPI.h"
+#include "pxr/usd/usdGeom/pointInstancer.h"
+#include "pxr/usd/usdGeom/xformable.h"
+#include "pxr/usd/usdSkel/root.h"
 #include "USDIncludesEnd.h"
 
 static bool GCollapsePrimsWithoutKind = true;
 static FAutoConsoleVariableRef CVarCollapsePrimsWithoutKind(
-	TEXT( "USD.CollapsePrimsWithoutKind" ),
+	TEXT("USD.CollapsePrimsWithoutKind"),
 	GCollapsePrimsWithoutKind,
-	TEXT( "Allow collapsing prims that have no authored 'Kind' value" ) );
+	TEXT("Allow collapsing prims that have no authored 'Kind' value")
+);
 
 static bool GEnableCollision = true;
 static FAutoConsoleVariableRef CVarEnableCollision(
-	TEXT( "USD.EnableCollision" ),
+	TEXT("USD.EnableCollision"),
 	GEnableCollision,
-	TEXT( "Whether to have collision enabled for spawned components and generated meshes" ) );
+	TEXT("Whether to have collision enabled for spawned components and generated meshes")
+);
 
 namespace UE::UsdXformableTranslatorImpl::Private
 {
-	void SetUpSceneComponentForLiveLink( const FUsdSchemaTranslationContext& Context, USceneComponent* Component, const pxr::UsdPrim& Prim )
+	void SetUpSceneComponentForLiveLink(const FUsdSchemaTranslationContext& Context, USceneComponent* Component, const pxr::UsdPrim& Prim)
 	{
-		if ( !Component || !Prim )
+		if (!Component || !Prim)
 		{
 			return;
 		}
 
 		AActor* Parent = Component->GetOwner();
-		if ( !Parent )
+		if (!Parent)
 		{
 			return;
 		}
 
 		USceneComponent* RootComponent = Parent->GetRootComponent();
-		if ( !RootComponent )
+		if (!RootComponent)
 		{
 			return;
 		}
@@ -99,7 +91,7 @@ namespace UE::UsdXformableTranslatorImpl::Private
 			// or not, so this just helps us make that a little bit faster. Its important because UpdateComponents (which calls us)
 			// is the main function that is called to animate components, so it can be spammed in case this prim has animations
 			static TMap<TWeakObjectPtr<USceneComponent>, TWeakObjectPtr<ULiveLinkComponentController>> LiveLinkEnabledComponents;
-			if ( ULiveLinkComponentController* ExistingController = LiveLinkEnabledComponents.FindRef( Component ).Get() )
+			if (ULiveLinkComponentController* ExistingController = LiveLinkEnabledComponents.FindRef(Component).Get())
 			{
 				// We found an existing controller we created to track this component, so use that
 				Controller = ExistingController;
@@ -108,11 +100,11 @@ namespace UE::UsdXformableTranslatorImpl::Private
 			else
 			{
 				TArray<ULiveLinkComponentController*> LiveLinkComponents;
-				Parent->GetComponents( LiveLinkComponents );
+				Parent->GetComponents(LiveLinkComponents);
 
-				for ( ULiveLinkComponentController* LiveLinkComponent : LiveLinkComponents )
+				for (ULiveLinkComponentController* LiveLinkComponent : LiveLinkComponents)
 				{
-					if ( LiveLinkComponent->GetControlledComponent( ULiveLinkTransformRole::StaticClass() ) == Component )
+					if (LiveLinkComponent->GetControlledComponent(ULiveLinkTransformRole::StaticClass()) == Component)
 					{
 						// We found some other controller handling this component somehow, use that
 						Controller = LiveLinkComponent;
@@ -120,54 +112,55 @@ namespace UE::UsdXformableTranslatorImpl::Private
 					}
 				}
 
-				if ( !Controller )
+				if (!Controller)
 				{
 					// We'll get a warning from the live link controller component in case the component its controlling is not movable
 					Component->Mobility = EComponentMobility::Movable;
 
-					Controller = NewObject<ULiveLinkComponentController>( Parent, NAME_None, Context.ObjectFlags );
+					Controller = NewObject<ULiveLinkComponentController>(Parent, NAME_None, Context.ObjectFlags);
 					Controller->bUpdateInEditor = true;
 
 					// Important because of how ULiveLinkComponentController::TickComponent also checks for the sequencer
 					// tag to try and guess if the controlled component is a spawnable
 					Controller->bDisableEvaluateLiveLinkWhenSpawnable = false;
 
-					Parent->AddInstanceComponent( Controller );
+					Parent->AddInstanceComponent(Controller);
 					Controller->RegisterComponent();
 				}
 
-				if ( Controller )
+				if (Controller)
 				{
-					LiveLinkEnabledComponents.Add( Component, Controller );
+					LiveLinkEnabledComponents.Add(Component, Controller);
 				}
 			}
 		}
 
 		// Configure controller with our desired parameters
-		if ( Controller )
+		if (Controller)
 		{
 			FScopedUsdAllocs Allocs;
 
 			FLiveLinkSubjectRepresentation SubjectRepresentation;
 			SubjectRepresentation.Role = ULiveLinkTransformRole::StaticClass();
 
-			if ( pxr::UsdAttribute Attr = Prim.GetAttribute( UnrealIdentifiers::UnrealLiveLinkSubjectName ) )
+			if (pxr::UsdAttribute Attr = Prim.GetAttribute(UnrealIdentifiers::UnrealLiveLinkSubjectName))
 			{
 				std::string SubjectName;
-				if ( Attr.Get( &SubjectName ) )
+				if (Attr.Get(&SubjectName))
 				{
-					SubjectRepresentation.Subject = FName{ *UsdToUnreal::ConvertString( SubjectName ) };
+					SubjectRepresentation.Subject = FName{*UsdToUnreal::ConvertString(SubjectName)};
 				}
 			}
-			Controller->SetSubjectRepresentation( SubjectRepresentation );
+			Controller->SetSubjectRepresentation(SubjectRepresentation);
 
-			// This should be done after setting the subject representation to ensure that the LiveLink component's ControllerMap has a transform controller
-			Controller->SetControlledComponent( ULiveLinkTransformRole::StaticClass(), Component );
+			// This should be done after setting the subject representation to ensure that the LiveLink component's ControllerMap has a transform
+			// controller
+			Controller->SetControlledComponent(ULiveLinkTransformRole::StaticClass(), Component);
 
-			if ( pxr::UsdAttribute Attr = Prim.GetAttribute( UnrealIdentifiers::UnrealLiveLinkEnabled ) )
+			if (pxr::UsdAttribute Attr = Prim.GetAttribute(UnrealIdentifiers::UnrealLiveLinkEnabled))
 			{
 				bool bEnabled = true;
-				if ( Attr.Get( &bEnabled ) )
+				if (Attr.Get(&bEnabled))
 				{
 					Controller->bEvaluateLiveLink = bEnabled;
 				}
@@ -175,39 +168,39 @@ namespace UE::UsdXformableTranslatorImpl::Private
 		}
 	}
 
-	void RemoveLiveLinkFromComponent( USceneComponent* Component )
+	void RemoveLiveLinkFromComponent(USceneComponent* Component)
 	{
-		if ( !Component )
+		if (!Component)
 		{
 			return;
 		}
 
 		AActor* Parent = Component->GetOwner();
-		if ( !Parent )
+		if (!Parent)
 		{
 			return;
 		}
 
 		TArray<ULiveLinkComponentController*> LiveLinkComponents;
-		Parent->GetComponents( LiveLinkComponents );
+		Parent->GetComponents(LiveLinkComponents);
 
-		for ( ULiveLinkComponentController* LiveLinkComponent : LiveLinkComponents )
+		for (ULiveLinkComponentController* LiveLinkComponent : LiveLinkComponents)
 		{
-			if ( LiveLinkComponent->GetControlledComponent( ULiveLinkTransformRole::StaticClass() ) == Component )
+			if (LiveLinkComponent->GetControlledComponent(ULiveLinkTransformRole::StaticClass()) == Component)
 			{
-				LiveLinkComponent->SetControlledComponent( ULiveLinkTransformRole::StaticClass(), nullptr );
-				Parent->RemoveInstanceComponent( LiveLinkComponent );
+				LiveLinkComponent->SetControlledComponent(ULiveLinkTransformRole::StaticClass(), nullptr);
+				Parent->RemoveInstanceComponent(LiveLinkComponent);
 				break;
 			}
 		}
 	}
-}
+}	 // namespace UE::UsdXformableTranslatorImpl::Private
 
 class FUsdGeomXformableCreateAssetsTaskChain : public FBuildStaticMeshTaskChain
 {
 public:
-	explicit FUsdGeomXformableCreateAssetsTaskChain( const TSharedRef< FUsdSchemaTranslationContext >& InContext, const UE::FSdfPath& InPrimPath )
-		: FBuildStaticMeshTaskChain( InContext, InPrimPath )
+	explicit FUsdGeomXformableCreateAssetsTaskChain(const TSharedRef<FUsdSchemaTranslationContext>& InContext, const UE::FSdfPath& InPrimPath)
+		: FBuildStaticMeshTaskChain(InContext, InPrimPath)
 	{
 		SetupTasks();
 	}
@@ -221,58 +214,53 @@ void FUsdGeomXformableCreateAssetsTaskChain::SetupTasks()
 	FScopedUnrealAllocs UnrealAllocs;
 
 	// Create mesh description (Async)
-	Do( ESchemaTranslationLaunchPolicy::Async,
-		[ this ]() -> bool
-		{
-			// We will never have multiple LODs of meshes that were collapsed together, as LOD'd meshes don't collapse. So just parse the mesh we get as LOD0
-			LODIndexToMeshDescription.Reset(1);
-			LODIndexToMaterialInfo.Reset(1);
-
-			FMeshDescription& AddedMeshDescription = LODIndexToMeshDescription.Emplace_GetRef();
-			UsdUtils::FUsdPrimMaterialAssignmentInfo& AssignmentInfo = LODIndexToMaterialInfo.Emplace_GetRef();
-
-			pxr::TfToken RenderContextToken = pxr::UsdShadeTokens->universalRenderContext;
-			if ( !Context->RenderContext.IsNone() )
-			{
-				RenderContextToken = UnrealToUsd::ConvertToken( *Context->RenderContext.ToString() ).Get();
-			}
-
-			pxr::TfToken MaterialPurposeToken = pxr::UsdShadeTokens->allPurpose;
-			if (!Context->MaterialPurpose.IsNone())
-			{
-				MaterialPurposeToken = UnrealToUsd::ConvertToken(*Context->MaterialPurpose.ToString()).Get();
-			}
-
-			// We're going to put Prim's transform and visibility on the component, so we don't need to bake it into the combined mesh
-			const bool bSkipRootPrimTransformAndVis = true;
-
-			UsdToUnreal::FUsdMeshConversionOptions Options;
-			Options.TimeCode = Context->Time;
-			Options.PurposesToLoad = Context->PurposesToLoad;
-			Options.RenderContext = RenderContextToken;
-			Options.MaterialPurpose = MaterialPurposeToken;
-			Options.bMergeIdenticalMaterialSlots = Context->bMergeIdenticalMaterialSlots;
-			Options.SubdivisionLevel = Context->SubdivisionLevel;
-
-			UsdToUnreal::ConvertGeomMeshHierarchy(
-				GetPrim(),
-				AddedMeshDescription,
-				AssignmentInfo,
-				Options,
-				bSkipRootPrimTransformAndVis
-			);
-
-			return !AddedMeshDescription.IsEmpty();
-		} );
+	Do(ESchemaTranslationLaunchPolicy::Async,
+	   [this]() -> bool
+	   {
+		   // We will never have multiple LODs of meshes that were collapsed together, as LOD'd meshes don't collapse. So just parse the mesh we get
+		   // as LOD0
+		   LODIndexToMeshDescription.Reset(1);
+		   LODIndexToMaterialInfo.Reset(1);
+
+		   FMeshDescription& AddedMeshDescription = LODIndexToMeshDescription.Emplace_GetRef();
+		   UsdUtils::FUsdPrimMaterialAssignmentInfo& AssignmentInfo = LODIndexToMaterialInfo.Emplace_GetRef();
+
+		   pxr::TfToken RenderContextToken = pxr::UsdShadeTokens->universalRenderContext;
+		   if (!Context->RenderContext.IsNone())
+		   {
+			   RenderContextToken = UnrealToUsd::ConvertToken(*Context->RenderContext.ToString()).Get();
+		   }
+
+		   pxr::TfToken MaterialPurposeToken = pxr::UsdShadeTokens->allPurpose;
+		   if (!Context->MaterialPurpose.IsNone())
+		   {
+			   MaterialPurposeToken = UnrealToUsd::ConvertToken(*Context->MaterialPurpose.ToString()).Get();
+		   }
+
+		   // We're going to put Prim's transform and visibility on the component, so we don't need to bake it into the combined mesh
+		   const bool bSkipRootPrimTransformAndVis = true;
+
+		   UsdToUnreal::FUsdMeshConversionOptions Options;
+		   Options.TimeCode = Context->Time;
+		   Options.PurposesToLoad = Context->PurposesToLoad;
+		   Options.RenderContext = RenderContextToken;
+		   Options.MaterialPurpose = MaterialPurposeToken;
+		   Options.bMergeIdenticalMaterialSlots = Context->bMergeIdenticalMaterialSlots;
+		   Options.SubdivisionLevel = Context->SubdivisionLevel;
+
+		   UsdToUnreal::ConvertGeomMeshHierarchy(GetPrim(), AddedMeshDescription, AssignmentInfo, Options, bSkipRootPrimTransformAndVis);
+
+		   return !AddedMeshDescription.IsEmpty();
+	   });
 
 	FBuildStaticMeshTaskChain::SetupTasks();
 }
 
 void FUsdGeomXformableTranslator::CreateAssets()
 {
-	TRACE_CPUPROFILER_EVENT_SCOPE( FUsdGeomMeshTranslator::CreateAssets );
+	TRACE_CPUPROFILER_EVENT_SCOPE(FUsdGeomMeshTranslator::CreateAssets);
 
-	if ( !CollapsesChildren( ECollapsingType::Assets ) )
+	if (!CollapsesChildren(ECollapsingType::Assets))
 	{
 		// We only have to create assets when our children are collapsed together
 		return;
@@ -286,12 +274,16 @@ void FUsdGeomXformableTranslator::CreateAssets()
 		return;
 	}
 
-	Context->TranslatorTasks.Add( MakeShared< FUsdGeomXformableCreateAssetsTaskChain >( Context, PrimPath ) );
+	Context->TranslatorTasks.Add(MakeShared<FUsdGeomXformableCreateAssetsTaskChain>(Context, PrimPath));
 }
 
-FUsdGeomXformableTranslator::FUsdGeomXformableTranslator( TSubclassOf< USceneComponent > InComponentTypeOverride, TSharedRef< FUsdSchemaTranslationContext > InContext, const UE::FUsdTyped& InSchema )
-	: FUsdSchemaTranslator( InContext, InSchema )
-	, ComponentTypeOverride( InComponentTypeOverride )
+FUsdGeomXformableTranslator::FUsdGeomXformableTranslator(
+	TSubclassOf<USceneComponent> InComponentTypeOverride,
+	TSharedRef<FUsdSchemaTranslationContext> InContext,
+	const UE::FUsdTyped& InSchema
+)
+	: FUsdSchemaTranslator(InContext, InSchema)
+	, ComponentTypeOverride(InComponentTypeOverride)
 {
 }
 
@@ -322,9 +314,7 @@ USceneComponent* FUsdGeomXformableTranslator::CreateComponents()
 	{
 		if (Context->InfoCache)
 		{
-			if (UStaticMesh* StaticMesh = Context->InfoCache->GetSingleAssetForPrim<UStaticMesh>(
-				PrimPath
-			))
+			if (UStaticMesh* StaticMesh = Context->InfoCache->GetSingleAssetForPrim<UStaticMesh>(PrimPath))
 			{
 				TArray<UMaterialInterface*> ExistingAssignments;
 				for (FStaticMaterial& StaticMaterial : StaticMesh->GetStaticMaterials())
@@ -352,36 +342,36 @@ USceneComponent* FUsdGeomXformableTranslator::CreateComponents()
 	return SceneComponent;
 }
 
-USceneComponent* FUsdGeomXformableTranslator::CreateComponentsEx( TOptional< TSubclassOf< USceneComponent > > ComponentType, TOptional< bool > bNeedsActor )
+USceneComponent* FUsdGeomXformableTranslator::CreateComponentsEx(TOptional<TSubclassOf<USceneComponent>> ComponentType, TOptional<bool> bNeedsActor)
 {
-	TRACE_CPUPROFILER_EVENT_SCOPE( FUsdGeomXformableTranslator::CreateComponentsEx );
+	TRACE_CPUPROFILER_EVENT_SCOPE(FUsdGeomXformableTranslator::CreateComponentsEx);
 
-	if ( !Context->IsValid() )
+	if (!Context->IsValid())
 	{
 		return nullptr;
 	}
 
 	UE::FUsdPrim Prim = GetPrim();
-	if ( !Prim )
+	if (!Prim)
 	{
 		return nullptr;
 	}
 
 	FScopedUnrealAllocs UnrealAllocs;
 
-	if ( !bNeedsActor.IsSet() )
+	if (!bNeedsActor.IsSet())
 	{
 		// Don't add components to the AUsdStageActor or the USDStageImport 'scene actor'
 		UE::FUsdPrim ParentPrim = Prim.GetParent();
 		bool bIsTopLevelPrim = ParentPrim.IsValid() && ParentPrim.IsPseudoRoot();
 
 		// If we don't have any parent prim with a type that generates a component, we are still technically a top-level prim
-		if ( !bIsTopLevelPrim )
+		if (!bIsTopLevelPrim)
 		{
 			bool bHasParentComponent = false;
-			while ( ParentPrim.IsValid() )
+			while (ParentPrim.IsValid())
 			{
-				if ( UE::FUsdGeomXformable( ParentPrim ) )
+				if (UE::FUsdGeomXformable(ParentPrim))
 				{
 					bHasParentComponent = true;
 					break;
@@ -389,14 +379,15 @@ USceneComponent* FUsdGeomXformableTranslator::CreateComponentsEx( TOptional< TSu
 
 				ParentPrim = ParentPrim.GetParent();
 			}
-			if ( !bHasParentComponent )
+			if (!bHasParentComponent)
 			{
 				bIsTopLevelPrim = true;
 			}
 		}
 
-		auto PrimNeedsActor = []( const UE::FUsdPrim& UsdPrim ) -> bool
+		auto PrimNeedsActor = [](const UE::FUsdPrim& UsdPrim) -> bool
 		{
+			// clang-format off
 			return  UsdPrim.IsPseudoRoot() ||
 					UsdPrim.IsModel() ||
 					UsdPrim.IsGroup() ||
@@ -408,29 +399,25 @@ USceneComponent* FUsdGeomXformableTranslator::CreateComponentsEx( TOptional< TSu
 					UsdPrim.IsA(TEXT("SkelRoot"));  // Now that we use the UsdSkelSkeletonTranslator, UsdSkelRoots will be handled like regular
 													// Xforms. We likely always want then to show up on the outliner though, as they are important
 													// prims
+			// clang-format on
 		};
 
-		bNeedsActor =
-		(
-			bIsTopLevelPrim ||
-			Context->ParentComponent == nullptr ||
-			PrimNeedsActor( Prim )
-		);
+		bNeedsActor = (bIsTopLevelPrim || Context->ParentComponent == nullptr || PrimNeedsActor(Prim));
 
 		// We don't want to start a component hierarchy if one of our child will break it by being an actor
-		if ( !bNeedsActor.GetValue() )
+		if (!bNeedsActor.GetValue())
 		{
-			TFunction< bool( const UE::FUsdPrim& ) > RecursiveChildPrimsNeedsActor;
-			RecursiveChildPrimsNeedsActor = [ PrimNeedsActor, &RecursiveChildPrimsNeedsActor ]( const UE::FUsdPrim& UsdPrim ) -> bool
+			TFunction<bool(const UE::FUsdPrim&)> RecursiveChildPrimsNeedsActor;
+			RecursiveChildPrimsNeedsActor = [PrimNeedsActor, &RecursiveChildPrimsNeedsActor](const UE::FUsdPrim& UsdPrim) -> bool
 			{
 				const bool bTraverseInstanceProxies = true;
-				for ( const pxr::UsdPrim& Child : UsdPrim.GetFilteredChildren( bTraverseInstanceProxies ) )
+				for (const pxr::UsdPrim& Child : UsdPrim.GetFilteredChildren(bTraverseInstanceProxies))
 				{
-					if ( PrimNeedsActor( UE::FUsdPrim( Child ) ) )
+					if (PrimNeedsActor(UE::FUsdPrim(Child)))
 					{
 						return true;
 					}
-					else if ( RecursiveChildPrimsNeedsActor( UE::FUsdPrim( Child ) ) )
+					else if (RecursiveChildPrimsNeedsActor(UE::FUsdPrim(Child)))
 					{
 						return true;
 					}
@@ -439,7 +426,7 @@ USceneComponent* FUsdGeomXformableTranslator::CreateComponentsEx( TOptional< TSu
 				return false;
 			};
 
-			bNeedsActor = RecursiveChildPrimsNeedsActor( UE::FUsdPrim( Prim ) );
+			bNeedsActor = RecursiveChildPrimsNeedsActor(UE::FUsdPrim(Prim));
 		}
 	}
 
@@ -454,14 +441,14 @@ USceneComponent* FUsdGeomXformableTranslator::CreateComponentsEx( TOptional< TSu
 	const EObjectFlags PreComponentFlags = Context->ObjectFlags & ~(RF_Standalone | RF_Public | RF_Transient);
 	const EObjectFlags PostComponentFlags = Context->ObjectFlags & RF_Transient;
 
-	if ( bNeedsActor.GetValue() )
+	if (bNeedsActor.GetValue())
 	{
 		// Spawn actor
 		FActorSpawnParameters SpawnParameters;
 		SpawnParameters.ObjectFlags = PreComponentFlags;
-		SpawnParameters.OverrideLevel =  Context->Level;
+		SpawnParameters.OverrideLevel = Context->Level;
 		SpawnParameters.Name = Prim.GetName();
-		SpawnParameters.NameMode = FActorSpawnParameters::ESpawnActorNameMode::Requested; // Will generate a unique name in case of a conflict
+		SpawnParameters.NameMode = FActorSpawnParameters::ESpawnActorNameMode::Requested;	 // Will generate a unique name in case of a conflict
 
 		UClass* ActorClass = nullptr;
 		if (ComponentType.Get({}) == UUsdDrawModeComponent::StaticClass())
@@ -475,26 +462,26 @@ USceneComponent* FUsdGeomXformableTranslator::CreateComponentsEx( TOptional< TSu
 			ActorClass = UsdUtils::GetActorTypeForPrim(Prim);
 		}
 
-		AActor* SpawnedActor = Context->Level->GetWorld()->SpawnActor( ActorClass, nullptr, SpawnParameters );
+		AActor* SpawnedActor = Context->Level->GetWorld()->SpawnActor(ActorClass, nullptr, SpawnParameters);
 
-		if ( SpawnedActor )
+		if (SpawnedActor)
 		{
 			SpawnedActor->SetFlags(PostComponentFlags);
 
 #if WITH_EDITOR
 			const bool bMarkDirty = false;
-			SpawnedActor->SetActorLabel( Prim.GetName().ToString(), bMarkDirty );
+			SpawnedActor->SetActorLabel(Prim.GetName().ToString(), bMarkDirty);
 
 			// If our AUsdStageActor is in a hidden level/layer and we spawn actors, they should also be hidden
-			if ( Context->ParentComponent )
+			if (Context->ParentComponent)
 			{
-				if ( AActor* ParentActor = Context->ParentComponent->GetOwner() )
+				if (AActor* ParentActor = Context->ParentComponent->GetOwner())
 				{
 					SpawnedActor->bHiddenEdLevel = ParentActor->bHiddenEdLevel;
 					SpawnedActor->bHiddenEdLayer = ParentActor->bHiddenEdLayer;
 				}
 			}
-#endif // WITH_EDITOR
+#endif	  // WITH_EDITOR
 
 			SceneComponent = SpawnedActor->GetRootComponent();
 
@@ -506,45 +493,44 @@ USceneComponent* FUsdGeomXformableTranslator::CreateComponentsEx( TOptional< TSu
 		ComponentOuter = Context->ParentComponent;
 	}
 
-	if ( !ComponentOuter )
+	if (!ComponentOuter)
 	{
-		UE_LOG( LogUsd, Warning, TEXT("Invalid outer when trying to create SceneComponent for prim (%s)"), *PrimPath.GetString() );
+		UE_LOG(LogUsd, Warning, TEXT("Invalid outer when trying to create SceneComponent for prim (%s)"), *PrimPath.GetString());
 		return nullptr;
 	}
 
-	if ( !SceneComponent )
+	if (!SceneComponent)
 	{
-		if ( !ComponentType.IsSet() )
+		if (!ComponentType.IsSet())
 		{
-			if ( ComponentTypeOverride.IsSet() )
+			if (ComponentTypeOverride.IsSet())
 			{
 				ComponentType = ComponentTypeOverride.GetValue();
 			}
 			else
 			{
-				ComponentType = UsdUtils::GetComponentTypeForPrim( Prim );
+				ComponentType = UsdUtils::GetComponentTypeForPrim(Prim);
 
-				// For now only upgrade actual scene components to static mesh components (important because skeletal mesh components will also fit this
-				// criteria but we don't want to use a static mesh component for those)
-				if ( CollapsesChildren( ECollapsingType::Assets ) && ComponentType.IsSet() && ComponentType.GetValue() == USceneComponent::StaticClass() )
+				// For now only upgrade actual scene components to static mesh components (important because skeletal mesh components will also fit
+				// this criteria but we don't want to use a static mesh component for those)
+				if (CollapsesChildren(ECollapsingType::Assets) && ComponentType.IsSet() && ComponentType.GetValue() == USceneComponent::StaticClass())
 				{
-					// If we're a type that collapses assets, we should probably be a static mesh component as we only really collapse static meshes together right now.
-					// We can't just check if there's a static mesh for this prim on the cache, because the prims with meshes could be potentially invisible (and so
-					// we don't have parsed their meshes yet), so here we traverse our child hierarchy and if we have any chance of ever generating a Mesh, we go
-					// for a static mesh component
-					TArray< UE::FUsdPrim > ChildGprims = UsdUtils::GetAllPrimsOfType( Prim, TEXT( "UsdGeomGprim" ) );
-					if ( ChildGprims.Num() > 0 )
+					// If we're a type that collapses assets, we should probably be a static mesh component as we only really collapse static meshes
+					// together right now. We can't just check if there's a static mesh for this prim on the cache, because the prims with meshes
+					// could be potentially invisible (and so we don't have parsed their meshes yet), so here we traverse our child hierarchy and if
+					// we have any chance of ever generating a Mesh, we go for a static mesh component
+					TArray<UE::FUsdPrim> ChildGprims = UsdUtils::GetAllPrimsOfType(Prim, TEXT("UsdGeomGprim"));
+					if (ChildGprims.Num() > 0)
 					{
 						ComponentType = UStaticMeshComponent::StaticClass();
 					}
 				}
 				// If this is a component for a point instancer that just collapsed itself into a static mesh, just make
 				// a static mesh component that can receive it
-				else if ( pxr::UsdPrim{ Prim }.IsA<pxr::UsdGeomPointInstancer>() )
+				else if (pxr::UsdPrim{Prim}.IsA<pxr::UsdGeomPointInstancer>())
 				{
-					static IConsoleVariable* CollapseCvar =
-						IConsoleManager::Get().FindConsoleVariable( TEXT( "USD.CollapseTopLevelPointInstancers" ) );
-					if ( CollapseCvar && CollapseCvar->GetBool() )
+					static IConsoleVariable* CollapseCvar = IConsoleManager::Get().FindConsoleVariable(TEXT("USD.CollapseTopLevelPointInstancers"));
+					if (CollapseCvar && CollapseCvar->GetBool())
 					{
 						ComponentType = UStaticMeshComponent::StaticClass();
 					}
@@ -552,19 +538,19 @@ USceneComponent* FUsdGeomXformableTranslator::CreateComponentsEx( TOptional< TSu
 			}
 		}
 
-		if ( ComponentType.IsSet() && ComponentType.GetValue() != nullptr )
+		if (ComponentType.IsSet() && ComponentType.GetValue() != nullptr)
 		{
 			const FName ComponentName = MakeUniqueObjectName(
 				ComponentOuter,
 				ComponentType.GetValue(),
 				*IUsdClassesModule::SanitizeObjectName(Prim.GetName().ToString())
 			);
-			SceneComponent = NewObject< USceneComponent >( ComponentOuter, ComponentType.GetValue(), ComponentName, PreComponentFlags );
+			SceneComponent = NewObject<USceneComponent>(ComponentOuter, ComponentType.GetValue(), ComponentName, PreComponentFlags);
 			SceneComponent->SetFlags(PostComponentFlags);
 
-			if ( AActor* Owner = SceneComponent->GetOwner() )
+			if (AActor* Owner = SceneComponent->GetOwner())
 			{
-				Owner->AddInstanceComponent( SceneComponent );
+				Owner->AddInstanceComponent(SceneComponent);
 			}
 		}
 	}
@@ -588,23 +574,23 @@ USceneComponent* FUsdGeomXformableTranslator::CreateComponentsEx( TOptional< TSu
 		);
 	}
 
-	if ( SceneComponent )
+	if (SceneComponent)
 	{
-		if ( !GEnableCollision )
+		if (!GEnableCollision)
 		{
 			// In most cases this will have no benefit memory-wise, as regular UStaticMeshComponents build their physics meshes anyway
 			// when registering, regardless of these. HISM components will *not* build them though, so disabling the cvar may lead
 			// to some memory savings
-			if ( UPrimitiveComponent* PrimComp = Cast<UPrimitiveComponent>( SceneComponent ) )
+			if (UPrimitiveComponent* PrimComp = Cast<UPrimitiveComponent>(SceneComponent))
 			{
-				PrimComp->SetCollisionEnabled( ECollisionEnabled::NoCollision );
-				PrimComp->SetCollisionProfileName( UCollisionProfile::NoCollision_ProfileName );
+				PrimComp->SetCollisionEnabled(ECollisionEnabled::NoCollision);
+				PrimComp->SetCollisionProfileName(UCollisionProfile::NoCollision_ProfileName);
 			}
 		}
 
-		if ( !SceneComponent->GetOwner()->GetRootComponent() )
+		if (!SceneComponent->GetOwner()->GetRootComponent())
 		{
-			SceneComponent->GetOwner()->SetRootComponent( SceneComponent );
+			SceneComponent->GetOwner()->SetRootComponent(SceneComponent);
 		}
 
 		// If we're spawning into a level that is being streamed in, our construction scripts will be rerun, and may want to set the scene component
@@ -633,9 +619,9 @@ USceneComponent* FUsdGeomXformableTranslator::CreateComponentsEx( TOptional< TSu
 		// Do this before UpdatingComponents as we may need to use the parent transform to set a world transform directly
 		// (in case of resetXformStack). Besides, this is more consistent anyway as during stage updates we'll call
 		// UpdateComponents with all the components already attached
-		SceneComponent->AttachToComponent( Context->ParentComponent, FAttachmentTransformRules::KeepRelativeTransform );
+		SceneComponent->AttachToComponent(Context->ParentComponent, FAttachmentTransformRules::KeepRelativeTransform);
 
-		if ( !SceneComponent->IsRegistered() )
+		if (!SceneComponent->IsRegistered())
 		{
 			SceneComponent->RegisterComponent();
 		}
@@ -644,9 +630,9 @@ USceneComponent* FUsdGeomXformableTranslator::CreateComponentsEx( TOptional< TSu
 	return SceneComponent;
 }
 
-void FUsdGeomXformableTranslator::UpdateComponents( USceneComponent* SceneComponent )
+void FUsdGeomXformableTranslator::UpdateComponents(USceneComponent* SceneComponent)
 {
-	TRACE_CPUPROFILER_EVENT_SCOPE( FUsdGeomXformableTranslator::UpdateComponents );
+	TRACE_CPUPROFILER_EVENT_SCOPE(FUsdGeomXformableTranslator::UpdateComponents);
 
 	if (SceneComponent && Context->InfoCache)
 	{
@@ -664,32 +650,31 @@ void FUsdGeomXformableTranslator::UpdateComponents( USceneComponent* SceneCompon
 		// If the user modified a mesh parameter (e.g. vertex color), the hash will be different and it will become a separate asset
 		// so we must check for this and assign the new StaticMesh
 		bool bHasMultipleLODs = false;
-		if ( UStaticMeshComponent* StaticMeshComponent = Cast< UStaticMeshComponent >( SceneComponent ) )
+		if (UStaticMeshComponent* StaticMeshComponent = Cast<UStaticMeshComponent>(SceneComponent))
 		{
-			UStaticMesh* PrimStaticMesh = Context->InfoCache->GetSingleAssetForPrim<UStaticMesh>(
-				PrimPath
-			);
+			UStaticMesh* PrimStaticMesh = Context->InfoCache->GetSingleAssetForPrim<UStaticMesh>(PrimPath);
 
-			if ( PrimStaticMesh )
+			if (PrimStaticMesh)
 			{
 				bHasMultipleLODs = PrimStaticMesh->GetNumLODs() > 1;
 			}
 
-			if ( PrimStaticMesh != StaticMeshComponent->GetStaticMesh() )
+			if (PrimStaticMesh != StaticMeshComponent->GetStaticMesh())
 			{
 				// Need to make sure the mesh's resources are initialized here as it may have just been built in another thread
 				// Only do this if required though, as this mesh could using these resources currently (e.g. PIE and editor world sharing the mesh)
-				if (PrimStaticMesh && !PrimStaticMesh->AreRenderingResourcesInitialized() && (FApp::CanEverRender() || !FPlatformProperties::RequiresCookedData()))
+				if (PrimStaticMesh && !PrimStaticMesh->AreRenderingResourcesInitialized()
+					&& (FApp::CanEverRender() || !FPlatformProperties::RequiresCookedData()))
 				{
 					PrimStaticMesh->InitResources();
 				}
 
-				if ( StaticMeshComponent->IsRegistered() )
+				if (StaticMeshComponent->IsRegistered())
 				{
 					StaticMeshComponent->UnregisterComponent();
 				}
 
-				StaticMeshComponent->SetStaticMesh( PrimStaticMesh );
+				StaticMeshComponent->SetStaticMesh(PrimStaticMesh);
 
 				// We can't register yet, as UsdToUnreal::ConvertXformable below us may want to move the component.
 				// We'll always re-register when needed below, though.
@@ -718,19 +703,19 @@ void FUsdGeomXformableTranslator::UpdateComponents( USceneComponent* SceneCompon
 		// skeletal version of the LiveLink configuration, we only handle setting up LiveLink for simple transforms
 		if (!Prim.IsA(TEXT("SkelRoot")) && !Prim.IsA(TEXT("Skeleton")))
 		{
-			if ( UsdUtils::PrimHasSchema( Prim, UnrealIdentifiers::LiveLinkAPI ) )
+			if (UsdUtils::PrimHasSchema(Prim, UnrealIdentifiers::LiveLinkAPI))
 			{
-				UE::UsdXformableTranslatorImpl::Private::SetUpSceneComponentForLiveLink( Context.Get(), SceneComponent, Prim );
+				UE::UsdXformableTranslatorImpl::Private::SetUpSceneComponentForLiveLink(Context.Get(), SceneComponent, Prim);
 			}
 			else
 			{
-				UE::UsdXformableTranslatorImpl::Private::RemoveLiveLinkFromComponent( SceneComponent );
+				UE::UsdXformableTranslatorImpl::Private::RemoveLiveLinkFromComponent(SceneComponent);
 			}
 		}
 
 		// Only put the transform into the component if we haven't parsed LODs for our static mesh: The Mesh transforms will already be baked
 		// into the mesh at that case, as each LOD could technically have a separate transform
-		if ( !Context->bAllowInterpretingLODs || !bHasMultipleLODs )
+		if (!Context->bAllowInterpretingLODs || !bHasMultipleLODs)
 		{
 			// Don't update the component's transform if this is already factored in as root motion within the AnimSequence
 			bool bConvertTransform = true;
@@ -759,13 +744,13 @@ void FUsdGeomXformableTranslator::UpdateComponents( USceneComponent* SceneCompon
 				}
 			}
 
-			UsdToUnreal::ConvertXformable( Context->Stage, pxr::UsdGeomXformable( Prim ), *SceneComponent, Context->Time, bConvertTransform );
+			UsdToUnreal::ConvertXformable(Context->Stage, pxr::UsdGeomXformable(Prim), *SceneComponent, Context->Time, bConvertTransform);
 		}
 
 		// Note how we should only register if we unregistered ourselves: If we did this every time we would
 		// register too early during the process of duplicating into PIE, and that would prevent a future RegisterComponent
 		// call from naturally creating the required render state
-		if ( !SceneComponent->IsRegistered() )
+		if (!SceneComponent->IsRegistered())
 		{
 			SceneComponent->RegisterComponent();
 		}
@@ -774,7 +759,12 @@ void FUsdGeomXformableTranslator::UpdateComponents( USceneComponent* SceneCompon
 
 namespace UE::UsdXformableTranslatorImpl::Private
 {
-	void AssignBoundsComponentTextures(UE::FUsdPrim Prim, UUsdDrawModeComponent* BoundsComponent, UUsdAssetCache2& AssetCache, FUsdInfoCache& InfoCache)
+	void AssignBoundsComponentTextures(
+		UE::FUsdPrim Prim,
+		UUsdDrawModeComponent* BoundsComponent,
+		UUsdAssetCache2& AssetCache,
+		FUsdInfoCache& InfoCache
+	)
 	{
 		if (!Prim)
 		{
@@ -893,7 +883,7 @@ USceneComponent* FUsdGeomXformableTranslator::CreateAlternativeDrawModeComponent
 		case EUsdDrawMode::Cards:
 		{
 			UUsdDrawModeComponent* Component = Cast<UUsdDrawModeComponent>(CreateComponentsEx({UUsdDrawModeComponent::StaticClass()}, bNeedsActor));
-			if(ensure(Component) && Context->AssetCache && Context->InfoCache)
+			if (ensure(Component) && Context->AssetCache && Context->InfoCache)
 			{
 				// For now we only assign textures when creating components, not when updating. Maybe in the future we can
 				// add support for "texture animations"
@@ -1000,7 +990,7 @@ void FUsdGeomXformableTranslator::CreateAlternativeDrawModeAssets(EUsdDrawMode D
 
 bool FUsdGeomXformableTranslator::CollapsesChildren(ECollapsingType CollapsingType) const
 {
-	TRACE_CPUPROFILER_EVENT_SCOPE( FUsdGeomXformableTranslator::CollapsesChildren );
+	TRACE_CPUPROFILER_EVENT_SCOPE(FUsdGeomXformableTranslator::CollapsesChildren);
 
 	if (!Context->bIsBuildingInfoCache)
 	{
@@ -1020,7 +1010,7 @@ bool FUsdGeomXformableTranslator::CollapsesChildren(ECollapsingType CollapsingTy
 	FScopedUsdAllocs UsdAllocs;
 
 	pxr::UsdPrim Prim = GetPrim();
-	pxr::UsdModelAPI Model{ pxr::UsdTyped( Prim ) };
+	pxr::UsdModelAPI Model{pxr::UsdTyped(Prim)};
 
 	// Now that we use UsdSkelSkeletonTranslator the SkelRoots will be handled by the FUsdGeomXformableTranslator (here).
 	// SkelRoots are likely going to end up with SkeletalMeshes though, so we can assume we won't be collapsing them just from that
@@ -1029,42 +1019,43 @@ bool FUsdGeomXformableTranslator::CollapsesChildren(ECollapsingType CollapsingTy
 		return false;
 	}
 
-	if ( Model )
+	if (Model)
 	{
-		EUsdDefaultKind PrimKind = UsdUtils::GetDefaultKind( Prim );
+		EUsdDefaultKind PrimKind = UsdUtils::GetDefaultKind(Prim);
 
 		// Note that this is false if PrimKind is None
-		const bool bPrimKindShouldCollapse = EnumHasAnyFlags( Context->KindsToCollapse, PrimKind );
+		const bool bPrimKindShouldCollapse = EnumHasAnyFlags(Context->KindsToCollapse, PrimKind);
 
-		bCollapsesChildren = Context->KindsToCollapse != EUsdDefaultKind::None &&
-			( bPrimKindShouldCollapse || ( PrimKind == EUsdDefaultKind::None && GCollapsePrimsWithoutKind ) );
+		bCollapsesChildren = Context->KindsToCollapse != EUsdDefaultKind::None
+							 && (bPrimKindShouldCollapse || (PrimKind == EUsdDefaultKind::None && GCollapsePrimsWithoutKind));
 
-		if ( !bCollapsesChildren )
+		if (!bCollapsesChildren)
 		{
 			// Temp support for the prop kind
-			bCollapsesChildren = Model.IsKind( pxr::TfToken( "prop" ), pxr::UsdModelAPI::KindValidationNone );
+			bCollapsesChildren = Model.IsKind(pxr::TfToken("prop"), pxr::UsdModelAPI::KindValidationNone);
 		}
 
-		if ( bCollapsesChildren )
+		if (bCollapsesChildren)
 		{
-			IUsdSchemasModule& UsdSchemasModule = FModuleManager::Get().LoadModuleChecked< IUsdSchemasModule >( TEXT("USDSchemas") );
-
-			// TODO: This can be optimized in order to make FUsdInfoCache::RebuildCacheForSubtree faster: If we have a child prim that we know doesn't collapse,
-			// any of our parents should be able to know they can't collapse *us* either.
-			// This is somewhat niche though: Realistically to waste time here a prim and its children need to have a kind that allows collapsing, and also not be able to collapse.
-			// Also, if any of these prims *does* manage to collapse, FUsdInfoCache will already not actually query the subtree children if they can collapse or not anymore,
-			// and just consider them collapsed by the parent
-			TArray< TUsdStore< pxr::UsdPrim > > ChildXformPrims = UsdUtils::GetAllPrimsOfType( Prim, pxr::TfType::Find< pxr::UsdGeomXformable >() );
-			for ( const TUsdStore< pxr::UsdPrim >& ChildXformPrim : ChildXformPrims )
+			IUsdSchemasModule& UsdSchemasModule = FModuleManager::Get().LoadModuleChecked<IUsdSchemasModule>(TEXT("USDSchemas"));
+
+			// TODO: This can be optimized in order to make FUsdInfoCache::RebuildCacheForSubtree faster: If we have a child prim that we know doesn't
+			// collapse, any of our parents should be able to know they can't collapse *us* either. This is somewhat niche though: Realistically to
+			// waste time here a prim and its children need to have a kind that allows collapsing, and also not be able to collapse. Also, if any of
+			// these prims *does* manage to collapse, FUsdInfoCache will already not actually query the subtree children if they can collapse or not
+			// anymore, and just consider them collapsed by the parent
+			TArray<TUsdStore<pxr::UsdPrim>> ChildXformPrims = UsdUtils::GetAllPrimsOfType(Prim, pxr::TfType::Find<pxr::UsdGeomXformable>());
+			for (const TUsdStore<pxr::UsdPrim>& ChildXformPrim : ChildXformPrims)
 			{
 				if (ChildXformPrim.Get().IsA<pxr::UsdSkelRoot>())
 				{
 					return false;
 				}
 
-				if ( TSharedPtr< FUsdSchemaTranslator > SchemaTranslator = UsdSchemasModule.GetTranslatorRegistry().CreateTranslatorForSchema( Context, UE::FUsdTyped( ChildXformPrim.Get() ) ) )
+				if (TSharedPtr<FUsdSchemaTranslator> SchemaTranslator = UsdSchemasModule.GetTranslatorRegistry()
+																			.CreateTranslatorForSchema(Context, UE::FUsdTyped(ChildXformPrim.Get())))
 				{
-					if ( !SchemaTranslator->CanBeCollapsed( CollapsingType ) )
+					if (!SchemaTranslator->CanBeCollapsed(CollapsingType))
 					{
 						return false;
 					}
@@ -1076,21 +1067,18 @@ bool FUsdGeomXformableTranslator::CollapsesChildren(ECollapsingType CollapsingTy
 	return bCollapsesChildren;
 }
 
-bool FUsdGeomXformableTranslator::CanBeCollapsed( ECollapsingType CollapsingType ) const
+bool FUsdGeomXformableTranslator::CanBeCollapsed(ECollapsingType CollapsingType) const
 {
 	FScopedUsdAllocs UsdAllocs;
 
-	pxr::UsdPrim UsdPrim{ GetPrim() };
-	if ( !UsdPrim )
+	pxr::UsdPrim UsdPrim{GetPrim()};
+	if (!UsdPrim)
 	{
 		return false;
 	}
 
-	if (UsdUtils::IsAnimated(UsdPrim) ||
-		UsdUtils::PrimHasSchema(UsdPrim, UnrealIdentifiers::LiveLinkAPI) ||
-		UsdPrim.IsA<pxr::UsdSkelRoot>() ||
-		(Context->bAllowInterpretingLODs && UsdUtils::DoesPrimContainMeshLODs(UsdPrim))
-	)
+	if (UsdUtils::IsAnimated(UsdPrim) || UsdUtils::PrimHasSchema(UsdPrim, UnrealIdentifiers::LiveLinkAPI) || UsdPrim.IsA<pxr::UsdSkelRoot>()
+		|| (Context->bAllowInterpretingLODs && UsdUtils::DoesPrimContainMeshLODs(UsdPrim)))
 	{
 		return false;
 	}
@@ -1121,10 +1109,7 @@ TSet<UE::FSdfPath> FUsdGeomXformableTranslator::CollectAuxiliaryPrims() const
 
 		// We check imageable because that is the most basal schema that is still relevant for collapsed meshes (it
 		// holds the visibility attribute)
-		TArray<TUsdStore<pxr::UsdPrim>> ChildPrims = UsdUtils::GetAllPrimsOfType(
-			Prim,
-			pxr::TfType::Find<pxr::UsdGeomImageable>()
-		);
+		TArray<TUsdStore<pxr::UsdPrim>> ChildPrims = UsdUtils::GetAllPrimsOfType(Prim, pxr::TfType::Find<pxr::UsdGeomImageable>());
 
 		Result.Reserve(ChildPrims.Num());
 		for (const TUsdStore<pxr::UsdPrim>& ChildPrim : ChildPrims)
@@ -1135,4 +1120,4 @@ TSet<UE::FSdfPath> FUsdGeomXformableTranslator::CollectAuxiliaryPrims() const
 	return Result;
 }
 
-#endif // #if USE_USD_SDK
+#endif	  // #if USE_USD_SDK
diff --git a/Engine/Plugins/Importers/USDImporter/Source/USDSchemas/Private/USDGeometryCacheTranslator.cpp b/Engine/Plugins/Importers/USDImporter/Source/USDSchemas/Private/USDGeometryCacheTranslator.cpp
--- a/Engine/Plugins/Importers/USDImporter/Source/USDSchemas/Private/USDGeometryCacheTranslator.cpp
+++ b/Engine/Plugins/Importers/USDImporter/Source/USDSchemas/Private/USDGeometryCacheTranslator.cpp
@@ -7,9 +7,9 @@
 
 #include "MeshTranslationImpl.h"
 #include "USDAssetUserData.h"
-#include "USDDrawModeComponent.h"
 #include "USDClassesModule.h"
 #include "USDConversionUtils.h"
+#include "USDDrawModeComponent.h"
 #include "USDGroomTranslatorUtils.h"
 #include "USDInfoCache.h"
 #include "USDIntegrationUtils.h"
@@ -38,30 +38,33 @@
 #include "UObject/Package.h"
 
 #include "USDIncludesStart.h"
-	#include "pxr/usd/usd/typed.h"
-	#include "pxr/usd/usdGeom/imageable.h"
-	#include "pxr/usd/usdGeom/mesh.h"
-	#include "pxr/usd/usdGeom/subset.h"
-	#include "pxr/usd/usdShade/materialBindingAPI.h"
+#include "pxr/usd/usd/typed.h"
+#include "pxr/usd/usdGeom/mesh.h"
+#include "pxr/usd/usdGeom/subset.h"
+#include "pxr/usd/usdShade/materialBindingAPI.h"
 #include "USDIncludesEnd.h"
 
 static int32 GUsdGeometryCacheParallelMeshReads = 16;
 static FAutoConsoleVariableRef CVarUsdGeometryCacheParallelMeshReads(
 	TEXT("USD.GeometryCache.Import.ParallelMeshReads"),
 	GUsdGeometryCacheParallelMeshReads,
-	TEXT("Maximum number of mesh to process in parallel"));
+	TEXT("Maximum number of mesh to process in parallel")
+);
 
 static int32 GUsdGeometryCacheParallelFrameReads = 16;
 static FAutoConsoleVariableRef CVarUsdGeometryCacheParallelFrameReads(
 	TEXT("USD.GeometryCache.Import.ParallelFrameReads"),
 	GUsdGeometryCacheParallelFrameReads,
-	TEXT("Maximum number of mesh frames to read in parallel"));
+	TEXT("Maximum number of mesh frames to read in parallel")
+);
 
 static bool GEnableSubdiv = false;
 static FAutoConsoleVariableRef CVarEnableSubdiv(
 	TEXT("USD.GeometryCache.EnableSubdiv"),
 	GEnableSubdiv,
-	TEXT("Whether to subdivide Mesh prim data when parsing GeometryCaches via OpenSubdiv, the same way we subdivide the Mesh data that ends up in StaticMeshes"));
+	TEXT("Whether to subdivide Mesh prim data when parsing GeometryCaches via OpenSubdiv, the same way we subdivide the Mesh data that ends up in "
+		 "StaticMeshes")
+);
 
 namespace UsdGeometryCacheTranslatorImpl
 {
@@ -101,8 +104,14 @@ namespace UsdGeometryCacheTranslatorImpl
 				else
 				{
 					// Warn, but still add a material slot to preserve the materials order
-					UE_LOG(LogUsd, Warning, TEXT("Failed to resolve material '%s' for slot '%d' for geometry cache '%s'"),
-						*Slot.MaterialSource, LODSlotIndex, *UsdToUnreal::ConvertPath(UsdPrim.GetPath()));
+					UE_LOG(
+						LogUsd,
+						Warning,
+						TEXT("Failed to resolve material '%s' for slot '%d' for geometry cache '%s'"),
+						*Slot.MaterialSource,
+						LODSlotIndex,
+						*UsdToUnreal::ConvertPath(UsdPrim.GetPath())
+					);
 				}
 
 				if (!GeometryCache.Materials.IsValidIndex(SlotIndex))
@@ -121,7 +130,12 @@ namespace UsdGeometryCacheTranslatorImpl
 		return bMaterialAssignementsHaveChanged;
 	}
 
-	void LoadMeshDescription(pxr::UsdTyped UsdMesh, FMeshDescription& OutMeshDescription, UsdUtils::FUsdPrimMaterialAssignmentInfo& OutMaterialInfo, const UsdToUnreal::FUsdMeshConversionOptions& Options)
+	void LoadMeshDescription(
+		pxr::UsdTyped UsdMesh,
+		FMeshDescription& OutMeshDescription,
+		UsdUtils::FUsdPrimMaterialAssignmentInfo& OutMaterialInfo,
+		const UsdToUnreal::FUsdMeshConversionOptions& Options
+	)
 	{
 		if (!UsdMesh)
 		{
@@ -146,7 +160,12 @@ namespace UsdGeometryCacheTranslatorImpl
 		}
 	}
 
-	void GeometryCacheDataForMeshDescription(FGeometryCacheMeshData& OutMeshData, FMeshDescription& MeshDescription, int32 MaterialIndex, const float SecondsPerFrame);
+	void GeometryCacheDataForMeshDescription(
+		FGeometryCacheMeshData& OutMeshData,
+		FMeshDescription& MeshDescription,
+		int32 MaterialIndex,
+		const float SecondsPerFrame
+	);
 
 	void GetGeometryCacheDataTimeCodeRange(const UE::FUsdStage& Stage, const FString& PrimPath, int32& OutStartFrame, int32& OutEndFrame)
 	{
@@ -198,7 +217,8 @@ namespace UsdGeometryCacheTranslatorImpl
 		FReadMeshDataArgs(const UE::FUsdStage& InStage, const UE::FUsdPrim& InRootPrim)
 			: Stage(InStage)
 			, RootPrim(InRootPrim)
-		{}
+		{
+		}
 
 		UE::FUsdStageWeak Stage;
 		UE::FUsdPrim RootPrim;
@@ -250,13 +270,20 @@ namespace UsdGeometryCacheTranslatorImpl
 		Args.Options.PurposesToLoad = Context->PurposesToLoad;
 		Args.Options.RenderContext = RenderContextToken;
 		Args.Options.MaterialPurpose = MaterialPurposeToken;
-		Args.Options.bMergeIdenticalMaterialSlots = false; // Don't merge because the GeometryCache is processed as unflattened (ie. one track per mesh)
+		Args.Options.bMergeIdenticalMaterialSlots = false;	  // Don't merge because the GeometryCache is processed as unflattened (ie. one track per
+															  // mesh)
 		Args.Options.SubdivisionLevel = GEnableSubdiv ? Context->SubdivisionLevel : 0;
 
 		return Args;
 	}
 
-	bool ReadMeshData(const FReadMeshDataArgs& Args, const UE::FUsdPrim& MeshPrim, int32 MaterialOffset, float Time, FGeometryCacheMeshData& OutMeshData)
+	bool ReadMeshData(
+		const FReadMeshDataArgs& Args,
+		const UE::FUsdPrim& MeshPrim,
+		int32 MaterialOffset,
+		float Time,
+		FGeometryCacheMeshData& OutMeshData
+	)
 	{
 		// MeshDescriptions are always allocated on the UE allocator as the allocation happens within
 		// another dll, so we need to deallocate them using it too
@@ -283,7 +310,8 @@ namespace UsdGeometryCacheTranslatorImpl
 			// Compute the normals and tangents for the mesh
 			const float ComparisonThreshold = THRESH_POINTS_ARE_SAME;
 
-			// This function make sure the Polygon Normals Tangents Binormals are computed and also remove degenerated triangle from the render mesh description.
+			// This function make sure the Polygon Normals Tangents Binormals are computed and also remove degenerated triangle from the render mesh
+			// description.
 			FStaticMeshOperations::ComputeTriangleTangentsAndNormals(MeshDescription, ComparisonThreshold);
 
 			// Compute any missing normals or tangents.
@@ -301,11 +329,17 @@ namespace UsdGeometryCacheTranslatorImpl
 		return false;
 	}
 
-	UGeometryCacheTrackUsd* CreateUsdStreamTrack(UGeometryCache* GeometryCache, const FReadMeshDataArgs& Args, const FString& PrimPath, int32 MaterialOffset)
+	UGeometryCacheTrackUsd* CreateUsdStreamTrack(
+		UGeometryCache* GeometryCache,
+		const FReadMeshDataArgs& Args,
+		const FString& PrimPath,
+		int32 MaterialOffset
+	)
 	{
 		// Create and configure a new USDTrack to be added to the GeometryCache
 		UGeometryCacheTrackUsd* UsdTrack = NewObject<UGeometryCacheTrackUsd>(GeometryCache);
-		UsdTrack->MeshConversionOptions = Args.Options;  // Also pass along the options we'll use for mesh conversion so that we can properly hash the prim
+		UsdTrack->MeshConversionOptions = Args.Options;	   // Also pass along the options we'll use for mesh conversion so that we can properly hash
+														   // the prim
 		UsdTrack->Initialize(
 			Args.Stage,
 			PrimPath,
@@ -342,7 +376,7 @@ namespace UsdGeometryCacheTranslatorImpl
 		FString ObjectName = IUsdClassesModule::SanitizeObjectName(FPaths::GetBaseFilename(PrimPath));
 
 		FName CodecName = MakeUniqueObjectName(GeometryCache, UGeometryCacheCodecV1::StaticClass(), FName(ObjectName + FString(TEXT("_Codec"))));
-		UGeometryCacheCodecV1* Codec = NewObject<UGeometryCacheCodecV1>(GeometryCache, CodecName , RF_Public);
+		UGeometryCacheCodecV1* Codec = NewObject<UGeometryCacheCodecV1>(GeometryCache, CodecName, RF_Public);
 
 		// Compression settings for good quality
 		const float VertexQuantizationPrecision = 0.0005f;
@@ -361,8 +395,15 @@ namespace UsdGeometryCacheTranslatorImpl
 		return StreamableTrack;
 	}
 
-	UGeometryCache* CreateGeometryCache(const UE::FUsdPrim& RootPrim, const FMeshDescription& MeshDescription, const TArray<UE::FSdfPath>& MeshPaths, const TArray<int32>& MaterialOffsets,
-										TSharedRef<FUsdSchemaTranslationContext> Context, bool& bOutIsNew, float& StartOffsetTime)
+	UGeometryCache* CreateGeometryCache(
+		const UE::FUsdPrim& RootPrim,
+		const FMeshDescription& MeshDescription,
+		const TArray<UE::FSdfPath>& MeshPaths,
+		const TArray<int32>& MaterialOffsets,
+		TSharedRef<FUsdSchemaTranslationContext> Context,
+		bool& bOutIsNew,
+		float& StartOffsetTime
+	)
 	{
 		FString RootPrimPath = RootPrim.GetPrimPath().GetString();
 
@@ -396,8 +437,16 @@ namespace UsdGeometryCacheTranslatorImpl
 		{
 			bOutIsNew = true;
 
-			const FName AssetName = MakeUniqueObjectName(GetTransientPackage(), UGeometryCache::StaticClass(), *IUsdClassesModule::SanitizeObjectName(FPaths::GetBaseFilename(RootPrimPath)));
-			GeometryCache = NewObject<UGeometryCache>(GetTransientPackage(), AssetName, Context->ObjectFlags | EObjectFlags::RF_Public | RF_Transient);
+			const FName AssetName = MakeUniqueObjectName(
+				GetTransientPackage(),
+				UGeometryCache::StaticClass(),
+				*IUsdClassesModule::SanitizeObjectName(FPaths::GetBaseFilename(RootPrimPath))
+			);
+			GeometryCache = NewObject<UGeometryCache>(
+				GetTransientPackage(),
+				AssetName,
+				Context->ObjectFlags | EObjectFlags::RF_Public | RF_Transient
+			);
 
 			TOptional<FReadMeshDataArgs> Args;
 			if (!Context->bIsImporting)
@@ -442,7 +491,13 @@ namespace UsdGeometryCacheTranslatorImpl
 		return GeometryCache;
 	}
 
-	void FillGeometryCacheTracks(const FString& RootPrimPath, const TArray<UE::FSdfPath>& MeshPrims, const TArray<int32>& MaterialOffsets, TSharedRef<FUsdSchemaTranslationContext> Context, UGeometryCache* GeometryCache)
+	void FillGeometryCacheTracks(
+		const FString& RootPrimPath,
+		const TArray<UE::FSdfPath>& MeshPrims,
+		const TArray<int32>& MaterialOffsets,
+		TSharedRef<FUsdSchemaTranslationContext> Context,
+		UGeometryCache* GeometryCache
+	)
 	{
 		TRACE_CPUPROFILER_EVENT_SCOPE(UsdGeometryCacheTranslatorImpl::FillGeometryCacheTracks);
 
@@ -456,7 +511,8 @@ namespace UsdGeometryCacheTranslatorImpl
 			{
 				UGeometryCacheTrackStreamable* StreamableTrack = Cast<UGeometryCacheTrackStreamable>(GeometryCache->Tracks[Index]);
 				UE::FUsdPrim MeshPrim = Args.Stage.GetPrimAtPath(MeshPrims[Index]);
-				const bool bConstantTopology = UsdUtils::GetMeshTopologyVariance(pxr::UsdGeomMesh(MeshPrim)) != UsdUtils::EMeshTopologyVariance::Heterogenous;
+				const bool bConstantTopology = UsdUtils::GetMeshTopologyVariance(pxr::UsdGeomMesh(MeshPrim))
+											   != UsdUtils::EMeshTopologyVariance::Heterogenous;
 
 				for (int32 FrameIndex = Args.StartFrame; FrameIndex < Args.EndFrame; ++FrameIndex)
 				{
@@ -499,59 +555,75 @@ namespace UsdGeometryCacheTranslatorImpl
 			}
 
 			// Parallel mesh reads: Meshes can be read independently of each other
-			ParallelFor(NumMeshThreads, [&Args, &MeshPrims, &MaterialOffsets, &SyncEvents, NumMeshThreads, NumFrameThreads, GeometryCache, NumMeshes](int32 MeshThreadIndex)
-			{
-				int32 MeshIndex = MeshThreadIndex;
-
-				while (MeshIndex < NumMeshes)
+			ParallelFor(
+				NumMeshThreads,
+				[&Args, &MeshPrims, &MaterialOffsets, &SyncEvents, NumMeshThreads, NumFrameThreads, GeometryCache, NumMeshes](int32 MeshThreadIndex)
 				{
-					UE::FUsdPrim MeshPrim = Args.Stage.GetPrimAtPath(MeshPrims[MeshIndex]);
-					int32 MaterialOffset = MaterialOffsets[MeshIndex];
-					UGeometryCacheTrackStreamable* StreamableTrack = Cast<UGeometryCacheTrackStreamable>(GeometryCache->Tracks[MeshIndex]);
-					FEvent* FrameWrittenEvent = SyncEvents[MeshThreadIndex];
-					const bool bConstantTopology = UsdUtils::GetMeshTopologyVariance(pxr::UsdGeomMesh(MeshPrim)) != UsdUtils::EMeshTopologyVariance::Heterogenous;
-
-					std::atomic<int32> WriteFrameIndex = Args.StartFrame;
-					FCriticalSection Mutex;
+					int32 MeshIndex = MeshThreadIndex;
 
-					// Parallel frame read: frame data can be read concurrently but have to be processed in order for AddMeshSample
-					ParallelFor(NumFrameThreads, [&Args, &MeshPrim, &FrameWrittenEvent, &WriteFrameIndex, &Mutex, NumFrameThreads, MaterialOffset, StreamableTrack, bConstantTopology](int32 FrameThreadIndex)
+					while (MeshIndex < NumMeshes)
 					{
-						int32 FrameIndex = Args.StartFrame + FrameThreadIndex;
+						UE::FUsdPrim MeshPrim = Args.Stage.GetPrimAtPath(MeshPrims[MeshIndex]);
+						int32 MaterialOffset = MaterialOffsets[MeshIndex];
+						UGeometryCacheTrackStreamable* StreamableTrack = Cast<UGeometryCacheTrackStreamable>(GeometryCache->Tracks[MeshIndex]);
+						FEvent* FrameWrittenEvent = SyncEvents[MeshThreadIndex];
+						const bool bConstantTopology = UsdUtils::GetMeshTopologyVariance(pxr::UsdGeomMesh(MeshPrim))
+													   != UsdUtils::EMeshTopologyVariance::Heterogenous;
+
+						std::atomic<int32> WriteFrameIndex = Args.StartFrame;
+						FCriticalSection Mutex;
+
+						// Parallel frame read: frame data can be read concurrently but have to be processed in order for AddMeshSample
+						ParallelFor(
+							NumFrameThreads,
+							[&Args,
+							 &MeshPrim,
+							 &FrameWrittenEvent,
+							 &WriteFrameIndex,
+							 &Mutex,
+							 NumFrameThreads,
+							 MaterialOffset,
+							 StreamableTrack,
+							 bConstantTopology](int32 FrameThreadIndex)
+							{
+								int32 FrameIndex = Args.StartFrame + FrameThreadIndex;
 
-						while (FrameIndex < Args.EndFrame)
-						{
-							// Read frame data into memory
-							FGeometryCacheMeshData MeshData;
-							ReadMeshData(Args, MeshPrim, MaterialOffset, FrameIndex, MeshData);
+								while (FrameIndex < Args.EndFrame)
+								{
+									// Read frame data into memory
+									FGeometryCacheMeshData MeshData;
+									ReadMeshData(Args, MeshPrim, MaterialOffset, FrameIndex, MeshData);
 
-							// Wait until it's our turn to process this frame.
-							while (WriteFrameIndex < FrameIndex)
-							{
-								const uint32 WaitTimeInMs = 10;
-								FrameWrittenEvent->Wait(WaitTimeInMs);
-							}
+									// Wait until it's our turn to process this frame.
+									while (WriteFrameIndex < FrameIndex)
+									{
+										const uint32 WaitTimeInMs = 10;
+										FrameWrittenEvent->Wait(WaitTimeInMs);
+									}
 
-							{
-								FScopeLock WriteLock(&Mutex);
+									{
+										FScopeLock WriteLock(&Mutex);
 
-								// Add it to the track
-								StreamableTrack->AddMeshSample(MeshData, (FrameIndex - Args.StartFrame) / Args.FramesPerSecond, bConstantTopology);
+										// Add it to the track
+										StreamableTrack
+											->AddMeshSample(MeshData, (FrameIndex - Args.StartFrame) / Args.FramesPerSecond, bConstantTopology);
 
-								// Mark the next frame index as ready for processing.
-								++WriteFrameIndex;
+										// Mark the next frame index as ready for processing.
+										++WriteFrameIndex;
 
-								FrameWrittenEvent->Trigger();
-							}
+										FrameWrittenEvent->Trigger();
+									}
 
-							// Get new frame index to read for next run cycle
-							FrameIndex += NumFrameThreads;
-						};
-					});
+									// Get new frame index to read for next run cycle
+									FrameIndex += NumFrameThreads;
+								};
+							}
+						);
 
-					MeshIndex += NumMeshThreads;
+						MeshIndex += NumMeshThreads;
+					}
 				}
-			});
+			);
 
 			for (FEvent* SyncEvent : SyncEvents)
 			{
@@ -572,7 +644,12 @@ namespace UsdGeometryCacheTranslatorImpl
 	}
 
 	// #ueent_todo: Replace MeshDescription with RawMesh and also make it work with StaticMesh
-	void GeometryCacheDataForMeshDescription(FGeometryCacheMeshData& OutMeshData, FMeshDescription& MeshDescription, int32 MaterialIndex, float FramesPerSecond)
+	void GeometryCacheDataForMeshDescription(
+		FGeometryCacheMeshData& OutMeshData,
+		FMeshDescription& MeshDescription,
+		int32 MaterialIndex,
+		float FramesPerSecond
+	)
 	{
 		TRACE_CPUPROFILER_EVENT_SCOPE(GeometryCacheDataForMeshDescription);
 
@@ -594,14 +671,16 @@ namespace UsdGeometryCacheTranslatorImpl
 
 		FStaticMeshAttributes MeshDescriptionAttributes(MeshDescription);
 
-		TVertexAttributesConstRef< FVector3f > VertexPositions = MeshDescriptionAttributes.GetVertexPositions();
-		TVertexInstanceAttributesConstRef< FVector3f > VertexInstanceNormals = MeshDescriptionAttributes.GetVertexInstanceNormals();
-		TVertexInstanceAttributesConstRef< FVector3f > VertexInstanceTangents = MeshDescriptionAttributes.GetVertexInstanceTangents();
-		TVertexInstanceAttributesConstRef< float > VertexInstanceBinormalSigns = MeshDescriptionAttributes.GetVertexInstanceBinormalSigns();
-		TVertexInstanceAttributesConstRef< FVector4f > VertexInstanceColors = MeshDescriptionAttributes.GetVertexInstanceColors();
-		TVertexInstanceAttributesConstRef< FVector2f > VertexInstanceUVs = MeshDescriptionAttributes.GetVertexInstanceUVs();
+		TVertexAttributesConstRef<FVector3f> VertexPositions = MeshDescriptionAttributes.GetVertexPositions();
+		TVertexInstanceAttributesConstRef<FVector3f> VertexInstanceNormals = MeshDescriptionAttributes.GetVertexInstanceNormals();
+		TVertexInstanceAttributesConstRef<FVector3f> VertexInstanceTangents = MeshDescriptionAttributes.GetVertexInstanceTangents();
+		TVertexInstanceAttributesConstRef<float> VertexInstanceBinormalSigns = MeshDescriptionAttributes.GetVertexInstanceBinormalSigns();
+		TVertexInstanceAttributesConstRef<FVector4f> VertexInstanceColors = MeshDescriptionAttributes.GetVertexInstanceColors();
+		TVertexInstanceAttributesConstRef<FVector2f> VertexInstanceUVs = MeshDescriptionAttributes.GetVertexInstanceUVs();
 
-		TVertexInstanceAttributesConstRef< FVector3f > VertexInstanceVelocities = MeshDescription.VertexInstanceAttributes().GetAttributesRef<FVector3f>(MeshAttribute::VertexInstance::Velocity);
+		TVertexInstanceAttributesConstRef<FVector3f>
+			VertexInstanceVelocities = MeshDescription.VertexInstanceAttributes().GetAttributesRef<FVector3f>(MeshAttribute::VertexInstance::Velocity
+			);
 
 		const bool bHasVelocities = VertexInstanceVelocities.IsValid();
 		OutMeshData.VertexInfo.bHasMotionVectors = bHasVelocities;
@@ -662,7 +741,7 @@ namespace UsdGeometryCacheTranslatorImpl
 						if (bHasVelocities)
 						{
 							FVector3f MotionVector = VertexInstanceVelocities[VertexInstanceID];
-							MotionVector *= (-1.f / FramesPerSecond); // Velocity is per seconds but we need per frame for motion vectors
+							MotionVector *= (-1.f / FramesPerSecond);	 // Velocity is per seconds but we need per frame for motion vectors
 
 							OutMeshData.MotionVectors.Add(MotionVector);
 						}
@@ -683,7 +762,7 @@ namespace UsdGeometryCacheTranslatorImpl
 class FGeometryCacheCreateAssetsTaskChain : public FBuildStaticMeshTaskChain
 {
 public:
-	explicit FGeometryCacheCreateAssetsTaskChain(const TSharedRef< FUsdSchemaTranslationContext >& InContext, const UE::FSdfPath& InPrimPath)
+	explicit FGeometryCacheCreateAssetsTaskChain(const TSharedRef<FUsdSchemaTranslationContext>& InContext, const UE::FSdfPath& InPrimPath)
 		: FBuildStaticMeshTaskChain(InContext, InPrimPath)
 	{
 		SetupTasks();
@@ -700,52 +779,54 @@ void FGeometryCacheCreateAssetsTaskChain::SetupTasks()
 {
 	// Create the mesh description (Async)
 	Do(ESchemaTranslationLaunchPolicy::Async,
-		[this]() -> bool
-		{
-			FScopedUnrealAllocs UnrealAllocs;
-
-			pxr::TfToken RenderContextToken = pxr::UsdShadeTokens->universalRenderContext;
-			if (!Context->RenderContext.IsNone())
-			{
-				RenderContextToken = UnrealToUsd::ConvertToken(*Context->RenderContext.ToString()).Get();
-			}
-
-			pxr::TfToken MaterialPurposeToken = pxr::UsdShadeTokens->allPurpose;
-			if (!Context->MaterialPurpose.IsNone())
-			{
-				MaterialPurposeToken = UnrealToUsd::ConvertToken(*Context->MaterialPurpose.ToString()).Get();
-			}
-
-			UsdToUnreal::FUsdMeshConversionOptions Options;
-			Options.TimeCode = UsdUtils::GetEarliestTimeCode();
-			Options.PurposesToLoad = Context->PurposesToLoad;
-			Options.RenderContext = RenderContextToken;
-			Options.MaterialPurpose = MaterialPurposeToken;
-			Options.bMergeIdenticalMaterialSlots = false; // Don't merge because the GeometryCache is processed as unflattened (ie. one track per mesh)
-			Options.SubdivisionLevel = GEnableSubdiv ? Context->SubdivisionLevel : 0;
-
-			// GeometryCache has only one LOD so add just one MeshDescription and MaterialAssignmentInfo
-			FMeshDescription& AddedMeshDescription = LODIndexToMeshDescription.Emplace_GetRef();
-			UsdUtils::FUsdPrimMaterialAssignmentInfo& AssignmentInfo = LODIndexToMaterialInfo.Emplace_GetRef();
-
-			// The collapsed mesh description here will be used to cache the GeometryCache asset, but not to fill it since its content will be unflattened
-			// Bake the prim's transform into the mesh data
-			const bool bSkipRootPrimTransformAndVis = false;
-			UsdToUnreal::ConvertGeomMeshHierarchy(GetPrim(), AddedMeshDescription, AssignmentInfo, Options, bSkipRootPrimTransformAndVis);
-
-			// If we have at least one valid LOD, we should proceed to the next step
-			for (const FMeshDescription& MeshDescription : LODIndexToMeshDescription)
-			{
-				if (!MeshDescription.IsEmpty())
-				{
-					return true;
-				}
-			}
-			return false;
-		});
+	   [this]() -> bool
+	   {
+		   FScopedUnrealAllocs UnrealAllocs;
+
+		   pxr::TfToken RenderContextToken = pxr::UsdShadeTokens->universalRenderContext;
+		   if (!Context->RenderContext.IsNone())
+		   {
+			   RenderContextToken = UnrealToUsd::ConvertToken(*Context->RenderContext.ToString()).Get();
+		   }
+
+		   pxr::TfToken MaterialPurposeToken = pxr::UsdShadeTokens->allPurpose;
+		   if (!Context->MaterialPurpose.IsNone())
+		   {
+			   MaterialPurposeToken = UnrealToUsd::ConvertToken(*Context->MaterialPurpose.ToString()).Get();
+		   }
+
+		   UsdToUnreal::FUsdMeshConversionOptions Options;
+		   Options.TimeCode = UsdUtils::GetEarliestTimeCode();
+		   Options.PurposesToLoad = Context->PurposesToLoad;
+		   Options.RenderContext = RenderContextToken;
+		   Options.MaterialPurpose = MaterialPurposeToken;
+		   Options.bMergeIdenticalMaterialSlots = false;	// Don't merge because the GeometryCache is processed as unflattened (ie. one track per
+															// mesh)
+		   Options.SubdivisionLevel = GEnableSubdiv ? Context->SubdivisionLevel : 0;
+
+		   // GeometryCache has only one LOD so add just one MeshDescription and MaterialAssignmentInfo
+		   FMeshDescription& AddedMeshDescription = LODIndexToMeshDescription.Emplace_GetRef();
+		   UsdUtils::FUsdPrimMaterialAssignmentInfo& AssignmentInfo = LODIndexToMaterialInfo.Emplace_GetRef();
+
+		   // The collapsed mesh description here will be used to cache the GeometryCache asset, but not to fill it since its content will be
+		   // unflattened Bake the prim's transform into the mesh data
+		   const bool bSkipRootPrimTransformAndVis = false;
+		   UsdToUnreal::ConvertGeomMeshHierarchy(GetPrim(), AddedMeshDescription, AssignmentInfo, Options, bSkipRootPrimTransformAndVis);
+
+		   // If we have at least one valid LOD, we should proceed to the next step
+		   for (const FMeshDescription& MeshDescription : LODIndexToMeshDescription)
+		   {
+			   if (!MeshDescription.IsEmpty())
+			   {
+				   return true;
+			   }
+		   }
+		   return false;
+	   });
 
 	// Create the GeometryCache (Main thread)
-	Then(ESchemaTranslationLaunchPolicy::Sync,
+	Then(
+		ESchemaTranslationLaunchPolicy::Sync,
 		[this]() -> bool
 		{
 			{
@@ -792,11 +873,20 @@ void FGeometryCacheCreateAssetsTaskChain::SetupTasks()
 
 			bool bIsNew = true;
 			float StartTimeOffset = 0.0f;
-			GeometryCache.Reset(UsdGeometryCacheTranslatorImpl::CreateGeometryCache(GetPrim(), LODIndexToMeshDescription[0], MeshPrimPaths, MaterialOffsets, Context, bIsNew, StartTimeOffset));
+			GeometryCache.Reset(UsdGeometryCacheTranslatorImpl::CreateGeometryCache(
+				GetPrim(),
+				LODIndexToMeshDescription[0],
+				MeshPrimPaths,
+				MaterialOffsets,
+				Context,
+				bIsNew,
+				StartTimeOffset
+			));
 
 			if (GeometryCache)
 			{
-				if (UUsdGeometryCacheAssetUserData* UserData = UsdUtils::GetOrCreateAssetUserData<UUsdGeometryCacheAssetUserData>(GeometryCache.Get()))
+				if (UUsdGeometryCacheAssetUserData* UserData = UsdUtils::GetOrCreateAssetUserData<UUsdGeometryCacheAssetUserData>(GeometryCache.Get()
+					))
 				{
 					UserData->PrimvarToUVIndex = LODIndexToMaterialInfo[0].PrimvarToUVIndex;	// We use the same primvar mapping for all LODs
 					UserData->LayerStartOffsetSeconds = StartTimeOffset;
@@ -842,25 +932,36 @@ void FGeometryCacheCreateAssetsTaskChain::SetupTasks()
 
 			// Continue with the import steps
 			return Context->bIsImporting && GeometryCache && bIsNew;
-		});
+		}
+	);
 
 	// Fill the GeometryCache tracks with the frame data
-	// It is done Sync to avoid starvation issue because FillGeometryCacheTracks is highly parallelized based on the number of meshes and frames to read
-	// Filling GeometryCaches in parallel could cause deadlocks
-	Then(ESchemaTranslationLaunchPolicy::Sync,
+	// It is done Sync to avoid starvation issue because FillGeometryCacheTracks is highly parallelized based on the number of meshes and frames to
+	// read Filling GeometryCaches in parallel could cause deadlocks
+	Then(
+		ESchemaTranslationLaunchPolicy::Sync,
 		[this]() -> bool
 		{
-			UsdGeometryCacheTranslatorImpl::FillGeometryCacheTracks(PrimPath.GetString(), MeshPrimPaths, MaterialOffsets, Context, GeometryCache.Get());
+			UsdGeometryCacheTranslatorImpl::FillGeometryCacheTracks(
+				PrimPath.GetString(),
+				MeshPrimPaths,
+				MaterialOffsets,
+				Context,
+				GeometryCache.Get()
+			);
 			return true;
-		});
+		}
+	);
 
 	// Finalize the GeometryCache (Main Thread)
-	Then(ESchemaTranslationLaunchPolicy::Sync,
+	Then(
+		ESchemaTranslationLaunchPolicy::Sync,
 		[this]() -> bool
 		{
 			UsdGeometryCacheTranslatorImpl::FinalizeGeometryCache(GeometryCache.Get());
 			return false;
-		});
+		}
+	);
 }
 
 void FUsdGeometryCacheTranslator::CreateAssets()
@@ -958,9 +1059,11 @@ USceneComponent* FUsdGeometryCacheTranslator::CreateComponents()
 
 					// For the groom binding to work, the GroomComponent must be a child of the SceneComponent
 					// so the Context ParentComponent is set to the SceneComponent temporarily
-					TGuardValue< USceneComponent* > ParentComponentGuard{Context->ParentComponent, SceneComponent};
+					TGuardValue<USceneComponent*> ParentComponentGuard{Context->ParentComponent, SceneComponent};
 					const bool bNeedsActor = false;
-					UGroomComponent* GroomComponent = Cast< UGroomComponent >(CreateComponentsEx(TSubclassOf< USceneComponent >(UGroomComponent::StaticClass()), bNeedsActor));
+					UGroomComponent* GroomComponent = Cast<UGroomComponent>(
+						CreateComponentsEx(TSubclassOf<USceneComponent>(UGroomComponent::StaticClass()), bNeedsActor)
+					);
 					if (GroomComponent)
 					{
 						UpdateComponents(SceneComponent);
@@ -1156,4 +1259,4 @@ bool FUsdGeometryCacheTranslator::IsPotentialGeometryCacheRoot() const
 	return Context->InfoCache->IsPotentialGeometryCacheRoot(PrimPath);
 }
 
-#endif // #if USE_USD_SDK
+#endif	  // #if USE_USD_SDK
