USD: Reformat USDSchemas module.

diff --git a/Engine/Plugins/Importers/USDImporter/Source/USDSchemas/Private/USDGroomTranslator.cpp b/Engine/Plugins/Importers/USDImporter/Source/USDSchemas/Private/USDGroomTranslator.cpp
--- a/Engine/Plugins/Importers/USDImporter/Source/USDSchemas/Private/USDGroomTranslator.cpp
+++ b/Engine/Plugins/Importers/USDImporter/Source/USDSchemas/Private/USDGroomTranslator.cpp
@@ -26,20 +26,20 @@
 #include "USDIntegrationUtils.h"
 #include "USDLog.h"
 #include "USDPrimConversion.h"
-#include "USDTypesConversion.h"
 
 #include "USDIncludesStart.h"
-	#include "pxr/base/tf/type.h"
-	#include "pxr/usd/usd/prim.h"
-	#include "pxr/usd/usd/timeCode.h"
-	#include "pxr/usd/usdGeom/basisCurves.h"
-	#include "pxr/usd/usdGeom/curves.h"
-	#include "pxr/usd/usdGeom/imageable.h"
+#include "pxr/usd/usd/prim.h"
+#include "pxr/usd/usd/timeCode.h"
+#include "pxr/usd/usdGeom/basisCurves.h"
+#include "pxr/usd/usdGeom/imageable.h"
 #include "USDIncludesEnd.h"
 
 namespace UE::UsdGroomTranslator::Private
 {
-	UGroomImportOptions* CreateGroomImportOptions(const FHairDescriptionGroups& GroupsDescription, const TArray<FHairGroupsInterpolation>& BuildSettings)
+	UGroomImportOptions* CreateGroomImportOptions(
+		const FHairDescriptionGroups& GroupsDescription,
+		const TArray<FHairGroupsInterpolation>& BuildSettings
+	)
 	{
 		// Create a new groom import options and populate the interpolation settings based on the group count
 		UGroomImportOptions* ImportOptions = NewObject<UGroomImportOptions>();
@@ -107,7 +107,7 @@ namespace UE::UsdGroomTranslator::Private
 			Hash.Update((const uint8*)SHAHash.Hash, sizeof(SHAHash.Hash));
 		}
 	}
-}
+}	 // namespace UE::UsdGroomTranslator::Private
 
 class FUsdGroomCreateAssetsTaskChain : public FUsdSchemaTranslatorTaskChain
 {
@@ -136,7 +136,10 @@ protected:
 	FString PrefixedGroomCacheHash;
 
 protected:
-	UE::FUsdPrim GetPrim() const { return Context->Stage.GetPrimAtPath(PrimPath); }
+	UE::FUsdPrim GetPrim() const
+	{
+		return Context->Stage.GetPrimAtPath(PrimPath);
+	}
 
 protected:
 	void SetupTasks()
@@ -145,24 +148,31 @@ protected:
 
 		// Create hair description (Async)
 		Do(ESchemaTranslationLaunchPolicy::Async,
-			[this]() -> bool
-			{
-				const bool bSuccess = UsdToUnreal::ConvertGroomHierarchy(GetPrim(), pxr::UsdTimeCode::EarliestTime(), FTransform::Identity, HairDescription, &AnimInfo);
-
-				const double StageTimeCodesPerSecond = Context->Stage.GetTimeCodesPerSecond();
-				AnimInfo.SecondsPerFrame = float(1.0f / StageTimeCodesPerSecond);
-
-				if (bSuccess && AnimInfo.IsValid())
-				{
-					AnimInfo.Duration = AnimInfo.NumFrames * AnimInfo.SecondsPerFrame;
-					AnimInfo.StartTime = AnimInfo.StartFrame * AnimInfo.SecondsPerFrame;
-					AnimInfo.EndTime = AnimInfo.EndFrame * AnimInfo.SecondsPerFrame;
-				}
-
-				return bSuccess && HairDescription.IsValid();
-			});
+		   [this]() -> bool
+		   {
+			   const bool bSuccess = UsdToUnreal::ConvertGroomHierarchy(
+				   GetPrim(),
+				   pxr::UsdTimeCode::EarliestTime(),
+				   FTransform::Identity,
+				   HairDescription,
+				   &AnimInfo
+			   );
+
+			   const double StageTimeCodesPerSecond = Context->Stage.GetTimeCodesPerSecond();
+			   AnimInfo.SecondsPerFrame = float(1.0f / StageTimeCodesPerSecond);
+
+			   if (bSuccess && AnimInfo.IsValid())
+			   {
+				   AnimInfo.Duration = AnimInfo.NumFrames * AnimInfo.SecondsPerFrame;
+				   AnimInfo.StartTime = AnimInfo.StartFrame * AnimInfo.SecondsPerFrame;
+				   AnimInfo.EndTime = AnimInfo.EndFrame * AnimInfo.SecondsPerFrame;
+			   }
+
+			   return bSuccess && HairDescription.IsValid();
+		   });
 		// Build groom asset from hair description (Sync)
-		Then(ESchemaTranslationLaunchPolicy::Sync,
+		Then(
+			ESchemaTranslationLaunchPolicy::Sync,
 			[this]()
 			{
 				TRACE_CPUPROFILER_EVENT_SCOPE(FUsdGroomCreateAssetsTaskChain::Build);
@@ -171,7 +181,8 @@ protected:
 				FHairDescriptionGroups GroupsDescription;
 				FGroomBuilder::BuildHairDescriptionGroups(HairDescription, GroupsDescription);
 
-				ImportOptions.Reset(UE::UsdGroomTranslator::Private::CreateGroomImportOptions(GroupsDescription, Context->GroomInterpolationSettings));
+				ImportOptions.Reset(UE::UsdGroomTranslator::Private::CreateGroomImportOptions(GroupsDescription, Context->GroomInterpolationSettings)
+				);
 
 				FSHAHash SHAHash = UE::UsdGroomTranslator::Private::ComputeHairDescriptionHash(HairDescription, ImportOptions->InterpolationSettings);
 
@@ -225,9 +236,11 @@ protected:
 
 				// Next step is to parse the GroomCache data if it was determined that the groom has animated attributes
 				return GroomAsset != nullptr && AnimInfo.IsValid();
-			});
+			}
+		);
 		// Parse GroomCache data into processor (Async)
-		Then(ESchemaTranslationLaunchPolicy::Async,
+		Then(
+			ESchemaTranslationLaunchPolicy::Async,
 			[this]() -> bool
 			{
 				if (!Context->InfoCache)
@@ -235,9 +248,7 @@ protected:
 					return false;
 				}
 
-				UGroomAsset* GroomAsset = Context->InfoCache->GetSingleAssetForPrim<UGroomAsset>(
-					PrimPath
-				);
+				UGroomAsset* GroomAsset = Context->InfoCache->GetSingleAssetForPrim<UGroomAsset>(PrimPath);
 				if (!GroomAsset)
 				{
 					return false;
@@ -245,8 +256,18 @@ protected:
 
 				// Compute GroomCache hash from the first and last frame HairDescription...
 				FSHA1 SHA1;
-				UE::UsdGroomTranslator::Private::ComputeFrameHairDescriptionHash(GetPrim(), ImportOptions->InterpolationSettings, AnimInfo.StartFrame, SHA1);
-				UE::UsdGroomTranslator::Private::ComputeFrameHairDescriptionHash(GetPrim(), ImportOptions->InterpolationSettings, AnimInfo.EndFrame, SHA1);
+				UE::UsdGroomTranslator::Private::ComputeFrameHairDescriptionHash(
+					GetPrim(),
+					ImportOptions->InterpolationSettings,
+					AnimInfo.StartFrame,
+					SHA1
+				);
+				UE::UsdGroomTranslator::Private::ComputeFrameHairDescriptionHash(
+					GetPrim(),
+					ImportOptions->InterpolationSettings,
+					AnimInfo.EndFrame,
+					SHA1
+				);
 
 				// Along with relevant AnimInfo data
 				SHA1.Update((const uint8*)&AnimInfo.NumFrames, sizeof(AnimInfo.NumFrames));
@@ -272,7 +293,12 @@ protected:
 					for (int32 FrameIndex = AnimInfo.StartFrame; FrameIndex < AnimInfo.EndFrame + 1; ++FrameIndex)
 					{
 						FHairDescription FrameHairDescription;
-						bSuccess = UsdToUnreal::ConvertGroomHierarchy(GetPrim(), pxr::UsdTimeCode(FrameIndex), FTransform::Identity, FrameHairDescription);
+						bSuccess = UsdToUnreal::ConvertGroomHierarchy(
+							GetPrim(),
+							pxr::UsdTimeCode(FrameIndex),
+							FTransform::Identity,
+							FrameHairDescription
+						);
 
 						if (!bSuccess)
 						{
@@ -295,7 +321,13 @@ protected:
 						{
 							const FHairDescriptionGroup& HairGroup = HairDescriptionGroups.HairGroups[GroupIndex];
 							FHairDescriptionGroup& HairGroupData = HairGroupsData[GroupIndex];
-							FGroomBuilder::BuildData(HairGroup, GroomAsset->GetHairGroupsInterpolation()[GroupIndex], HairGroupsInfo[GroupIndex], HairGroupData.Strands, HairGroupData.Guides);
+							FGroomBuilder::BuildData(
+								HairGroup,
+								GroomAsset->GetHairGroupsInterpolation()[GroupIndex],
+								HairGroupsInfo[GroupIndex],
+								HairGroupData.Strands,
+								HairGroupData.Guides
+							);
 						}
 
 						// Validate that the GroomCache has the same topology as the static groom
@@ -303,11 +335,19 @@ protected:
 						{
 							for (uint32 GroupIndex = 0; GroupIndex < GroupCount; ++GroupIndex)
 							{
-								if (HairGroupsData[GroupIndex].Strands.GetNumPoints() != GroomHairGroupsData[GroupIndex].Strands.BulkData.GetNumPoints())
+								if (HairGroupsData[GroupIndex].Strands.GetNumPoints()
+									!= GroomHairGroupsData[GroupIndex].Strands.BulkData.GetNumPoints())
 								{
 									bSuccess = false;
-									UE_LOG(LogUsd, Warning, TEXT("GroomCache frame %d does not have the same number of vertices as the static groom (%u instead of %u). Aborting GroomCache import."),
-										FrameIndex, HairGroupsData[GroupIndex].Strands.GetNumPoints(), GroomHairGroupsData[GroupIndex].Strands.BulkData.GetNumPoints());
+									UE_LOG(
+										LogUsd,
+										Warning,
+										TEXT("GroomCache frame %d does not have the same number of vertices as the static groom (%u instead of %u). "
+											 "Aborting GroomCache import."),
+										FrameIndex,
+										HairGroupsData[GroupIndex].Strands.GetNumPoints(),
+										GroomHairGroupsData[GroupIndex].Strands.BulkData.GetNumPoints()
+									);
 									break;
 								}
 							}
@@ -315,8 +355,14 @@ protected:
 						else
 						{
 							bSuccess = false;
-							UE_LOG(LogUsd, Warning, TEXT("GroomCache does not have the same number of groups as the static groom (%d instead of %d). Aborting GroomCache import."),
-								HairGroupsData.Num(), GroomHairGroupsData.Num());
+							UE_LOG(
+								LogUsd,
+								Warning,
+								TEXT("GroomCache does not have the same number of groups as the static groom (%d instead of %d). Aborting GroomCache "
+									 "import."),
+								HairGroupsData.Num(),
+								GroomHairGroupsData.Num()
+							);
 						}
 
 						if (!bSuccess)
@@ -341,9 +387,11 @@ protected:
 
 				// Go to the next step only if it needs to create a GroomCache asset
 				return GroomCache == nullptr && GroomCacheProcessor;
-			});
+			}
+		);
 		// Create GroomCache asset from processor (Sync)
-		Then(ESchemaTranslationLaunchPolicy::Sync,
+		Then(
+			ESchemaTranslationLaunchPolicy::Sync,
 			[this]() -> bool
 			{
 				// TEMP: This is a small trick to prevent two concurrent task chains from running into a hash collision in the asset cache.
@@ -358,7 +406,8 @@ protected:
 				}
 
 				const FString StrandsGroomCachePrimPath = UsdGroomTranslatorUtils::GetStrandsGroomCachePrimPath(PrimPath);
-				FHairImportContext HairImportContext(nullptr, GetTransientPackage(), nullptr, FName(), Context->ObjectFlags | RF_Public | RF_Transient);
+				FHairImportContext
+					HairImportContext(nullptr, GetTransientPackage(), nullptr, FName(), Context->ObjectFlags | RF_Public | RF_Transient);
 				FName UniqueName = MakeUniqueObjectName(
 					GetTransientPackage(),
 					UGroomCache::StaticClass(),
@@ -366,7 +415,8 @@ protected:
 				);
 
 				// Once the processing has completed successfully, the data is transferred to the GroomCache
-				UGroomCache* GroomCache = FGroomCacheImporter::ProcessToGroomCache(*GroomCacheProcessor, AnimInfo, HairImportContext, UniqueName.ToString());
+				UGroomCache*
+					GroomCache = FGroomCacheImporter::ProcessToGroomCache(*GroomCacheProcessor, AnimInfo, HairImportContext, UniqueName.ToString());
 				if (GroomCache)
 				{
 					if (UUsdAssetUserData* UserData = UsdUtils::GetOrCreateAssetUserData(GroomCache))
@@ -390,7 +440,8 @@ protected:
 				}
 
 				return GroomCache != nullptr;
-			});
+			}
+		);
 	}
 };
 
@@ -458,11 +509,9 @@ void FUsdGroomTranslator::UpdateComponents(USceneComponent* SceneComponent)
 			GroomComponent->Modify();
 
 			UGroomAsset* Groom = nullptr;
-			if(Context->InfoCache)
+			if (Context->InfoCache)
 			{
-				Groom = Context->InfoCache->GetSingleAssetForPrim<UGroomAsset>(
-					PrimPath
-				);
+				Groom = Context->InfoCache->GetSingleAssetForPrim<UGroomAsset>(PrimPath);
 			}
 
 			bool bShouldRegister = false;
@@ -569,4 +618,4 @@ TSet<UE::FSdfPath> FUsdGroomTranslator::CollectAuxiliaryPrims() const
 	return Result;
 }
 
-#endif // #if USE_USD_SDK
\ No newline at end of file
+#endif	  // #if USE_USD_SDK
diff --git a/Engine/Plugins/Importers/USDImporter/Source/USDSchemas/Private/USDGroomTranslatorUtils.cpp b/Engine/Plugins/Importers/USDImporter/Source/USDSchemas/Private/USDGroomTranslatorUtils.cpp
--- a/Engine/Plugins/Importers/USDImporter/Source/USDSchemas/Private/USDGroomTranslatorUtils.cpp
+++ b/Engine/Plugins/Importers/USDImporter/Source/USDSchemas/Private/USDGroomTranslatorUtils.cpp
@@ -1,41 +1,40 @@
 // Copyright Epic Games, Inc. All Rights Reserved.
 
 #include "USDGroomTranslatorUtils.h"
-#include "UObject/Package.h"
 
 #if USE_USD_SDK && WITH_EDITOR
 
+#include "UnrealUSDWrapper.h"
+#include "USDAssetCache2.h"
+#include "USDClassesModule.h"
+#include "USDInfoCache.h"
+#include "USDIntegrationUtils.h"
+#include "USDTypesConversion.h"
+#include "UsdWrappers/SdfPath.h"
+
 #include "Engine/SkeletalMesh.h"
 #include "GeometryCache.h"
 #include "GroomAsset.h"
 #include "GroomBindingAsset.h"
 #include "GroomComponent.h"
 #include "Misc/SecureHash.h"
+#include "UObject/Package.h"
 #include "UObject/UObjectGlobals.h"
-
-#include "UnrealUSDWrapper.h"
-#include "USDAssetCache.h"
-#include "USDClassesModule.h"
-#include "USDGeomXformableTranslator.h"
-#include "USDInfoCache.h"
-#include "USDIntegrationUtils.h"
-#include "USDSchemaTranslator.h"
-#include "USDTypesConversion.h"
-#include "UsdWrappers/SdfPath.h"
+#include "USDConversionUtils.h"
 
 #include "USDIncludesStart.h"
-	#include "pxr/usd/usd/prim.h"
-	#include "pxr/usd/usd/relationship.h"
-	#include "pxr/usd/usdGeom/mesh.h"
-	#include "pxr/usd/usdSkel/root.h"
-	#include "pxr/usd/usdSkel/skeleton.h"
+#include "pxr/usd/usd/prim.h"
+#include "pxr/usd/usd/relationship.h"
+#include "pxr/usd/usdGeom/mesh.h"
+#include "pxr/usd/usdSkel/root.h"
+#include "pxr/usd/usdSkel/skeleton.h"
 #include "USDIncludesEnd.h"
 
 namespace UE::UsdGroomTranslatorUtils::Private
 {
 	struct FGroomBindingBuildSettings : public FGCObject
 	{
-		virtual void AddReferencedObjects(FReferenceCollector& Collector)  override
+		virtual void AddReferencedObjects(FReferenceCollector& Collector) override
 		{
 			Collector.AddReferencedObject(Groom);
 			Collector.AddReferencedObject(SourceMesh);
@@ -185,7 +184,7 @@ namespace UE::UsdGroomTranslatorUtils::Private
 		}
 		return nullptr;
 	}
-}
+}	 // namespace UE::UsdGroomTranslatorUtils::Private
 
 namespace UsdGroomTranslatorUtils
 {
@@ -210,9 +209,7 @@ namespace UsdGroomTranslatorUtils
 		}
 
 		// The GroomAsset should already be processed and cached by the USDGroomTranslator
-		UGroomAsset* GroomAsset = InfoCache.GetSingleAssetForPrim<UGroomAsset>(
-			UE::FSdfPath{*GroomPrimPath}
-		);
+		UGroomAsset* GroomAsset = InfoCache.GetSingleAssetForPrim<UGroomAsset>(UE::FSdfPath{*GroomPrimPath});
 		if (!GroomAsset)
 		{
 			return;
@@ -275,9 +272,7 @@ namespace UsdGroomTranslatorUtils
 			return;
 		}
 
-		UGroomAsset* GroomAsset = InfoCache.GetSingleAssetForPrim<UGroomAsset>(
-			UE::FSdfPath{*GroomPrimPath}
-		);
+		UGroomAsset* GroomAsset = InfoCache.GetSingleAssetForPrim<UGroomAsset>(UE::FSdfPath{*GroomPrimPath});
 		if (!GroomAsset)
 		{
 			return;
@@ -285,9 +280,7 @@ namespace UsdGroomTranslatorUtils
 
 		const FString PrimPath(UsdToUnreal::ConvertPath(Prim.GetPath()));
 		const FString GroomBindingPath = FString::Printf(TEXT("%s_groombinding"), *PrimPath);
-		UGroomBindingAsset* GroomBinding = InfoCache.GetSingleAssetForPrim<UGroomBindingAsset>(
-			UE::FSdfPath{*GroomBindingPath}
-		);
+		UGroomBindingAsset* GroomBinding = InfoCache.GetSingleAssetForPrim<UGroomBindingAsset>(UE::FSdfPath{*GroomBindingPath});
 
 		// Set the GroomAsset and GroomBindingAsset on the child GroomComponent of SceneComponent that was set up in the translator
 		TArray<USceneComponent*> Children;
@@ -310,6 +303,6 @@ namespace UsdGroomTranslatorUtils
 	{
 		return FString::Printf(TEXT("%s_strands_cache"), *PrimPath.GetString());
 	}
-}
+}	 // namespace UsdGroomTranslatorUtils
 
-#endif // #if USE_USD_SDK && WITH_EDITOR
+#endif	  // #if USE_USD_SDK && WITH_EDITOR
diff --git a/Engine/Plugins/Importers/USDImporter/Source/USDSchemas/Private/USDInfoCache.cpp b/Engine/Plugins/Importers/USDImporter/Source/USDSchemas/Private/USDInfoCache.cpp
--- a/Engine/Plugins/Importers/USDImporter/Source/USDSchemas/Private/USDInfoCache.cpp
+++ b/Engine/Plugins/Importers/USDImporter/Source/USDSchemas/Private/USDInfoCache.cpp
@@ -2,9 +2,7 @@
 
 #include "USDInfoCache.h"
 
-#include "USDAssetCache2.h"
 #include "USDGeomMeshConversion.h"
-#include "USDGeomXformableTranslator.h"
 #include "USDLog.h"
 #include "USDSchemasModule.h"
 #include "USDSchemaTranslator.h"
@@ -14,47 +12,49 @@
 #include "UsdWrappers/UsdPrim.h"
 
 #include "Async/ParallelFor.h"
-#include "Math/NumericLimits.h"
 #include "Misc/Paths.h"
 #include "Modules/ModuleManager.h"
 
 #if USE_USD_SDK
 #include "USDIncludesStart.h"
-	#include "pxr/usd/pcp/layerStack.h"
-	#include "pxr/usd/sdf/path.h"
-	#include "pxr/usd/usd/collectionAPI.h"
-	#include "pxr/usd/usd/prim.h"
-	#include "pxr/usd/usd/primCompositionQuery.h"
-	#include "pxr/usd/usdGeom/mesh.h"
-	#include "pxr/usd/usdGeom/pointInstancer.h"
-	#include "pxr/usd/usdGeom/scope.h"
-	#include "pxr/usd/usdGeom/subset.h"
-	#include "pxr/usd/usdGeom/xform.h"
-	#include "pxr/usd/usdShade/materialBindingAPI.h"
-	#include "pxr/usd/usdSkel/bindingAPI.h"
-	#include "pxr/usd/usdSkel/root.h"
-	#include "pxr/usd/usd/primRange.h"
+#include "pxr/usd/pcp/layerStack.h"
+#include "pxr/usd/sdf/path.h"
+#include "pxr/usd/usd/collectionAPI.h"
+#include "pxr/usd/usd/prim.h"
+#include "pxr/usd/usd/primCompositionQuery.h"
+#include "pxr/usd/usd/primRange.h"
+#include "pxr/usd/usdGeom/mesh.h"
+#include "pxr/usd/usdGeom/pointInstancer.h"
+#include "pxr/usd/usdGeom/scope.h"
+#include "pxr/usd/usdGeom/subset.h"
+#include "pxr/usd/usdGeom/xform.h"
+#include "pxr/usd/usdShade/materialBindingAPI.h"
+#include "pxr/usd/usdSkel/bindingAPI.h"
+#include "pxr/usd/usdSkel/root.h"
 #include "USDIncludesEnd.h"
-#endif // USE_USD_SDK
+#endif	  // USE_USD_SDK
 
 static int32 GMaxNumVerticesCollapsedMesh = 5000000;
 static FAutoConsoleVariableRef CVarMaxNumVerticesCollapsedMesh(
 	TEXT("USD.MaxNumVerticesCollapsedMesh"),
 	GMaxNumVerticesCollapsedMesh,
-	TEXT("Maximum number of vertices that a combined Mesh can have for us to collapse it into a single StaticMesh"));
+	TEXT("Maximum number of vertices that a combined Mesh can have for us to collapse it into a single StaticMesh")
+);
 
 // Can toggle on/off to compare performance with StaticMesh instead of GeometryCache
 static bool GUseGeometryCacheUSD = true;
 static FAutoConsoleVariableRef CVarUsdUseGeometryCache(
 	TEXT("USD.GeometryCache.Enable"),
 	GUseGeometryCacheUSD,
-	TEXT("Use GeometryCache instead of static meshes for loading animated meshes"));
+	TEXT("Use GeometryCache instead of static meshes for loading animated meshes")
+);
 
 static int32 GGeometryCacheMaxDepth = 15;
 static FAutoConsoleVariableRef CVarGeometryCacheMaxDepth(
 	TEXT("USD.GeometryCache.MaxDepth"),
 	GGeometryCacheMaxDepth,
-	TEXT("Maximum distance between an animated mesh prim to its collapsed geometry cache root"));
+	TEXT("Maximum distance between an animated mesh prim to its collapsed geometry cache root")
+);
 
 namespace UE::UsdInfoCache::Private
 {
@@ -62,11 +62,11 @@ namespace UE::UsdInfoCache::Private
 	enum class EGeometryCachePrimState : uint8
 	{
 		None = 0x00,
-		Uncollapsible = 0x01,		// prim cannot be collapsed as part of a geometry cache
-		Mesh = 0x02,				// prim is a mesh, animated or not
-		Xform = 0x04,				// prim is a xform, animated or not
-		Collapsible = Mesh | Xform,	// only meshes and xforms can be collapsed into a geometry cache
-		ValidRoot = 0x08			// prim can collapse itself and its children into a geometry cache
+		Uncollapsible = 0x01,		   // prim cannot be collapsed as part of a geometry cache
+		Mesh = 0x02,				   // prim is a mesh, animated or not
+		Xform = 0x04,				   // prim is a xform, animated or not
+		Collapsible = Mesh | Xform,	   // only meshes and xforms can be collapsed into a geometry cache
+		ValidRoot = 0x08			   // prim can collapse itself and its children into a geometry cache
 	};
 	ENUM_CLASS_FLAGS(EGeometryCachePrimState)
 
@@ -81,7 +81,7 @@ namespace UE::UsdInfoCache::Private
 	};
 }
 
-FArchive& operator <<( FArchive& Ar, UE::UsdInfoCache::Private::FUsdPrimInfo& Info )
+FArchive& operator<<(FArchive& Ar, UE::UsdInfoCache::Private::FUsdPrimInfo& Info)
 {
 	Ar << Info.AssetCollapsedRoot;
 	Ar << Info.ComponentCollapsedRoot;
@@ -96,7 +96,7 @@ FArchive& operator <<( FArchive& Ar, UE::UsdInfoCache::Private::FUsdPrimInfo& In
 struct FUsdInfoCache::FUsdInfoCacheImpl
 {
 	FUsdInfoCacheImpl()
-	: AllowedExtensionsForGeometryCacheSource(UnrealUSDWrapper::GetNativeFileFormats())
+		: AllowedExtensionsForGeometryCacheSource(UnrealUSDWrapper::GetNativeFileFormats())
 	{
 		AllowedExtensionsForGeometryCacheSource.Add(TEXT("abc"));
 	}
@@ -120,7 +120,7 @@ struct FUsdInfoCache::FUsdInfoCacheImpl
 	}
 
 	// Information we must have about all prims on the stage
-	TMap< UE::FSdfPath, UE::UsdInfoCache::Private::FUsdPrimInfo > InfoMap;
+	TMap<UE::FSdfPath, UE::UsdInfoCache::Private::FUsdPrimInfo> InfoMap;
 	mutable FRWLock InfoMapLock;
 
 	// Information we may have about a subset of prims
@@ -158,10 +158,7 @@ public:
 
 		for (const UE::FSdfPath& AuxPrimPath : AuxPrimPaths)
 		{
-			UE_LOG(LogUsd, Verbose, TEXT("Registering main prim '%s' and aux prim '%s'"),
-				*MainPrimPath.GetString(),
-				*AuxPrimPath.GetString()
-			);
+			UE_LOG(LogUsd, Verbose, TEXT("Registering main prim '%s' and aux prim '%s'"), *MainPrimPath.GetString(), *AuxPrimPath.GetString());
 
 			AuxToMainPrims.FindOrAdd(AuxPrimPath).Add(MainPrimPath);
 		}
@@ -182,9 +179,9 @@ FUsdInfoCache::~FUsdInfoCache()
 {
 }
 
-bool FUsdInfoCache::Serialize( FArchive& Ar )
+bool FUsdInfoCache::Serialize(FArchive& Ar)
 {
-	if ( FUsdInfoCacheImpl* ImplPtr = Impl.Get() )
+	if (FUsdInfoCacheImpl* ImplPtr = Impl.Get())
 	{
 		{
 			FWriteScopeLock ScopeLock(ImplPtr->InfoMapLock);
@@ -209,12 +206,12 @@ bool FUsdInfoCache::Serialize( FArchive& Ar )
 	return true;
 }
 
-bool FUsdInfoCache::ContainsInfoAboutPrim( const UE::FSdfPath& Path ) const
+bool FUsdInfoCache::ContainsInfoAboutPrim(const UE::FSdfPath& Path) const
 {
-	if ( FUsdInfoCacheImpl* ImplPtr = Impl.Get() )
+	if (FUsdInfoCacheImpl* ImplPtr = Impl.Get())
 	{
 		FReadScopeLock ScopeLock(ImplPtr->InfoMapLock);
-		return ImplPtr->InfoMap.Contains( Path );
+		return ImplPtr->InfoMap.Contains(Path);
 	}
 
 	return false;
@@ -235,40 +232,38 @@ TSet<UE::FSdfPath> FUsdInfoCache::GetKnownPrims() const
 	return Result;
 }
 
-bool FUsdInfoCache::IsPathCollapsed( const UE::FSdfPath& Path, ECollapsingType CollapsingType ) const
+bool FUsdInfoCache::IsPathCollapsed(const UE::FSdfPath& Path, ECollapsingType CollapsingType) const
 {
-	if ( FUsdInfoCacheImpl* ImplPtr = Impl.Get() )
+	if (FUsdInfoCacheImpl* ImplPtr = Impl.Get())
 	{
 		FReadScopeLock ScopeLock(ImplPtr->InfoMapLock);
 
-		if ( const UE::UsdInfoCache::Private::FUsdPrimInfo* FoundInfo = ImplPtr->InfoMap.Find( Path ) )
+		if (const UE::UsdInfoCache::Private::FUsdPrimInfo* FoundInfo = ImplPtr->InfoMap.Find(Path))
 		{
-			const UE::FSdfPath* CollapsedRoot = CollapsingType == ECollapsingType::Assets
-				? &FoundInfo->AssetCollapsedRoot
-				: &FoundInfo->ComponentCollapsedRoot;
+			const UE::FSdfPath* CollapsedRoot = CollapsingType == ECollapsingType::Assets ? &FoundInfo->AssetCollapsedRoot
+																						  : &FoundInfo->ComponentCollapsedRoot;
 
 			// A non-empty path to another prim means this prim is collapsed into that one
 			return !CollapsedRoot->IsEmpty() && *CollapsedRoot != Path;
 		}
 
 		// This should never happen: We should have cached the entire tree
-		ensureMsgf( false, TEXT( "Prim path '%s' has not been cached!" ), *Path.GetString() );
+		ensureMsgf(false, TEXT("Prim path '%s' has not been cached!"), *Path.GetString());
 	}
 
 	return false;
 }
 
-bool FUsdInfoCache::DoesPathCollapseChildren( const UE::FSdfPath& Path, ECollapsingType CollapsingType ) const
+bool FUsdInfoCache::DoesPathCollapseChildren(const UE::FSdfPath& Path, ECollapsingType CollapsingType) const
 {
-	if ( FUsdInfoCacheImpl* ImplPtr = Impl.Get() )
+	if (FUsdInfoCacheImpl* ImplPtr = Impl.Get())
 	{
 		FReadScopeLock ScopeLock(ImplPtr->InfoMapLock);
 
-		if ( const UE::UsdInfoCache::Private::FUsdPrimInfo* FoundInfo = ImplPtr->InfoMap.Find( Path ) )
+		if (const UE::UsdInfoCache::Private::FUsdPrimInfo* FoundInfo = ImplPtr->InfoMap.Find(Path))
 		{
-			const UE::FSdfPath* CollapsedRoot = CollapsingType == ECollapsingType::Assets
-				? &FoundInfo->AssetCollapsedRoot
-				: &FoundInfo->ComponentCollapsedRoot;
+			const UE::FSdfPath* CollapsedRoot = CollapsingType == ECollapsingType::Assets ? &FoundInfo->AssetCollapsedRoot
+																						  : &FoundInfo->ComponentCollapsedRoot;
 
 			// We store our own Path in there when we collapse children.
 			// Otherwise we hold the path of our collapse root, or empty (in case nothing is collapsed up to here)
@@ -276,26 +271,25 @@ bool FUsdInfoCache::DoesPathCollapseChildren( const UE::FSdfPath& Path, ECollaps
 		}
 
 		// This should never happen: We should have cached the entire tree
-		ensureMsgf( false, TEXT( "Prim path '%s' has not been cached!" ), *Path.GetString() );
+		ensureMsgf(false, TEXT("Prim path '%s' has not been cached!"), *Path.GetString());
 	}
 
 	return false;
 }
 
-UE::FSdfPath FUsdInfoCache::UnwindToNonCollapsedPath( const UE::FSdfPath& Path, ECollapsingType CollapsingType ) const
+UE::FSdfPath FUsdInfoCache::UnwindToNonCollapsedPath(const UE::FSdfPath& Path, ECollapsingType CollapsingType) const
 {
-	if ( FUsdInfoCacheImpl* ImplPtr = Impl.Get() )
+	if (FUsdInfoCacheImpl* ImplPtr = Impl.Get())
 	{
 		FReadScopeLock ScopeLock(ImplPtr->InfoMapLock);
 
-		if ( const UE::UsdInfoCache::Private::FUsdPrimInfo* FoundInfo = ImplPtr->InfoMap.Find( Path ) )
+		if (const UE::UsdInfoCache::Private::FUsdPrimInfo* FoundInfo = ImplPtr->InfoMap.Find(Path))
 		{
-			const UE::FSdfPath* CollapsedRoot = CollapsingType == ECollapsingType::Assets
-				? &FoundInfo->AssetCollapsedRoot
-				: &FoundInfo->ComponentCollapsedRoot;
+			const UE::FSdfPath* CollapsedRoot = CollapsingType == ECollapsingType::Assets ? &FoundInfo->AssetCollapsedRoot
+																						  : &FoundInfo->ComponentCollapsedRoot;
 
 			// An empty path here means that we are not collapsed at all
-			if ( CollapsedRoot->IsEmpty() )
+			if (CollapsedRoot->IsEmpty())
 			{
 				return Path;
 			}
@@ -307,7 +301,7 @@ UE::FSdfPath FUsdInfoCache::UnwindToNonCollapsedPath( const UE::FSdfPath& Path,
 		}
 
 		// This should never happen: We should have cached the entire tree
-		ensureMsgf( false, TEXT( "Prim path '%s' has not been cached!" ), *Path.GetString() );
+		ensureMsgf(false, TEXT("Prim path '%s' has not been cached!"), *Path.GetString());
 	}
 
 	return Path;
@@ -381,27 +375,27 @@ namespace UE::USDInfoCacheImpl::Private
 		const pxr::UsdPrim& UsdPrim,
 		const FUsdSchemaTranslationContext& Context,
 		const FUsdInfoCache::FUsdInfoCacheImpl& Impl,
-		const TMap< UE::FSdfPath, TArray<UsdUtils::FUsdPrimMaterialSlot> >& InSubtreeToMaterialSlots,
+		const TMap<UE::FSdfPath, TArray<UsdUtils::FUsdPrimMaterialSlot>>& InSubtreeToMaterialSlots,
 		uint64& OutVertexCount,
 		TArray<UsdUtils::FUsdPrimMaterialSlot>& OutMaterialSlots
 	)
 	{
-		TRACE_CPUPROFILER_EVENT_SCOPE( UE::USDInfoCacheImpl::Private::GetPrimVertexCountAndSlots );
+		TRACE_CPUPROFILER_EVENT_SCOPE(UE::USDInfoCacheImpl::Private::GetPrimVertexCountAndSlots);
 
 		if (UsdPrim.IsA<pxr::UsdGeomGprim>() || UsdPrim.IsA<pxr::UsdGeomSubset>())
 		{
 			OutVertexCount = UsdUtils::GetGprimVertexCount(pxr::UsdGeomGprim{UsdPrim}, Context.Time);
 
 			pxr::TfToken RenderContextToken = pxr::UsdShadeTokens->universalRenderContext;
-			if ( !Context.RenderContext.IsNone() )
+			if (!Context.RenderContext.IsNone())
 			{
-				RenderContextToken = UnrealToUsd::ConvertToken( *Context.RenderContext.ToString() ).Get();
+				RenderContextToken = UnrealToUsd::ConvertToken(*Context.RenderContext.ToString()).Get();
 			}
 
 			pxr::TfToken MaterialPurposeToken = pxr::UsdShadeTokens->allPurpose;
-			if ( !Context.MaterialPurpose.IsNone() )
+			if (!Context.MaterialPurpose.IsNone())
 			{
-				MaterialPurposeToken = UnrealToUsd::ConvertToken( *Context.MaterialPurpose.ToString() ).Get();
+				MaterialPurposeToken = UnrealToUsd::ConvertToken(*Context.MaterialPurpose.ToString()).Get();
 			}
 
 			const bool bProvideMaterialIndices = false;
@@ -413,23 +407,23 @@ namespace UE::USDInfoCacheImpl::Private
 				MaterialPurposeToken
 			);
 
-			OutMaterialSlots.Append( MoveTemp( LocalInfo.Slots ) );
+			OutMaterialSlots.Append(MoveTemp(LocalInfo.Slots));
 		}
-		else if ( pxr::UsdGeomPointInstancer PointInstancer{ UsdPrim } )
+		else if (pxr::UsdGeomPointInstancer PointInstancer{UsdPrim})
 		{
 			const pxr::UsdRelationship& Prototypes = PointInstancer.GetPrototypesRel();
 
 			pxr::SdfPathVector PrototypePaths;
-			if ( Prototypes.GetTargets( &PrototypePaths ) )
+			if (Prototypes.GetTargets(&PrototypePaths))
 			{
 				TArray<uint64> PrototypeVertexCounts;
-				PrototypeVertexCounts.SetNumZeroed( PrototypePaths.size() );
+				PrototypeVertexCounts.SetNumZeroed(PrototypePaths.size());
 
 				{
 					FReadScopeLock ScopeLock(Impl.InfoMapLock);
-					for ( uint32 PrototypeIndex = 0; PrototypeIndex < PrototypePaths.size(); ++PrototypeIndex )
+					for (uint32 PrototypeIndex = 0; PrototypeIndex < PrototypePaths.size(); ++PrototypeIndex)
 					{
-						const pxr::SdfPath& PrototypePath = PrototypePaths[ PrototypeIndex ];
+						const pxr::SdfPath& PrototypePath = PrototypePaths[PrototypeIndex];
 
 						// Skip invisible prototypes here to mirror how they're skipped within
 						// USDGeomMeshConversion.cpp, in the RecursivelyCollapseChildMeshes function. Those two
@@ -451,26 +445,27 @@ namespace UE::USDInfoCacheImpl::Private
 
 						// If we're calling this for a point instancer we should have parsed the results for our
 						// prototype subtrees already
-						if ( const UE::UsdInfoCache::Private::FUsdPrimInfo* FoundInfo = Impl.InfoMap.Find( UE::FSdfPath{ PrototypePath } ) )
+						if (const UE::UsdInfoCache::Private::FUsdPrimInfo* FoundInfo = Impl.InfoMap.Find(UE::FSdfPath{PrototypePath}))
 						{
-							PrototypeVertexCounts[ PrototypeIndex ] = FoundInfo->ExpectedVertexCountForSubtree;
+							PrototypeVertexCounts[PrototypeIndex] = FoundInfo->ExpectedVertexCountForSubtree;
 						}
 
-						if ( const TArray<UsdUtils::FUsdPrimMaterialSlot>* FoundPrototypeSlots = InSubtreeToMaterialSlots.Find( UE::FSdfPath{ PrototypePath } ) )
+						if (const TArray<UsdUtils::FUsdPrimMaterialSlot>* FoundPrototypeSlots = InSubtreeToMaterialSlots.Find(UE::FSdfPath{
+								PrototypePath}))
 						{
-							OutMaterialSlots.Append( *FoundPrototypeSlots );
+							OutMaterialSlots.Append(*FoundPrototypeSlots);
 						}
 					}
 				}
 
-				if ( pxr::UsdAttribute ProtoIndicesAttr = PointInstancer.GetProtoIndicesAttr() )
+				if (pxr::UsdAttribute ProtoIndicesAttr = PointInstancer.GetProtoIndicesAttr())
 				{
 					pxr::VtArray<int> ProtoIndicesArr;
-					if ( ProtoIndicesAttr.Get( &ProtoIndicesArr, pxr::UsdTimeCode::EarliestTime() ) )
+					if (ProtoIndicesAttr.Get(&ProtoIndicesArr, pxr::UsdTimeCode::EarliestTime()))
 					{
-						for ( int ProtoIndex : ProtoIndicesArr )
+						for (int ProtoIndex : ProtoIndicesArr)
 						{
-							OutVertexCount += PrototypeVertexCounts[ ProtoIndex ];
+							OutVertexCount += PrototypeVertexCounts[ProtoIndex];
 						}
 					}
 				}
@@ -484,14 +479,14 @@ namespace UE::USDInfoCacheImpl::Private
 		const pxr::TfToken& MaterialPurposeToken,
 		FUsdInfoCache::FUsdInfoCacheImpl& Impl,
 		FUsdSchemaTranslatorRegistry& Registry,
-		TMap< UE::FSdfPath, TArray<UsdUtils::FUsdPrimMaterialSlot>>& InOutSubtreeToMaterialSlots,
-		TArray< FString >& InOutPointInstancerPaths,
+		TMap<UE::FSdfPath, TArray<UsdUtils::FUsdPrimMaterialSlot>>& InOutSubtreeToMaterialSlots,
+		TArray<FString>& InOutPointInstancerPaths,
 		uint64& OutSubtreeVertexCount,
 		TArray<UsdUtils::FUsdPrimMaterialSlot>& OutSubtreeSlots,
 		bool bPossibleInheritedBindings
 	)
 	{
-		TRACE_CPUPROFILER_EVENT_SCOPE( UE::USDInfoCacheImpl::Private::RecursivePropagateVertexAndMaterialSlotCounts );
+		TRACE_CPUPROFILER_EVENT_SCOPE(UE::USDInfoCacheImpl::Private::RecursivePropagateVertexAndMaterialSlotCounts);
 
 		if (!UsdPrim)
 		{
@@ -525,9 +520,8 @@ namespace UE::USDInfoCacheImpl::Private
 					// manually, as it won't automatically "compute the fallback" for us like when we ComputeBoundMaterial()
 					std::unordered_set<pxr::TfToken, pxr::TfHash> MaterialPurposeTokens{
 						MaterialPurposeToken,
-						pxr::UsdShadeTokens->universalRenderContext
-					};
-					for(const pxr::TfToken& SomeMaterialPurposeToken : MaterialPurposeTokens)
+						pxr::UsdShadeTokens->universalRenderContext};
+					for (const pxr::TfToken& SomeMaterialPurposeToken : MaterialPurposeTokens)
 					{
 						// Each one of those relationships must have two targets: A collection, and a material
 						for (const pxr::UsdRelationship& Rel : BindingAPI.GetCollectionBindingRels(SomeMaterialPurposeToken))
@@ -557,7 +551,8 @@ namespace UE::USDInfoCacheImpl::Private
 								UE_LOG(
 									LogUsd,
 									Warning,
-									TEXT("Prim '%s' describes a collection-based material binding, but the relationship '%s' is invalid: It should contain exactly one Material path and one path to a collection relationship"),
+									TEXT("Prim '%s' describes a collection-based material binding, but the relationship '%s' is invalid: It should "
+										 "contain exactly one Material path and one path to a collection relationship"),
 									*PrimPath.GetString(),
 									*UsdToUnreal::ConvertToken(Rel.GetName())
 								);
@@ -616,7 +611,7 @@ namespace UE::USDInfoCacheImpl::Private
 			}
 
 			FWriteScopeLock ScopeLock(Impl.MaterialUsersLock);
-			for(const TPair<UE::FSdfPath, TSet<UE::FSdfPath>>& NewMaterialToUsers : NewMaterialUsers)
+			for (const TPair<UE::FSdfPath, TSet<UE::FSdfPath>>& NewMaterialToUsers : NewMaterialUsers)
 			{
 				TSet<UE::FSdfPath>& UserPrimPaths = Impl.MaterialUsers.FindOrAdd(NewMaterialToUsers.Key);
 				UserPrimPaths.Reserve(UserPrimPaths.Num() + NewMaterialToUsers.Value.Num());
@@ -661,37 +656,40 @@ namespace UE::USDInfoCacheImpl::Private
 			}
 		}
 
-		pxr::UsdPrimSiblingRange PrimChildren = UsdPrim.GetFilteredChildren( pxr::UsdTraverseInstanceProxies( pxr::UsdPrimAllPrimsPredicate ) );
+		pxr::UsdPrimSiblingRange PrimChildren = UsdPrim.GetFilteredChildren(pxr::UsdTraverseInstanceProxies(pxr::UsdPrimAllPrimsPredicate));
 
 		TArray<pxr::UsdPrim> Prims;
-		for ( pxr::UsdPrim Child : PrimChildren )
+		for (pxr::UsdPrim Child : PrimChildren)
 		{
-			Prims.Emplace( Child );
+			Prims.Emplace(Child);
 		}
 
 		const uint32 NumChildren = Prims.Num();
 
 		TArray<uint64> ChildSubtreeVertexCounts;
-		ChildSubtreeVertexCounts.SetNumUninitialized( NumChildren );
+		ChildSubtreeVertexCounts.SetNumUninitialized(NumChildren);
 
 		TArray<TArray<UsdUtils::FUsdPrimMaterialSlot>> ChildSubtreeMaterialSlots;
-		ChildSubtreeMaterialSlots.SetNum( NumChildren );
+		ChildSubtreeMaterialSlots.SetNum(NumChildren);
 
 		const int32 MinBatchSize = 1;
 
-		ParallelFor( TEXT( "RecursivePropagateVertexAndMaterialSlotCounts" ), Prims.Num(), MinBatchSize,
-			[&]( int32 Index )
+		ParallelFor(
+			TEXT("RecursivePropagateVertexAndMaterialSlotCounts"),
+			Prims.Num(),
+			MinBatchSize,
+			[&](int32 Index)
 			{
 				RecursivePropagateVertexAndMaterialSlotCounts(
-					Prims[ Index ],
+					Prims[Index],
 					Context,
 					MaterialPurposeToken,
 					Impl,
 					Registry,
 					InOutSubtreeToMaterialSlots,
 					InOutPointInstancerPaths,
-					ChildSubtreeVertexCounts[ Index ],
-					ChildSubtreeMaterialSlots[ Index ],
+					ChildSubtreeVertexCounts[Index],
+					ChildSubtreeMaterialSlots[Index],
 					bPrimHasInheritableMaterialBindings || bPossibleInheritedBindings
 				);
 			}
@@ -723,38 +721,31 @@ namespace UE::USDInfoCacheImpl::Private
 		}
 
 		bool bIsPointInstancer = false;
-		if ( pxr::UsdGeomPointInstancer PointInstancer{ UsdPrim } )
+		if (pxr::UsdGeomPointInstancer PointInstancer{UsdPrim})
 		{
 			bIsPointInstancer = true;
 		}
 		else if (!bPrimIsInvisible)
 		{
-			GetPrimVertexCountAndSlots(
-				UsdPrim,
-				Context,
-				Impl,
-				InOutSubtreeToMaterialSlots,
-				OutSubtreeVertexCount,
-				OutSubtreeSlots
-			);
+			GetPrimVertexCountAndSlots(UsdPrim, Context, Impl, InOutSubtreeToMaterialSlots, OutSubtreeVertexCount, OutSubtreeSlots);
 
-			for ( uint32 ChildIndex = 0; ChildIndex < NumChildren; ++ChildIndex )
+			for (uint32 ChildIndex = 0; ChildIndex < NumChildren; ++ChildIndex)
 			{
-				OutSubtreeVertexCount += ChildSubtreeVertexCounts[ ChildIndex ];
-				OutSubtreeSlots.Append( ChildSubtreeMaterialSlots[ ChildIndex ] );
+				OutSubtreeVertexCount += ChildSubtreeVertexCounts[ChildIndex];
+				OutSubtreeSlots.Append(ChildSubtreeMaterialSlots[ChildIndex]);
 			}
 		}
 
 		{
 			FWriteScopeLock ScopeLock(Impl.InfoMapLock);
 
-			UE::UsdInfoCache::Private::FUsdPrimInfo& Info = Impl.InfoMap.FindOrAdd( UE::FSdfPath{ UsdPrimPath } );
+			UE::UsdInfoCache::Private::FUsdPrimInfo& Info = Impl.InfoMap.FindOrAdd(UE::FSdfPath{UsdPrimPath});
 
 			// For point instancers we can't guarantee we parsed the prototypes yet because they
 			// could technically be anywhere, so store them here for a later pass
-			if ( bIsPointInstancer )
+			if (bIsPointInstancer)
 			{
-				InOutPointInstancerPaths.Emplace( UE::FSdfPath{ UsdPrimPath }.GetString() );
+				InOutPointInstancerPaths.Emplace(UE::FSdfPath{UsdPrimPath}.GetString());
 			}
 			// While we will compute the totals for any and all children normally, don't just append the regular
 			// traversal vertex count to the point instancer prim itself just yet, as that doesn't really represent
@@ -763,7 +754,7 @@ namespace UE::USDInfoCacheImpl::Private
 			else
 			{
 				Info.ExpectedVertexCountForSubtree = OutSubtreeVertexCount;
-				InOutSubtreeToMaterialSlots.Emplace( UsdPrimPath, OutSubtreeSlots );
+				InOutSubtreeToMaterialSlots.Emplace(UsdPrimPath, OutSubtreeSlots);
 			}
 		}
 	}
@@ -781,36 +772,36 @@ namespace UE::USDInfoCacheImpl::Private
 		const pxr::UsdStageRefPtr& Stage,
 		const FUsdSchemaTranslationContext& Context,
 		FUsdInfoCache::FUsdInfoCacheImpl& Impl,
-		TArray< FString >& PointInstancerPaths,
-		TMap< UE::FSdfPath, TArray<UsdUtils::FUsdPrimMaterialSlot>>& InOutSubtreeMaterialSlots
+		TArray<FString>& PointInstancerPaths,
+		TMap<UE::FSdfPath, TArray<UsdUtils::FUsdPrimMaterialSlot>>& InOutSubtreeMaterialSlots
 	)
 	{
 		// We must sort point instancers in a particular order in case they depend on each other.
 		// At least we know that an ordering like this should be possible, because A with B as a prototype and B with A
 		// as a prototype leads to an invalid USD stage.
-		TFunction<bool(const FString&, const FString&)> SortFunction = [Stage]( const FString& LHS, const FString& RHS )
+		TFunction<bool(const FString&, const FString&)> SortFunction = [Stage](const FString& LHS, const FString& RHS)
 		{
 			FScopedUsdAllocs Allocs;
 
-			pxr::SdfPath LPath = UnrealToUsd::ConvertPath( *LHS ).Get();
-			pxr::SdfPath RPath = UnrealToUsd::ConvertPath( *RHS ).Get();
+			pxr::SdfPath LPath = UnrealToUsd::ConvertPath(*LHS).Get();
+			pxr::SdfPath RPath = UnrealToUsd::ConvertPath(*RHS).Get();
 
-			pxr::UsdGeomPointInstancer LPointInstancer = pxr::UsdGeomPointInstancer{ Stage->GetPrimAtPath( LPath ) };
-			pxr::UsdGeomPointInstancer RPointInstancer = pxr::UsdGeomPointInstancer{ Stage->GetPrimAtPath( RPath ) };
-			if ( LPointInstancer && RPointInstancer )
+			pxr::UsdGeomPointInstancer LPointInstancer = pxr::UsdGeomPointInstancer{Stage->GetPrimAtPath(LPath)};
+			pxr::UsdGeomPointInstancer RPointInstancer = pxr::UsdGeomPointInstancer{Stage->GetPrimAtPath(RPath)};
+			if (LPointInstancer && RPointInstancer)
 			{
 				const pxr::UsdRelationship& LPrototypes = LPointInstancer.GetPrototypesRel();
 				pxr::SdfPathVector LPrototypePaths;
-				if ( LPrototypes.GetTargets( &LPrototypePaths ) )
+				if (LPrototypes.GetTargets(&LPrototypePaths))
 				{
-					for ( const pxr::SdfPath& LPrototypePath : LPrototypePaths )
+					for (const pxr::SdfPath& LPrototypePath : LPrototypePaths)
 					{
 						// Consider RPointInstancer at RPath "/LPointInstancer/Prototypes/Nest/RPointInstancer", and
 						// LPointInstancer has prototype "/LPointInstancer/Prototypes/Nest". If RPath has the LPrototypePath as prefix,
 						// we should have R come before L in the sort order.
 						// Of course, in this scenario we could get away with just sorting by length, but that wouldn't help if the
 						// point instancers were not inside each other (e.g. siblings).
-						if ( RPath.HasPrefix( LPrototypePath ) )
+						if (RPath.HasPrefix(LPrototypePath))
 						{
 							return false;
 						}
@@ -826,11 +817,11 @@ namespace UE::USDInfoCacheImpl::Private
 		};
 		PointInstancerPaths.Sort(SortFunction);
 
-		for ( const FString& PointInstancerPath : PointInstancerPaths )
+		for (const FString& PointInstancerPath : PointInstancerPaths)
 		{
-			UE::FSdfPath UsdPointInstancerPath{ *PointInstancerPath };
+			UE::FSdfPath UsdPointInstancerPath{*PointInstancerPath};
 
-			if ( pxr::UsdPrim PointInstancer = Stage->GetPrimAtPath( UnrealToUsd::ConvertPath( *PointInstancerPath ).Get() ) )
+			if (pxr::UsdPrim PointInstancer = Stage->GetPrimAtPath(UnrealToUsd::ConvertPath(*PointInstancerPath).Get()))
 			{
 				uint64 PointInstancerVertexCount = 0;
 				TArray<UsdUtils::FUsdPrimMaterialSlot> PointInstancerMaterialSlots;
@@ -845,39 +836,39 @@ namespace UE::USDInfoCacheImpl::Private
 				);
 
 				FWriteScopeLock Lock{Impl.InfoMapLock};
-				UE::UsdInfoCache::Private::FUsdPrimInfo& Info = Impl.InfoMap.FindOrAdd( UsdPointInstancerPath );
+				UE::UsdInfoCache::Private::FUsdPrimInfo& Info = Impl.InfoMap.FindOrAdd(UsdPointInstancerPath);
 
 				Info.ExpectedVertexCountForSubtree = PointInstancerVertexCount;
-				InOutSubtreeMaterialSlots.Emplace( UsdPointInstancerPath, PointInstancerMaterialSlots );
+				InOutSubtreeMaterialSlots.Emplace(UsdPointInstancerPath, PointInstancerMaterialSlots);
 
 				// Now that we have info on the point instancer itself, update the counts of all ancestors.
 				// Note: The vertex/material slot count for the entire point instancer subtree are just the counts
 				// for the point instancer itself, as we stop regular traversal when we hit them
 				UE::FSdfPath ParentPath = UsdPointInstancerPath.GetParentPath();
-				pxr::UsdPrim Prim = Stage->GetPrimAtPath( ParentPath );
-				while ( Prim )
+				pxr::UsdPrim Prim = Stage->GetPrimAtPath(ParentPath);
+				while (Prim)
 				{
 					// If our ancestor is a point instancer itself, just abort as we'll only get the actual counts
 					// when we handle that ancestor directly. We don't want to update the ancestor point instancer's
 					// ancestors with incorrect values
-					if ( Prim.IsA<pxr::UsdGeomPointInstancer>() )
+					if (Prim.IsA<pxr::UsdGeomPointInstancer>())
 					{
 						break;
 					}
 
-					UE::UsdInfoCache::Private::FUsdPrimInfo& ParentInfo = Impl.InfoMap.FindOrAdd( ParentPath );
+					UE::UsdInfoCache::Private::FUsdPrimInfo& ParentInfo = Impl.InfoMap.FindOrAdd(ParentPath);
 					ParentInfo.ExpectedVertexCountForSubtree += PointInstancerVertexCount;
 
-					InOutSubtreeMaterialSlots[ ParentPath ].Append( PointInstancerMaterialSlots );
+					InOutSubtreeMaterialSlots[ParentPath].Append(PointInstancerMaterialSlots);
 
 					// Break only here so we update the pseudoroot too
-					if ( Prim.IsPseudoRoot() )
+					if (Prim.IsPseudoRoot())
 					{
 						break;
 					}
 
 					ParentPath = ParentPath.GetParentPath();
-					Prim = Stage->GetPrimAtPath( ParentPath );
+					Prim = Stage->GetPrimAtPath(ParentPath);
 				}
 			}
 		}
@@ -893,7 +884,7 @@ namespace UE::USDInfoCacheImpl::Private
 	 */
 	void CollectMaterialSlotCounts(
 		FUsdInfoCache::FUsdInfoCacheImpl& Impl,
-		const TMap< UE::FSdfPath, TArray<UsdUtils::FUsdPrimMaterialSlot>>& SubtreeMaterialSlots,
+		const TMap<UE::FSdfPath, TArray<UsdUtils::FUsdPrimMaterialSlot>>& SubtreeMaterialSlots,
 		bool bContextMergeIdenticalSlots
 	)
 	{
@@ -911,7 +902,8 @@ namespace UE::USDInfoCacheImpl::Private
 				// We only merge slots in the context of collapsing
 				const bool bPrimIsCollapsedOrCollapseRoot = !CollapsedRoot->IsEmpty() || PrimPath.IsAbsoluteRootPath();
 
-				const bool bPrimIsPotentialGeometryCacheRoot = FoundInfo->GeometryCacheState == UE::UsdInfoCache::Private::EGeometryCachePrimState::ValidRoot;
+				const bool bPrimIsPotentialGeometryCacheRoot = FoundInfo->GeometryCacheState
+															   == UE::UsdInfoCache::Private::EGeometryCachePrimState::ValidRoot;
 				bCanMergeSlotsForThisPrim = bPrimIsCollapsedOrCollapseRoot && !bPrimIsPotentialGeometryCacheRoot;
 			}
 
@@ -933,7 +925,7 @@ namespace UE::USDInfoCacheImpl::Private
 		const pxr::UsdPrim& UsdPrim,
 		FUsdSchemaTranslationContext& Context,
 		FUsdInfoCache::FUsdInfoCacheImpl& Impl,
-		const TSharedPtr< FUsdSchemaTranslator >& Translator
+		const TSharedPtr<FUsdSchemaTranslator>& Translator
 	)
 	{
 		if (!UsdPrim)
@@ -982,7 +974,7 @@ namespace UE::USDInfoCacheImpl::Private
 		bool bIsAssetCollapsed = !AssetCollapsedRoot.IsEmpty();
 		bool bIsComponentCollapsed = !ComponentCollapsedRoot.IsEmpty();
 
-		TSharedPtr< FUsdSchemaTranslator > SchemaTranslator = Registry.CreateTranslatorForSchema(Context.AsShared(), UE::FUsdTyped(UsdPrim));
+		TSharedPtr<FUsdSchemaTranslator> SchemaTranslator = Registry.CreateTranslatorForSchema(Context.AsShared(), UE::FUsdTyped(UsdPrim));
 		if (SchemaTranslator)
 		{
 			if (!bIsAssetCollapsed || !bIsComponentCollapsed)
@@ -990,11 +982,12 @@ namespace UE::USDInfoCacheImpl::Private
 				const bool bCanMeshSubtreeBeCollapsed = CanMeshSubtreeBeCollapsed(UsdPrim, Context, Impl, SchemaTranslator);
 
 				const UE::UsdInfoCache::Private::FUsdPrimInfo& Info = Impl.InfoMap.FindChecked(UE::FSdfPath{UsdPrimPath});
-				const bool bIsPotentialGeometryCacheRoot = Info.GeometryCacheState ==  UE::UsdInfoCache::Private::EGeometryCachePrimState::ValidRoot;
+				const bool bIsPotentialGeometryCacheRoot = Info.GeometryCacheState == UE::UsdInfoCache::Private::EGeometryCachePrimState::ValidRoot;
 
 				if (!bIsAssetCollapsed)
 				{
-					// The potential geometry cache root is checked first since the FUsdGeometryCacheTranslator::CollapsesChildren has no logic of its own
+					// The potential geometry cache root is checked first since the FUsdGeometryCacheTranslator::CollapsesChildren has no logic of its
+					// own
 					if (bIsPotentialGeometryCacheRoot || (SchemaTranslator->CollapsesChildren(ECollapsingType::Assets) && bCanMeshSubtreeBeCollapsed))
 					{
 						AssetCollapsedRootOverride = &UsdPrimPath;
@@ -1003,7 +996,8 @@ namespace UE::USDInfoCacheImpl::Private
 
 				if (!bIsComponentCollapsed)
 				{
-					if (bIsPotentialGeometryCacheRoot || (SchemaTranslator->CollapsesChildren(ECollapsingType::Components) && bCanMeshSubtreeBeCollapsed))
+					if (bIsPotentialGeometryCacheRoot
+						|| (SchemaTranslator->CollapsesChildren(ECollapsingType::Components) && bCanMeshSubtreeBeCollapsed))
 					{
 						ComponentCollapsedRootOverride = &UsdPrimPath;
 					}
@@ -1021,17 +1015,13 @@ namespace UE::USDInfoCacheImpl::Private
 
 		const int32 MinBatchSize = 1;
 
-		ParallelFor(TEXT("RecursiveQueryCollapsesChildren"), Prims.Num(), MinBatchSize,
+		ParallelFor(
+			TEXT("RecursiveQueryCollapsesChildren"),
+			Prims.Num(),
+			MinBatchSize,
 			[&](int32 Index)
 			{
-				RecursiveQueryCollapsesChildren(
-					Prims[Index],
-					Context,
-					Impl,
-					Registry,
-					*AssetCollapsedRootOverride,
-					*ComponentCollapsedRootOverride
-				);
+				RecursiveQueryCollapsesChildren(Prims[Index], Context, Impl, Registry, *AssetCollapsedRootOverride, *ComponentCollapsedRootOverride);
 			}
 		);
 
@@ -1105,21 +1095,17 @@ namespace UE::USDInfoCacheImpl::Private
 		return Result;
 	}
 
-	void RegisterInstanceableAuxPrims(
-		const pxr::UsdPrim& UsdPrim,
-		FUsdSchemaTranslationContext& Context,
-		FUsdInfoCache::FUsdInfoCacheImpl& Impl
-	)
+	void RegisterInstanceableAuxPrims(const pxr::UsdPrim& UsdPrim, FUsdSchemaTranslationContext& Context, FUsdInfoCache::FUsdInfoCacheImpl& Impl)
 	{
 		TRACE_CPUPROFILER_EVENT_SCOPE(UE::USDInfoCacheImpl::Private::RegisterInstanceableAuxPrims);
 		FScopedUsdAllocs Allocs;
 
 		pxr::UsdStageRefPtr Stage = UsdPrim.GetStage();
-		for(const pxr::UsdPrim& Prototype : Stage->GetPrototypes())
+		for (const pxr::UsdPrim& Prototype : Stage->GetPrototypes())
 		{
 			TRACE_CPUPROFILER_EVENT_SCOPE(RegisterInstanceableAuxPrims::Prototype);
 
-			if(!Prototype)
+			if (!Prototype)
 			{
 				continue;
 			}
@@ -1168,7 +1154,12 @@ namespace UE::USDInfoCacheImpl::Private
 		}
 	}
 
-	void FindValidGeometryCacheRoot(const pxr::UsdPrim& UsdPrim, FUsdSchemaTranslationContext& Context, FUsdInfoCache::FUsdInfoCacheImpl& Impl, UE::UsdInfoCache::Private::EGeometryCachePrimState& OutState)
+	void FindValidGeometryCacheRoot(
+		const pxr::UsdPrim& UsdPrim,
+		FUsdSchemaTranslationContext& Context,
+		FUsdInfoCache::FUsdInfoCacheImpl& Impl,
+		UE::UsdInfoCache::Private::EGeometryCachePrimState& OutState
+	)
 	{
 		TRACE_CPUPROFILER_EVENT_SCOPE(FindValidGeometryCacheRoot);
 
@@ -1235,8 +1226,8 @@ namespace UE::USDInfoCacheImpl::Private
 
 		FScopedUsdAllocs Allocs;
 
-		// With this recursive check for geometry cache, we want to find branches with an animated mesh at the leaf and find the root where they can meet.
-		// This root prim will collapses the static and animated meshes under it into a single geometry cache.
+		// With this recursive check for geometry cache, we want to find branches with an animated mesh at the leaf and find the root where they can
+		// meet. This root prim will collapses the static and animated meshes under it into a single geometry cache.
 
 		pxr::SdfPath UsdPrimPath = UsdPrim.GetPrimPath();
 		UE::FSdfPath PrimPath{UsdPrimPath};
@@ -1343,7 +1334,8 @@ namespace UE::USDInfoCacheImpl::Private
 		else if (bIsXform)
 		{
 			// An xform prim is considered collapsible since it could have a mesh prim under it. It has to bubble up its children state.
-			PrimState = ChildrenState != EGeometryCachePrimState::None ? ChildrenState | EGeometryCachePrimState::Xform : EGeometryCachePrimState::Xform;
+			PrimState = ChildrenState != EGeometryCachePrimState::None ? ChildrenState | EGeometryCachePrimState::Xform
+																	   : EGeometryCachePrimState::Xform;
 		}
 		else
 		{
@@ -1416,9 +1408,15 @@ namespace UE::USDInfoCacheImpl::Private
 		{
 			if (Depth > GGeometryCacheMaxDepth)
 			{
-				UE_LOG(LogUsd, Warning, TEXT("Prim '%s' is potentially a geometry cache %d levels deep, which exceeds the limit of %d. "
-					"This could affect its imported animation. The limit can be increased with the cvar USD.GeometryCache.MaxDepth if needed."),
-					*PrimPath.GetString(), Depth, GGeometryCacheMaxDepth);
+				UE_LOG(
+					LogUsd,
+					Warning,
+					TEXT("Prim '%s' is potentially a geometry cache %d levels deep, which exceeds the limit of %d. "
+						 "This could affect its imported animation. The limit can be increased with the cvar USD.GeometryCache.MaxDepth if needed."),
+					*PrimPath.GetString(),
+					Depth,
+					GGeometryCacheMaxDepth
+				);
 			}
 			FindValidGeometryCacheRoot(UsdPrim, Context, Impl, PrimState);
 			Depth = -1;
@@ -1471,9 +1469,14 @@ namespace UE::USDInfoCacheImpl::Private
 		{
 			if (Depth > GGeometryCacheMaxDepth)
 			{
-				UE_LOG(LogUsd, Warning, TEXT("The stage has a geometry cache %d levels deep, which exceeds the limit of %d. "
-					"This could affect its imported animation. The limit can be increased with the cvar USD.GeometryCache.MaxDepth if needed."),
-					Depth, GGeometryCacheMaxDepth);
+				UE_LOG(
+					LogUsd,
+					Warning,
+					TEXT("The stage has a geometry cache %d levels deep, which exceeds the limit of %d. "
+						 "This could affect its imported animation. The limit can be increased with the cvar USD.GeometryCache.MaxDepth if needed."),
+					Depth,
+					GGeometryCacheMaxDepth
+				);
 			}
 
 			FScopedUsdAllocs UsdAllocs;
@@ -1488,8 +1491,8 @@ namespace UE::USDInfoCacheImpl::Private
 			}
 		}
 	}
-#endif // USE_SD_SDK
-}
+#endif	  // USE_SD_SDK
+}	 // namespace UE::USDInfoCacheImpl::Private
 
 bool FUsdInfoCache::IsPotentialGeometryCacheRoot(const UE::FSdfPath& Path) const
 {
@@ -1509,37 +1512,37 @@ bool FUsdInfoCache::IsPotentialGeometryCacheRoot(const UE::FSdfPath& Path) const
 	return false;
 }
 
-TOptional<uint64> FUsdInfoCache::GetSubtreeVertexCount( const UE::FSdfPath& Path )
+TOptional<uint64> FUsdInfoCache::GetSubtreeVertexCount(const UE::FSdfPath& Path)
 {
-	if ( FUsdInfoCacheImpl* ImplPtr = Impl.Get() )
+	if (FUsdInfoCacheImpl* ImplPtr = Impl.Get())
 	{
 		FReadScopeLock ScopeLock(ImplPtr->InfoMapLock);
 
-		if ( const UE::UsdInfoCache::Private::FUsdPrimInfo* FoundInfo = ImplPtr->InfoMap.Find( Path ) )
+		if (const UE::UsdInfoCache::Private::FUsdPrimInfo* FoundInfo = ImplPtr->InfoMap.Find(Path))
 		{
 			return FoundInfo->ExpectedVertexCountForSubtree;
 		}
 
 		// This should never happen: We should have cached the entire tree
-		ensureMsgf( false, TEXT( "Prim path '%s' has not been cached!" ), *Path.GetString() );
+		ensureMsgf(false, TEXT("Prim path '%s' has not been cached!"), *Path.GetString());
 	}
 
 	return {};
 }
 
-TOptional<uint64> FUsdInfoCache::GetSubtreeMaterialSlotCount( const UE::FSdfPath& Path )
+TOptional<uint64> FUsdInfoCache::GetSubtreeMaterialSlotCount(const UE::FSdfPath& Path)
 {
-	if ( FUsdInfoCacheImpl* ImplPtr = Impl.Get() )
+	if (FUsdInfoCacheImpl* ImplPtr = Impl.Get())
 	{
 		FReadScopeLock ScopeLock(ImplPtr->InfoMapLock);
 
-		if ( const UE::UsdInfoCache::Private::FUsdPrimInfo* FoundInfo = ImplPtr->InfoMap.Find( Path ) )
+		if (const UE::UsdInfoCache::Private::FUsdPrimInfo* FoundInfo = ImplPtr->InfoMap.Find(Path))
 		{
 			return FoundInfo->SubtreeMaterialSlots.Num();
 		}
 
 		// This should never happen: We should have cached the entire tree
-		ensureMsgf( false, TEXT( "Prim path '%s' has not been cached!" ), *Path.GetString() );
+		ensureMsgf(false, TEXT("Prim path '%s' has not been cached!"), *Path.GetString());
 	}
 
 	return {};
@@ -1572,10 +1575,7 @@ void FUsdInfoCache::LinkAssetToPrim(const UE::FSdfPath& Path, UObject* Asset)
 	}
 	FWriteScopeLock ScopeLock(ImplPtr->PrimPathToAssetsLock);
 
-	UE_LOG(LogUsd, Verbose, TEXT("Linking asset '%s' to prim '%s'"),
-		*Asset->GetPathName(),
-		*Path.GetString()
-	);
+	UE_LOG(LogUsd, Verbose, TEXT("Linking asset '%s' to prim '%s'"), *Asset->GetPathName(), *Path.GetString());
 
 	ImplPtr->PrimPathToAssets.FindOrAdd(Path).AddUnique(Asset);
 	ImplPtr->AssetToPrimPaths.FindOrAdd(Asset).AddUnique(Path);
@@ -1590,9 +1590,7 @@ TArray<TWeakObjectPtr<UObject>> FUsdInfoCache::RemoveAllAssetPrimLinks(const UE:
 	}
 	FWriteScopeLock ScopeLock(ImplPtr->PrimPathToAssetsLock);
 
-	UE_LOG(LogUsd, Verbose, TEXT("Removing asset prim links for path '%s'"),
-		*Path.GetString()
-	);
+	UE_LOG(LogUsd, Verbose, TEXT("Removing asset prim links for path '%s'"), *Path.GetString());
 
 	TArray<TWeakObjectPtr<UObject>> Assets;
 	ImplPtr->PrimPathToAssets.RemoveAndCopyValue(Path, Assets);
@@ -1673,13 +1671,13 @@ TMap<UE::FSdfPath, TArray<TWeakObjectPtr<UObject>>> FUsdInfoCache::GetAllAssetPr
 	return ImplPtr->PrimPathToAssets;
 }
 
-void FUsdInfoCache::RebuildCacheForSubtree( const UE::FUsdPrim& Prim, FUsdSchemaTranslationContext& Context )
+void FUsdInfoCache::RebuildCacheForSubtree(const UE::FUsdPrim& Prim, FUsdSchemaTranslationContext& Context)
 {
 #if USE_USD_SDK
-	TRACE_CPUPROFILER_EVENT_SCOPE( FUsdInfoCache::RebuildCacheForSubtree );
+	TRACE_CPUPROFILER_EVENT_SCOPE(FUsdInfoCache::RebuildCacheForSubtree);
 
 	FUsdInfoCacheImpl* ImplPtr = Impl.Get();
-	if ( !ImplPtr )
+	if (!ImplPtr)
 	{
 		return;
 	}
@@ -1692,7 +1690,7 @@ void FUsdInfoCache::RebuildCacheForSubtree( const UE::FUsdPrim& Prim, FUsdSchema
 		FScopedUsdAllocs Allocs;
 
 		pxr::UsdPrim UsdPrim{Prim};
-		if ( !UsdPrim )
+		if (!UsdPrim)
 		{
 			return;
 		}
@@ -1713,11 +1711,11 @@ void FUsdInfoCache::RebuildCacheForSubtree( const UE::FUsdPrim& Prim, FUsdSchema
 			ImplPtr->MainToAuxPrims.Empty();
 		}
 
-		IUsdSchemasModule& UsdSchemasModule = FModuleManager::Get().LoadModuleChecked< IUsdSchemasModule >( TEXT( "USDSchemas" ) );
+		IUsdSchemasModule& UsdSchemasModule = FModuleManager::Get().LoadModuleChecked<IUsdSchemasModule>(TEXT("USDSchemas"));
 		FUsdSchemaTranslatorRegistry& Registry = UsdSchemasModule.GetTranslatorRegistry();
 
-		TMap< UE::FSdfPath, TArray<UsdUtils::FUsdPrimMaterialSlot>> TempSubtreeSlots;
-		TArray< FString > PointInstancerPaths;
+		TMap<UE::FSdfPath, TArray<UsdUtils::FUsdPrimMaterialSlot>> TempSubtreeSlots;
+		TArray<FString> PointInstancerPaths;
 
 		pxr::TfToken MaterialPurposeToken = pxr::UsdShadeTokens->allPurpose;
 		if (!Context.MaterialPurpose.IsNone())
@@ -1743,37 +1741,22 @@ void FUsdInfoCache::RebuildCacheForSubtree( const UE::FUsdPrim& Prim, FUsdSchema
 			bPossibleInheritedBindings
 		);
 
-		UE::USDInfoCacheImpl::Private::UpdateInfoForPointInstancers(
-			UsdPrim.GetStage(),
-			Context,
-			*ImplPtr,
-			PointInstancerPaths,
-			TempSubtreeSlots
-		);
+		UE::USDInfoCacheImpl::Private::UpdateInfoForPointInstancers(UsdPrim.GetStage(), Context, *ImplPtr, PointInstancerPaths, TempSubtreeSlots);
 
 		UE::USDInfoCacheImpl::Private::CheckForGeometryCache(UsdPrim, Context, *ImplPtr);
 
-		UE::USDInfoCacheImpl::Private::RecursiveQueryCollapsesChildren(
-			UsdPrim,
-			Context,
-			*ImplPtr,
-			Registry
-		);
+		UE::USDInfoCacheImpl::Private::RecursiveQueryCollapsesChildren(UsdPrim, Context, *ImplPtr, Registry);
 
 		UE::USDInfoCacheImpl::Private::RegisterInstanceableAuxPrims(UsdPrim, Context, *ImplPtr);
 
-		UE::USDInfoCacheImpl::Private::CollectMaterialSlotCounts(
-			*ImplPtr,
-			TempSubtreeSlots,
-			Context.bMergeIdenticalMaterialSlots
-		);
+		UE::USDInfoCacheImpl::Private::CollectMaterialSlotCounts(*ImplPtr, TempSubtreeSlots, Context.bMergeIdenticalMaterialSlots);
 	}
-#endif // USE_USD_SDK
+#endif	  // USE_USD_SDK
 }
 
 void FUsdInfoCache::Clear()
 {
-	if ( FUsdInfoCacheImpl* ImplPtr = Impl.Get() )
+	if (FUsdInfoCacheImpl* ImplPtr = Impl.Get())
 	{
 		{
 			FWriteScopeLock ScopeLock(ImplPtr->InfoMapLock);
@@ -1798,7 +1781,7 @@ void FUsdInfoCache::Clear()
 
 bool FUsdInfoCache::IsEmpty()
 {
-	if ( FUsdInfoCacheImpl* ImplPtr = Impl.Get() )
+	if (FUsdInfoCacheImpl* ImplPtr = Impl.Get())
 	{
 		FReadScopeLock ScopeLock(ImplPtr->InfoMapLock);
 		return ImplPtr->InfoMap.IsEmpty();
diff --git a/Engine/Plugins/Importers/USDImporter/Source/USDSchemas/Private/USDLuxLightTranslator.cpp b/Engine/Plugins/Importers/USDImporter/Source/USDSchemas/Private/USDLuxLightTranslator.cpp
--- a/Engine/Plugins/Importers/USDImporter/Source/USDSchemas/Private/USDLuxLightTranslator.cpp
+++ b/Engine/Plugins/Importers/USDImporter/Source/USDSchemas/Private/USDLuxLightTranslator.cpp
@@ -2,32 +2,24 @@
 
 #include "USDLuxLightTranslator.h"
 
-#include "USDDrawModeComponent.h"
 #include "USDConversionUtils.h"
+#include "USDDrawModeComponent.h"
 #include "USDLightConversion.h"
-#include "USDTypesConversion.h"
 
 #include "UsdWrappers/UsdPrim.h"
-#include "UsdWrappers/UsdStage.h"
 
 #include "Components/DirectionalLightComponent.h"
-#include "Components/LightComponent.h"
 #include "Components/PointLightComponent.h"
 #include "Components/RectLightComponent.h"
 #include "Components/SkyLightComponent.h"
 #include "Components/SpotLightComponent.h"
-#include "Engine/TextureCube.h"
 
 #if USE_USD_SDK
 
 #include "USDIncludesStart.h"
-	#include "pxr/usd/usdLux/diskLight.h"
-	#include "pxr/usd/usdLux/distantLight.h"
-	#include "pxr/usd/usdLux/domeLight.h"
-	#include "pxr/usd/usdLux/lightAPI.h"
-	#include "pxr/usd/usdLux/rectLight.h"
-	#include "pxr/usd/usdLux/shapingAPI.h"
-	#include "pxr/usd/usdLux/sphereLight.h"
+#include "pxr/usd/usdLux/diskLight.h"
+#include "pxr/usd/usdLux/lightAPI.h"
+#include "pxr/usd/usdLux/rectLight.h"
 #include "USDIncludesEnd.h"
 
 USceneComponent* FUsdLuxLightTranslator::CreateComponents()
@@ -49,13 +41,13 @@ USceneComponent* FUsdLuxLightTranslator::CreateComponents()
 	return SceneComponent;
 }
 
-void FUsdLuxLightTranslator::UpdateComponents( USceneComponent* SceneComponent )
+void FUsdLuxLightTranslator::UpdateComponents(USceneComponent* SceneComponent)
 {
-	Super::UpdateComponents( SceneComponent );
+	Super::UpdateComponents(SceneComponent);
 
-	ULightComponentBase* LightComponent = Cast< ULightComponentBase >( SceneComponent );
+	ULightComponentBase* LightComponent = Cast<ULightComponentBase>(SceneComponent);
 
-	if ( !LightComponent )
+	if (!LightComponent)
 	{
 		return;
 	}
@@ -64,54 +56,54 @@ void FUsdLuxLightTranslator::UpdateComponents( USceneComponent* SceneComponent )
 
 	pxr::UsdPrim Prim = GetPrim();
 
-	const pxr::UsdLuxLightAPI LightAPI( Prim );
+	const pxr::UsdLuxLightAPI LightAPI(Prim);
 
-	if ( !LightAPI )
+	if (!LightAPI)
 	{
 		return;
 	}
 
 	LightComponent->UnregisterComponent();
 
-	UsdToUnreal::ConvertLight( Prim, *LightComponent, Context->Time );
+	UsdToUnreal::ConvertLight(Prim, *LightComponent, Context->Time);
 
-	if ( UDirectionalLightComponent* DirectionalLightComponent = Cast< UDirectionalLightComponent >( SceneComponent ) )
+	if (UDirectionalLightComponent* DirectionalLightComponent = Cast<UDirectionalLightComponent>(SceneComponent))
 	{
-		UsdToUnreal::ConvertDistantLight( Prim, *DirectionalLightComponent, Context->Time );
+		UsdToUnreal::ConvertDistantLight(Prim, *DirectionalLightComponent, Context->Time);
 	}
-	else if ( URectLightComponent* RectLightComponent = Cast< URectLightComponent >( SceneComponent ) )
+	else if (URectLightComponent* RectLightComponent = Cast<URectLightComponent>(SceneComponent))
 	{
-		if ( pxr::UsdLuxRectLight UsdRectLight{ Prim } )
+		if (pxr::UsdLuxRectLight UsdRectLight{Prim})
 		{
-			UsdToUnreal::ConvertRectLight( Prim, *RectLightComponent, Context->Time );
+			UsdToUnreal::ConvertRectLight(Prim, *RectLightComponent, Context->Time);
 		}
-		else if ( pxr::UsdLuxDiskLight UsdDiskLight{ Prim } )
+		else if (pxr::UsdLuxDiskLight UsdDiskLight{Prim})
 		{
-			UsdToUnreal::ConvertDiskLight( Prim, *RectLightComponent, Context->Time );
+			UsdToUnreal::ConvertDiskLight(Prim, *RectLightComponent, Context->Time);
 		}
 	}
-	else if ( UPointLightComponent* PointLightComponent = Cast< UPointLightComponent >( SceneComponent ) )
+	else if (UPointLightComponent* PointLightComponent = Cast<UPointLightComponent>(SceneComponent))
 	{
-		if ( USpotLightComponent* SpotLightComponent = Cast< USpotLightComponent >( SceneComponent ) )
+		if (USpotLightComponent* SpotLightComponent = Cast<USpotLightComponent>(SceneComponent))
 		{
-			UsdToUnreal::ConvertLuxShapingAPI( Prim, *SpotLightComponent, Context->Time );
+			UsdToUnreal::ConvertLuxShapingAPI(Prim, *SpotLightComponent, Context->Time);
 		}
 
-		UsdToUnreal::ConvertSphereLight( Prim, *PointLightComponent, Context->Time );
+		UsdToUnreal::ConvertSphereLight(Prim, *PointLightComponent, Context->Time);
 	}
-	else if ( USkyLightComponent* SkyLightComponent = Cast< USkyLightComponent >( SceneComponent ) )
+	else if (USkyLightComponent* SkyLightComponent = Cast<USkyLightComponent>(SceneComponent))
 	{
 		UsdToUnreal::ConvertDomeLight(Prim, *SkyLightComponent, Context->AssetCache.Get(), Context->bReuseIdenticalAssets);
-		SkyLightComponent->Mobility = EComponentMobility::Movable; // We won't bake geometry in the sky light so it needs to be movable
+		SkyLightComponent->Mobility = EComponentMobility::Movable;	  // We won't bake geometry in the sky light so it needs to be movable
 	}
 
-	if ( !LightComponent->IsRegistered() )
+	if (!LightComponent->IsRegistered())
 	{
 		LightComponent->RegisterComponent();
 	}
 }
 
-bool FUsdLuxLightTranslator::CollapsesChildren( ECollapsingType CollapsingType ) const
+bool FUsdLuxLightTranslator::CollapsesChildren(ECollapsingType CollapsingType) const
 {
 	// If we have a custom draw mode, it means we should draw bounds/cards/etc. instead
 	// of our entire subtree, which is basically the same thing as collapsing
@@ -124,9 +116,9 @@ bool FUsdLuxLightTranslator::CollapsesChildren( ECollapsingType CollapsingType )
 	return false;
 }
 
-bool FUsdLuxLightTranslator::CanBeCollapsed( ECollapsingType CollapsingType ) const
+bool FUsdLuxLightTranslator::CanBeCollapsed(ECollapsingType CollapsingType) const
 {
 	return false;
 }
 
-#endif // #if USE_USD_SDK
\ No newline at end of file
+#endif	  // #if USE_USD_SDK
diff --git a/Engine/Plugins/Importers/USDImporter/Source/USDSchemas/Private/USDSchemaTranslator.cpp b/Engine/Plugins/Importers/USDImporter/Source/USDSchemas/Private/USDSchemaTranslator.cpp
--- a/Engine/Plugins/Importers/USDImporter/Source/USDSchemas/Private/USDSchemaTranslator.cpp
+++ b/Engine/Plugins/Importers/USDImporter/Source/USDSchemas/Private/USDSchemaTranslator.cpp
@@ -3,7 +3,6 @@
 #include "USDSchemaTranslator.h"
 
 #include "USDInfoCache.h"
-#include "USDErrorUtils.h"
 #include "USDSchemasModule.h"
 #include "USDTypesConversion.h"
 
@@ -20,35 +19,38 @@
 int32 FRegisteredSchemaTranslatorHandle::CurrentSchemaTranslatorId = 0;
 
 #if USE_USD_SDK
-
 #include "USDIncludesStart.h"
-	#include "pxr/base/tf/token.h"
-	#include "pxr/usd/usd/typed.h"
-	#include "pxr/usd/usdShade/tokens.h"
+#include "pxr/base/tf/token.h"
+#include "pxr/base/tf/type.h"
+#include "pxr/usd/usd/prim.h"
+#include "pxr/usd/usdShade/tokens.h"
 #include "USDIncludesEnd.h"
 
-#endif // #if USE_USD_SDK
+#endif	  // #if USE_USD_SDK
 
-TSharedPtr< FUsdSchemaTranslator > FUsdSchemaTranslatorRegistry::CreateTranslatorForSchema( TSharedRef< FUsdSchemaTranslationContext > InTranslationContext, const UE::FUsdTyped& InSchema )
+TSharedPtr<FUsdSchemaTranslator> FUsdSchemaTranslatorRegistry::CreateTranslatorForSchema(
+	TSharedRef<FUsdSchemaTranslationContext> InTranslationContext,
+	const UE::FUsdTyped& InSchema
+)
 {
 #if USE_USD_SDK
-	TUsdStore< pxr::UsdPrim > Prim = pxr::UsdPrim( InSchema.GetPrim() );
-	if ( !Prim.Get() )
+	TUsdStore<pxr::UsdPrim> Prim = pxr::UsdPrim(InSchema.GetPrim());
+	if (!Prim.Get())
 	{
 		return {};
 	}
 
-	for ( TPair< FString, FSchemaTranslatorsStack >& RegisteredSchemasStack : RegisteredSchemaTranslators )
+	for (TPair<FString, FSchemaTranslatorsStack>& RegisteredSchemasStack : RegisteredSchemaTranslators)
 	{
-		pxr::TfToken RegisteredSchemaToken( UnrealToUsd::ConvertString( *RegisteredSchemasStack.Key ).Get() );
-		pxr::TfType RegisteredSchemaType = pxr::UsdSchemaRegistry::GetTypeFromName( RegisteredSchemaToken );
+		pxr::TfToken RegisteredSchemaToken(UnrealToUsd::ConvertString(*RegisteredSchemasStack.Key).Get());
+		pxr::TfType RegisteredSchemaType = pxr::UsdSchemaRegistry::GetTypeFromName(RegisteredSchemaToken);
 
-		if ( !RegisteredSchemaType.IsUnknown() &&  Prim.Get().IsA( RegisteredSchemaType ) && RegisteredSchemasStack.Value.Num() > 0 )
+		if (!RegisteredSchemaType.IsUnknown() && Prim.Get().IsA(RegisteredSchemaType) && RegisteredSchemasStack.Value.Num() > 0)
 		{
-			return RegisteredSchemasStack.Value.Top().CreateFunction( InTranslationContext, InSchema );
+			return RegisteredSchemasStack.Value.Top().CreateFunction(InTranslationContext, InSchema);
 		}
 	}
-#endif // #if USE_USD_SDK
+#endif	  // #if USE_USD_SDK
 
 	return {};
 }
@@ -58,33 +60,33 @@ FUsdRenderContextRegistry::FUsdRenderContextRegistry()
 #if USE_USD_SDK
 	LLM_SCOPE_BYTAG(Usd);
 
-	UniversalRenderContext = FName( UsdToUnreal::ConvertToken( pxr::UsdShadeTokens->universalRenderContext ) );
-	Register( UniversalRenderContext );
+	UniversalRenderContext = FName(UsdToUnreal::ConvertToken(pxr::UsdShadeTokens->universalRenderContext));
+	Register(UniversalRenderContext);
 
-	UnrealRenderContext = FName( UsdToUnreal::ConvertToken( UnrealIdentifiers::Unreal ) );
-	Register( UnrealRenderContext );
-#endif // #if USE_USD_SDK
+	UnrealRenderContext = FName(UsdToUnreal::ConvertToken(UnrealIdentifiers::Unreal));
+	Register(UnrealRenderContext);
+#endif	  // #if USE_USD_SDK
 }
 
-FRegisteredSchemaTranslatorHandle FUsdSchemaTranslatorRegistry::Register( const FString& SchemaName, FCreateTranslator CreateFunction )
+FRegisteredSchemaTranslatorHandle FUsdSchemaTranslatorRegistry::Register(const FString& SchemaName, FCreateTranslator CreateFunction)
 {
 #if USE_USD_SDK
-	FSchemaTranslatorsStack* SchemaTranslatorsStack = FindSchemaTranslatorStack( SchemaName );
+	FSchemaTranslatorsStack* SchemaTranslatorsStack = FindSchemaTranslatorStack(SchemaName);
 
-	if ( !SchemaTranslatorsStack )
+	if (!SchemaTranslatorsStack)
 	{
 		// Insert most specialized first
 		int32 SchemaRegistryIndex = 0;
 
-		pxr::TfToken SchemaToRegisterToken( UnrealToUsd::ConvertString( *SchemaName ).Get() );
-		pxr::TfType SchemaToRegisterType = pxr::UsdSchemaRegistry::GetTypeFromName( SchemaToRegisterToken );
+		pxr::TfToken SchemaToRegisterToken(UnrealToUsd::ConvertString(*SchemaName).Get());
+		pxr::TfType SchemaToRegisterType = pxr::UsdSchemaRegistry::GetTypeFromName(SchemaToRegisterToken);
 
-		for ( TPair< FString, FSchemaTranslatorsStack >& RegisteredSchemasStack : RegisteredSchemaTranslators )
+		for (TPair<FString, FSchemaTranslatorsStack>& RegisteredSchemasStack : RegisteredSchemaTranslators)
 		{
-			pxr::TfToken RegisteredSchemaToken( UnrealToUsd::ConvertString( *RegisteredSchemasStack.Key ).Get() );
-			pxr::TfType RegisteredSchemaType = pxr::UsdSchemaRegistry::GetTypeFromName( RegisteredSchemaToken );
+			pxr::TfToken RegisteredSchemaToken(UnrealToUsd::ConvertString(*RegisteredSchemasStack.Key).Get());
+			pxr::TfType RegisteredSchemaType = pxr::UsdSchemaRegistry::GetTypeFromName(RegisteredSchemaToken);
 
-			if ( SchemaToRegisterType.IsA( RegisteredSchemaType ) )
+			if (SchemaToRegisterType.IsA(RegisteredSchemaType))
 			{
 				// We need to be registered before our ancestor types
 				break;
@@ -95,34 +97,34 @@ FRegisteredSchemaTranslatorHandle FUsdSchemaTranslatorRegistry::Register( const
 			}
 		}
 
-		SchemaTranslatorsStack = &RegisteredSchemaTranslators.EmplaceAt_GetRef( SchemaRegistryIndex, SchemaName, FSchemaTranslatorsStack() ).Value;
+		SchemaTranslatorsStack = &RegisteredSchemaTranslators.EmplaceAt_GetRef(SchemaRegistryIndex, SchemaName, FSchemaTranslatorsStack()).Value;
 	}
 
 	FRegisteredSchemaTranslator RegisteredSchemaTranslator;
-	RegisteredSchemaTranslator.Handle = FRegisteredSchemaTranslatorHandle( SchemaName );
+	RegisteredSchemaTranslator.Handle = FRegisteredSchemaTranslatorHandle(SchemaName);
 	RegisteredSchemaTranslator.CreateFunction = CreateFunction;
 
-	SchemaTranslatorsStack->Push( RegisteredSchemaTranslator );
+	SchemaTranslatorsStack->Push(RegisteredSchemaTranslator);
 
 	return RegisteredSchemaTranslator.Handle;
 #else
 	return FRegisteredSchemaTranslatorHandle();
-#endif // #if USE_USD_SDK
+#endif	  // #if USE_USD_SDK
 }
 
-void FUsdSchemaTranslatorRegistry::Unregister( const FRegisteredSchemaTranslatorHandle& TranslatorHandle )
+void FUsdSchemaTranslatorRegistry::Unregister(const FRegisteredSchemaTranslatorHandle& TranslatorHandle)
 {
-	FSchemaTranslatorsStack* SchemaTranslatorsStack = FindSchemaTranslatorStack( TranslatorHandle.GetSchemaName() );
+	FSchemaTranslatorsStack* SchemaTranslatorsStack = FindSchemaTranslatorStack(TranslatorHandle.GetSchemaName());
 
-	if ( !SchemaTranslatorsStack )
+	if (!SchemaTranslatorsStack)
 	{
 		return;
 	}
 
-	for ( FSchemaTranslatorsStack::TIterator RegisteredSchemaTranslatorIt = SchemaTranslatorsStack->CreateIterator();
-		RegisteredSchemaTranslatorIt; ++RegisteredSchemaTranslatorIt )
+	for (FSchemaTranslatorsStack::TIterator RegisteredSchemaTranslatorIt = SchemaTranslatorsStack->CreateIterator(); RegisteredSchemaTranslatorIt;
+		 ++RegisteredSchemaTranslatorIt)
 	{
-		if ( RegisteredSchemaTranslatorIt->Handle.GetId() == TranslatorHandle.GetId() )
+		if (RegisteredSchemaTranslatorIt->Handle.GetId() == TranslatorHandle.GetId())
 		{
 			RegisteredSchemaTranslatorIt.RemoveCurrent();
 			break;
@@ -130,23 +132,25 @@ void FUsdSchemaTranslatorRegistry::Unregister( const FRegisteredSchemaTranslator
 	}
 }
 
-FUsdSchemaTranslationContext::FUsdSchemaTranslationContext( const UE::FUsdStage& InStage, UUsdAssetCache2& InAssetCache )
-	: Stage( InStage )
-	, AssetCache( &InAssetCache )
+FUsdSchemaTranslationContext::FUsdSchemaTranslationContext(const UE::FUsdStage& InStage, UUsdAssetCache2& InAssetCache)
+	: Stage(InStage)
+	, AssetCache(&InAssetCache)
 {
-	IUsdSchemasModule& UsdSchemasModule = FModuleManager::Get().LoadModuleChecked< IUsdSchemasModule >( TEXT("USDSchemas") );
+	IUsdSchemasModule& UsdSchemasModule = FModuleManager::Get().LoadModuleChecked<IUsdSchemasModule>(TEXT("USDSchemas"));
 	RenderContext = UsdSchemasModule.GetRenderContextRegistry().GetUniversalRenderContext();
 }
 
-FUsdSchemaTranslatorRegistry::FSchemaTranslatorsStack* FUsdSchemaTranslatorRegistry::FindSchemaTranslatorStack( const FString& SchemaName )
+FUsdSchemaTranslatorRegistry::FSchemaTranslatorsStack* FUsdSchemaTranslatorRegistry::FindSchemaTranslatorStack(const FString& SchemaName)
 {
-	TPair< FString, FSchemaTranslatorsStack >* Result = Algo::FindByPredicate( RegisteredSchemaTranslators,
-		[ &SchemaName ]( const TPair< FString, FSchemaTranslatorsStack >& Element ) -> bool
+	TPair<FString, FSchemaTranslatorsStack>* Result = Algo::FindByPredicate(
+		RegisteredSchemaTranslators,
+		[&SchemaName](const TPair<FString, FSchemaTranslatorsStack>& Element) -> bool
 		{
 			return Element.Key == SchemaName;
-		} );
+		}
+	);
 
-	if ( Result )
+	if (Result)
 	{
 		return &Result->Value;
 	}
@@ -158,38 +162,39 @@ FUsdSchemaTranslatorRegistry::FSchemaTranslatorsStack* FUsdSchemaTranslatorRegis
 
 void FUsdSchemaTranslationContext::CompleteTasks()
 {
-	TRACE_CPUPROFILER_EVENT_SCOPE( FUsdSchemaTranslationContext::CompleteTasks );
+	TRACE_CPUPROFILER_EVENT_SCOPE(FUsdSchemaTranslationContext::CompleteTasks);
 
-	FScopedSlowTask SlowTask( TranslatorTasks.Num(), LOCTEXT( "TasksProgress", "Executing USD Schema tasks" ) );
+	FScopedSlowTask SlowTask(TranslatorTasks.Num(), LOCTEXT("TasksProgress", "Executing USD Schema tasks"));
 
 	// These are just pointers into the TranslatorTasks items that are task chains that are in the ESchemaTranslationStatus::Pending state.
 	// We'll use this below to execute tasks until all items in TranslatorTasks are pending, and then we switch to/from an ExclusiveSync pass.
 	// This because we need to ensure that all ExclusiveSync tasks are run in isolation.
-	TSet< FUsdSchemaTranslatorTaskChain* > PendingTaskChains;
+	TSet<FUsdSchemaTranslatorTaskChain*> PendingTaskChains;
 
 	// Note that this first pass is for tasks that allow concurrent execution (so *not* exclusive sync tasks). If this is ever changed,
 	// we would also need to change the behavior of StartIfAsync to delay until the proper async pass, and not start right away
 	bool bExclusiveSyncTasks = false;
 
-	bool bFinished = ( TranslatorTasks.Num() == 0 );
-	while ( !bFinished )
+	bool bFinished = (TranslatorTasks.Num() == 0);
+	while (!bFinished)
 	{
 		while (TranslatorTasks.Num() > PendingTaskChains.Num())
 		{
-			for ( TArray< TSharedPtr< FUsdSchemaTranslatorTaskChain > >::TIterator TaskChainIterator = TranslatorTasks.CreateIterator(); TaskChainIterator; ++TaskChainIterator )
+			for (TArray<TSharedPtr<FUsdSchemaTranslatorTaskChain>>::TIterator TaskChainIterator = TranslatorTasks.CreateIterator(); TaskChainIterator;
+				 ++TaskChainIterator)
 			{
-				TSharedPtr< FUsdSchemaTranslatorTaskChain >& TaskChain = *TaskChainIterator;
+				TSharedPtr<FUsdSchemaTranslatorTaskChain>& TaskChain = *TaskChainIterator;
 
-				ESchemaTranslationStatus TaskChainStatus = TaskChain->Execute( bExclusiveSyncTasks );
+				ESchemaTranslationStatus TaskChainStatus = TaskChain->Execute(bExclusiveSyncTasks);
 
-				if ( TaskChainStatus == ESchemaTranslationStatus::Done )
+				if (TaskChainStatus == ESchemaTranslationStatus::Done)
 				{
- 					SlowTask.EnterProgressFrame();
+					SlowTask.EnterProgressFrame();
 					TaskChainIterator.RemoveCurrent();
 				}
-				else if ( TaskChainStatus == ESchemaTranslationStatus::Pending )
+				else if (TaskChainStatus == ESchemaTranslationStatus::Pending)
 				{
-					PendingTaskChains.Add( TaskChain.Get() );
+					PendingTaskChains.Add(TaskChain.Get());
 				}
 			}
 		}
@@ -197,14 +202,14 @@ void FUsdSchemaTranslationContext::CompleteTasks()
 		bExclusiveSyncTasks = !bExclusiveSyncTasks;
 		PendingTaskChains.Reset();
 
-		bFinished = ( TranslatorTasks.Num() == 0 );
+		bFinished = (TranslatorTasks.Num() == 0);
 	}
 }
 
-bool FUsdSchemaTranslator::IsCollapsed( ECollapsingType CollapsingType ) const
+bool FUsdSchemaTranslator::IsCollapsed(ECollapsingType CollapsingType) const
 {
 #if USE_USD_SDK
-	TRACE_CPUPROFILER_EVENT_SCOPE( FUsdSchemaTranslator::IsCollapsed );
+	TRACE_CPUPROFILER_EVENT_SCOPE(FUsdSchemaTranslator::IsCollapsed);
 
 	if (!Context->bIsBuildingInfoCache)
 	{
@@ -212,31 +217,32 @@ bool FUsdSchemaTranslator::IsCollapsed( ECollapsingType CollapsingType ) const
 	}
 
 	// This is merely a fallback, and we should never need this
-	return CanBeCollapsed( CollapsingType );
-#endif // #if USE_USD_SDK
+	return CanBeCollapsed(CollapsingType);
+#endif	  // #if USE_USD_SDK
 
 	return false;
 }
 
 void FSchemaTranslatorTask::Start()
 {
-	if ( LaunchPolicy == ESchemaTranslationLaunchPolicy::Async && IsInGameThread() )
+	if (LaunchPolicy == ESchemaTranslationLaunchPolicy::Async && IsInGameThread())
 	{
 		Result = Async(
 #if WITH_EDITOR
 			EAsyncExecution::LargeThreadPool,
 #else
 			EAsyncExecution::ThreadPool,
-#endif // WITH_EDITOR
-			[ this ]() -> bool
+#endif	  // WITH_EDITOR
+			[this]() -> bool
 			{
 				return DoWork();
-			} );
+			}
+		);
 	}
 	else
 	{
 		// Execute on this thread
-		if ( !DoWork() )
+		if (!DoWork())
 		{
 			Continuation.Reset();
 		}
@@ -245,7 +251,7 @@ void FSchemaTranslatorTask::Start()
 
 void FSchemaTranslatorTask::StartIfAsync()
 {
-	if ( LaunchPolicy == ESchemaTranslationLaunchPolicy::Async )
+	if (LaunchPolicy == ESchemaTranslationLaunchPolicy::Async)
 	{
 		Start();
 	}
@@ -253,41 +259,41 @@ void FSchemaTranslatorTask::StartIfAsync()
 
 bool FSchemaTranslatorTask::DoWork()
 {
-	ensure( bIsDone == false );
+	ensure(bIsDone == false);
 	bool bContinue = Callable();
 	bIsDone = true;
 
 	return bContinue;
 }
 
-FUsdSchemaTranslatorTaskChain& FUsdSchemaTranslatorTaskChain::Do( ESchemaTranslationLaunchPolicy InPolicy, TFunction< bool() > Callable )
+FUsdSchemaTranslatorTaskChain& FUsdSchemaTranslatorTaskChain::Do(ESchemaTranslationLaunchPolicy InPolicy, TFunction<bool()> Callable)
 {
-	if ( !CurrentTask )
+	if (!CurrentTask)
 	{
-		CurrentTask = MakeShared< FSchemaTranslatorTask >( InPolicy, Callable );
+		CurrentTask = MakeShared<FSchemaTranslatorTask>(InPolicy, Callable);
 
-		CurrentTask->StartIfAsync(); // Queue it right now if async
+		CurrentTask->StartIfAsync();	// Queue it right now if async
 	}
 	else
 	{
-		Then( InPolicy, Callable );
+		Then(InPolicy, Callable);
 	}
 
 	return *this;
 }
 
-FUsdSchemaTranslatorTaskChain& FUsdSchemaTranslatorTaskChain::Then( ESchemaTranslationLaunchPolicy InPolicy, TFunction< bool() > Callable )
+FUsdSchemaTranslatorTaskChain& FUsdSchemaTranslatorTaskChain::Then(ESchemaTranslationLaunchPolicy InPolicy, TFunction<bool()> Callable)
 {
-	TSharedPtr< FSchemaTranslatorTask > LastTask = CurrentTask;
+	TSharedPtr<FSchemaTranslatorTask> LastTask = CurrentTask;
 
-	while ( LastTask->Continuation.IsValid() )
+	while (LastTask->Continuation.IsValid())
 	{
 		LastTask = LastTask->Continuation;
 	}
 
-	if ( LastTask )
+	if (LastTask)
 	{
-		LastTask->Continuation = MakeShared< FSchemaTranslatorTask >( InPolicy, Callable );
+		LastTask->Continuation = MakeShared<FSchemaTranslatorTask>(InPolicy, Callable);
 	}
 
 	return *this;
@@ -295,9 +301,9 @@ FUsdSchemaTranslatorTaskChain& FUsdSchemaTranslatorTaskChain::Then( ESchemaTrans
 
 namespace UsdSchemaTranslatorTaskChainImpl
 {
-	FORCEINLINE bool CanStart( FSchemaTranslatorTask* Task, bool bExclusiveSyncTasks )
+	FORCEINLINE bool CanStart(FSchemaTranslatorTask* Task, bool bExclusiveSyncTasks)
 	{
-		return ( Task->LaunchPolicy == ESchemaTranslationLaunchPolicy::ExclusiveSync ) == bExclusiveSyncTasks;
+		return (Task->LaunchPolicy == ESchemaTranslationLaunchPolicy::ExclusiveSync) == bExclusiveSyncTasks;
 	}
 }
 
@@ -305,16 +311,16 @@ ESchemaTranslationStatus FUsdSchemaTranslatorTaskChain::Execute(bool bExclusiveS
 {
 	FSchemaTranslatorTask* TranslatorTask = CurrentTask.Get();
 
-	if ( TranslatorTask == nullptr )
+	if (TranslatorTask == nullptr)
 	{
 		return ESchemaTranslationStatus::Done;
 	}
 
-	if ( !TranslatorTask->IsDone() )
+	if (!TranslatorTask->IsDone())
 	{
-		if ( !TranslatorTask->IsStarted() )
+		if (!TranslatorTask->IsStarted())
 		{
-			if ( UsdSchemaTranslatorTaskChainImpl::CanStart( TranslatorTask, bExclusiveSyncTasks ) )
+			if (UsdSchemaTranslatorTaskChainImpl::CanStart(TranslatorTask, bExclusiveSyncTasks))
 			{
 				TranslatorTask->Start();
 			}
@@ -328,7 +334,7 @@ ESchemaTranslationStatus FUsdSchemaTranslatorTaskChain::Execute(bool bExclusiveS
 	}
 	else
 	{
-		if ( CurrentTask->Result.IsSet() )
+		if (CurrentTask->Result.IsSet())
 		{
 			CurrentTask = CurrentTask->Result->Get() ? CurrentTask->Continuation : nullptr;
 		}
@@ -337,13 +343,13 @@ ESchemaTranslationStatus FUsdSchemaTranslatorTaskChain::Execute(bool bExclusiveS
 			CurrentTask = CurrentTask->Continuation;
 		}
 
-		if (( TranslatorTask = CurrentTask.Get()) != nullptr )
+		if ((TranslatorTask = CurrentTask.Get()) != nullptr)
 		{
-			if ( UsdSchemaTranslatorTaskChainImpl::CanStart( TranslatorTask, bExclusiveSyncTasks ) )
+			if (UsdSchemaTranslatorTaskChainImpl::CanStart(TranslatorTask, bExclusiveSyncTasks))
 			{
-				if ( IsInGameThread() )
+				if (IsInGameThread())
 				{
-					TranslatorTask->StartIfAsync(); // Queue the next task asap if async
+					TranslatorTask->StartIfAsync();	   // Queue the next task asap if async
 				}
 				else
 				{
diff --git a/Engine/Plugins/Importers/USDImporter/Source/USDSchemas/Private/USDSchemasModule.cpp b/Engine/Plugins/Importers/USDImporter/Source/USDSchemas/Private/USDSchemasModule.cpp
--- a/Engine/Plugins/Importers/USDImporter/Source/USDSchemas/Private/USDSchemasModule.cpp
+++ b/Engine/Plugins/Importers/USDImporter/Source/USDSchemas/Private/USDSchemasModule.cpp
@@ -6,26 +6,25 @@
 
 #if WITH_EDITOR
 #include "MDLImporterModule.h"
-#endif // #if WITH_EDITOR
+#endif	  // #if WITH_EDITOR
 
 #include "Modules/ModuleManager.h"
-#include "Misc/Paths.h"
 
 #if USE_USD_SDK
+#include "Custom/MaterialXUSDShadeMaterialTranslator.h"
+#include "Custom/MDLUSDShadeMaterialTranslator.h"
 #include "USDGeomCameraTranslator.h"
 #include "USDGeometryCacheTranslator.h"
 #include "USDGeomMeshTranslator.h"
 #include "USDGeomPointInstancerTranslator.h"
+#include "USDGeomPrimitiveTranslator.h"
 #include "USDGeomXformableTranslator.h"
 #include "USDGroomTranslator.h"
-#include "USDGeomPrimitiveTranslator.h"
 #include "USDLuxLightTranslator.h"
 #include "USDMemory.h"
 #include "USDShadeMaterialTranslator.h"
 #include "USDSkelSkeletonTranslator.h"
-#include "Custom/MaterialXUSDShadeMaterialTranslator.h"
-#include "Custom/MDLUSDShadeMaterialTranslator.h"
-#endif // #if USE_USD_SDK
+#endif	  // #if USE_USD_SDK
 
 class FUsdSchemasModule : public IUsdSchemasModule
 {
@@ -76,9 +75,9 @@ public:
 				TranslatorHandles.Add(Registry.Register<FMdlUsdShadeMaterialTranslator>(TEXT("UsdShadeMaterial")));
 			}
 		}
-#endif // WITH_EDITOR
+#endif	  // WITH_EDITOR
 
-#endif // #if USE_USD_SDK
+#endif	  // #if USE_USD_SDK
 	}
 
 	virtual void ShutdownModule() override
@@ -94,7 +93,7 @@ public:
 #if USE_USD_SDK && WITH_EDITOR
 		ShaderRegistry.Unregister(FMdlUsdShadeMaterialTranslator::MdlRenderContext);
 		ShaderRegistry.Unregister(FMaterialXUsdShadeMaterialTranslator::MaterialXRenderContext);
-#endif // WITH_EDITOR
+#endif	  // WITH_EDITOR
 	}
 
 	virtual FUsdSchemaTranslatorRegistry& GetTranslatorRegistry() override
@@ -111,7 +110,7 @@ protected:
 	FUsdSchemaTranslatorRegistry UsdSchemaTranslatorRegistry;
 	FUsdRenderContextRegistry UsdRenderContextRegistry;
 
-	TArray< FRegisteredSchemaTranslatorHandle > TranslatorHandles;
+	TArray<FRegisteredSchemaTranslatorHandle> TranslatorHandles;
 };
 
-IMPLEMENT_MODULE_USD( FUsdSchemasModule, USDSchemas );
+IMPLEMENT_MODULE_USD(FUsdSchemasModule, USDSchemas);
