USD: Reformat USDSchemas module.

diff --git a/Engine/Plugins/Importers/USDImporter/Source/USDSchemas/Private/USDShadeMaterialTranslator.cpp b/Engine/Plugins/Importers/USDImporter/Source/USDSchemas/Private/USDShadeMaterialTranslator.cpp
--- a/Engine/Plugins/Importers/USDImporter/Source/USDSchemas/Private/USDShadeMaterialTranslator.cpp
+++ b/Engine/Plugins/Importers/USDImporter/Source/USDSchemas/Private/USDShadeMaterialTranslator.cpp
@@ -3,7 +3,6 @@
 #include "USDShadeMaterialTranslator.h"
 
 #include "MeshTranslationImpl.h"
-#include "USDAssetCache.h"
 #include "USDAssetUserData.h"
 #include "USDClassesModule.h"
 #include "USDLog.h"
@@ -16,7 +15,6 @@
 #include "EditorFramework/AssetImportData.h"
 #include "Engine/Level.h"
 #include "Engine/Texture.h"
-#include "Materials/Material.h"
 #include "Materials/MaterialInstanceConstant.h"
 #include "Materials/MaterialInstanceDynamic.h"
 #include "MaterialShared.h"
@@ -25,56 +23,56 @@
 #if USE_USD_SDK
 
 #include "USDIncludesStart.h"
-	#include "pxr/base/tf/token.h"
-	#include "pxr/usd/usdShade/material.h"
-	#include "pxr/usd/usdShade/tokens.h"
+#include "pxr/base/tf/token.h"
+#include "pxr/usd/usdShade/material.h"
+#include "pxr/usd/usdShade/tokens.h"
 #include "USDIncludesEnd.h"
 
 namespace UE::UsdShadeTranslator::Private
 {
 	void RecursiveUpgradeMaterialsAndTexturesToVT(
 		const TSet<UTexture*>& TexturesToUpgrade,
-		const TSharedRef< FUsdSchemaTranslationContext >& Context,
+		const TSharedRef<FUsdSchemaTranslationContext>& Context,
 		TSet<UMaterialInterface*>& VisitedMaterials,
 		TSet<UMaterialInterface*>& NewMaterials
 	)
 	{
-		for ( UTexture* Texture : TexturesToUpgrade )
+		for (UTexture* Texture : TexturesToUpgrade)
 		{
-			if ( Texture->VirtualTextureStreaming )
+			if (Texture->VirtualTextureStreaming)
 			{
 				continue;
 			}
 
-			UE_LOG( LogUsd, Log, TEXT( "Upgrading texture '%s' to VT as it is used by a material that must be VT" ),
-				*Texture->GetName()
+			UE_LOG(LogUsd, Log, TEXT("Upgrading texture '%s' to VT as it is used by a material that must be VT"), *Texture->GetName());
+			FPropertyChangedEvent PropertyChangeEvent(
+				UTexture::StaticClass()->FindPropertyByName(GET_MEMBER_NAME_CHECKED(UTexture, VirtualTextureStreaming))
 			);
-			FPropertyChangedEvent PropertyChangeEvent( UTexture::StaticClass()->FindPropertyByName( GET_MEMBER_NAME_CHECKED( UTexture, VirtualTextureStreaming ) ) );
 			Texture->Modify();
 			Texture->VirtualTextureStreaming = true;
 
 #if WITH_EDITOR
-			Texture->PostEditChangeProperty( PropertyChangeEvent );
-#endif // WITH_EDITOR
+			Texture->PostEditChangeProperty(PropertyChangeEvent);
+#endif	  // WITH_EDITOR
 
 			// Now that our texture is VT, all materials that use the texture must be VT
-			if ( const TSet<UMaterialInterface*>* UserMaterials = Context->TextureToUserMaterials.Find( Texture ) )
+			if (const TSet<UMaterialInterface*>* UserMaterials = Context->TextureToUserMaterials.Find(Texture))
 			{
-				for ( UMaterialInterface* UserMaterial : *UserMaterials )
+				for (UMaterialInterface* UserMaterial : *UserMaterials)
 				{
-					if ( VisitedMaterials.Contains( UserMaterial ) )
+					if (VisitedMaterials.Contains(UserMaterial))
 					{
 						continue;
 					}
 
-					if ( UMaterialInstance* MaterialInstance = Cast<UMaterialInstance>( UserMaterial ) )
+					if (UMaterialInstance* MaterialInstance = Cast<UMaterialInstance>(UserMaterial))
 					{
 						// Important to not use GetBaseMaterial() here because if our parent is the translucent we'll
 						// get the reference UsdPreviewSurface instead, as that is also *its* reference
 						UMaterialInterface* ReferenceMaterial = MaterialInstance->Parent.Get();
-						UMaterialInterface* ReferenceMaterialVT =
-							MeshTranslationImpl::GetVTVersionOfReferencePreviewSurfaceMaterial( ReferenceMaterial );
-						if ( ReferenceMaterial == ReferenceMaterialVT )
+						UMaterialInterface* ReferenceMaterialVT = MeshTranslationImpl::GetVTVersionOfReferencePreviewSurfaceMaterial(ReferenceMaterial
+						);
+						if (ReferenceMaterial == ReferenceMaterialVT)
 						{
 							// Material is already VT, we're good
 							continue;
@@ -83,160 +81,155 @@ namespace UE::UsdShadeTranslator::Private
 						// Visit it before we start recursing. We need this because we must convert textures to VT
 						// before materials (or else we get a warning) but we'll only actually swap the reference material
 						// at the end of this scope
-						VisitedMaterials.Add( UserMaterial );
+						VisitedMaterials.Add(UserMaterial);
 
 						// If we're going to update this material to VT, all of *its* textures need to be VT too
 						TSet<UTexture*> OtherUsedTextures;
-						for ( const FTextureParameterValue& TextureValue : MaterialInstance->TextureParameterValues )
+						for (const FTextureParameterValue& TextureValue : MaterialInstance->TextureParameterValues)
 						{
-							if ( UTexture* OtherTexture = TextureValue.ParameterValue )
+							if (UTexture* OtherTexture = TextureValue.ParameterValue)
 							{
-								OtherUsedTextures.Add( OtherTexture );
+								OtherUsedTextures.Add(OtherTexture);
 							}
 						}
 
-						RecursiveUpgradeMaterialsAndTexturesToVT(
-							OtherUsedTextures,
-							Context,
-							VisitedMaterials,
-							NewMaterials
-						);
+						RecursiveUpgradeMaterialsAndTexturesToVT(OtherUsedTextures, Context, VisitedMaterials, NewMaterials);
 
 						// We can't blindly recreate all component render states when a level is being added, because
 						// we may end up first creating render states for some components, and UWorld::AddToWorld
 						// calls FScene::AddPrimitive which expects the component to not have primitives yet
 						FMaterialUpdateContext::EOptions::Type Options = FMaterialUpdateContext::EOptions::Default;
-						if ( Context->Level->bIsAssociatingLevel )
+						if (Context->Level->bIsAssociatingLevel)
 						{
-							Options = ( FMaterialUpdateContext::EOptions::Type ) (
-								Options & ~FMaterialUpdateContext::EOptions::RecreateRenderStates
-							);
+							Options = (FMaterialUpdateContext::EOptions::Type)(Options & ~FMaterialUpdateContext::EOptions::RecreateRenderStates);
 						}
 
-						UE_LOG( LogUsd, Log, TEXT( "Upgrading material instance '%s' to having a VT reference as texture '%s' requires it" ),
+						UE_LOG(
+							LogUsd,
+							Log,
+							TEXT("Upgrading material instance '%s' to having a VT reference as texture '%s' requires it"),
 							*MaterialInstance->GetName(),
 							*Texture->GetName()
 						);
 
 #if WITH_EDITOR
-						UMaterialInstanceConstant* MIC = Cast<UMaterialInstanceConstant>( MaterialInstance );
-						if ( GIsEditor && MIC )
+						UMaterialInstanceConstant* MIC = Cast<UMaterialInstanceConstant>(MaterialInstance);
+						if (GIsEditor && MIC)
 						{
-							FMaterialUpdateContext UpdateContext( Options, GMaxRHIShaderPlatform );
-							UpdateContext.AddMaterialInstance( MIC );
-							MIC->PreEditChange( nullptr );
-							MIC->SetParentEditorOnly( ReferenceMaterialVT );
+							FMaterialUpdateContext UpdateContext(Options, GMaxRHIShaderPlatform);
+							UpdateContext.AddMaterialInstance(MIC);
+							MIC->PreEditChange(nullptr);
+							MIC->SetParentEditorOnly(ReferenceMaterialVT);
 							MIC->PostEditChange();
 						}
 						else
-#endif // WITH_EDITOR
-						if ( UMaterialInstanceDynamic* MID = Cast<UMaterialInstanceDynamic>( UserMaterial ) )
-						{
-							if ( Context->AssetCache && Context->InfoCache )
+#endif	  // WITH_EDITOR
+							if (UMaterialInstanceDynamic* MID = Cast<UMaterialInstanceDynamic>(UserMaterial))
 							{
-								UE::FSdfPath PrimPath;
-								for (const UE::FSdfPath& Path : Context->InfoCache->GetPrimsForAsset(MID))
+								if (Context->AssetCache && Context->InfoCache)
 								{
-									PrimPath = Path;
-									break;
-								}
-								const FString Hash = Context->AssetCache->GetHashForAsset( MID );
-
-								const FName NewInstanceName = MakeUniqueObjectName(
-									GetTransientPackage(),
-									UMaterialInstance::StaticClass(),
-									PrimPath.IsEmpty()
-										? TEXT("MaterialInstance")
-										: *IUsdClassesModule::SanitizeObjectName(FPaths::GetBaseFilename(PrimPath.GetString()))
-								);
-
-								// For MID we can't swap the reference material, so we need to create a brand new one and copy
-								// the overrides
-								UMaterialInstanceDynamic* NewMID = UMaterialInstanceDynamic::Create(
-									ReferenceMaterialVT,
-									GetTransientPackage(),
-									NewInstanceName
-								);
-								if ( !ensure( NewMID ) )
-								{
-									continue;
-								}
+									UE::FSdfPath PrimPath;
+									for (const UE::FSdfPath& Path : Context->InfoCache->GetPrimsForAsset(MID))
+									{
+										PrimPath = Path;
+										break;
+									}
+									const FString Hash = Context->AssetCache->GetHashForAsset(MID);
+
+									const FName NewInstanceName = MakeUniqueObjectName(
+										GetTransientPackage(),
+										UMaterialInstance::StaticClass(),
+										PrimPath.IsEmpty() ? TEXT("MaterialInstance")
+														   : *IUsdClassesModule::SanitizeObjectName(FPaths::GetBaseFilename(PrimPath.GetString()))
+									);
+
+									// For MID we can't swap the reference material, so we need to create a brand new one and copy
+									// the overrides
+									UMaterialInstanceDynamic* NewMID = UMaterialInstanceDynamic::Create(
+										ReferenceMaterialVT,
+										GetTransientPackage(),
+										NewInstanceName
+									);
+									if (!ensure(NewMID))
+									{
+										continue;
+									}
 
-								NewMID->CopyParameterOverrides( MID );
+									NewMID->CopyParameterOverrides(MID);
 
-								UUsdMaterialAssetUserData* OldUserData = UserMaterial->GetAssetUserData<UUsdMaterialAssetUserData>();
-								if(OldUserData)
-								{
-									UUsdMaterialAssetUserData* NewUserData = DuplicateObject(OldUserData, NewMID);
-									NewMID->AddAssetUserData(NewUserData);
-								}
+									UUsdMaterialAssetUserData* OldUserData = UserMaterial->GetAssetUserData<UUsdMaterialAssetUserData>();
+									if (OldUserData)
+									{
+										UUsdMaterialAssetUserData* NewUserData = DuplicateObject(OldUserData, NewMID);
+										NewMID->AddAssetUserData(NewUserData);
+									}
 
-								if (Context->AssetCache->CanRemoveAsset(Hash) && Context->AssetCache->RemoveAsset(Hash))
-								{
-									Context->AssetCache->CacheAsset(Hash, NewMID);
-									if (!PrimPath.IsEmpty())
+									if (Context->AssetCache->CanRemoveAsset(Hash) && Context->AssetCache->RemoveAsset(Hash))
 									{
-										Context->InfoCache->LinkAssetToPrim(PrimPath, NewMID);
+										Context->AssetCache->CacheAsset(Hash, NewMID);
+										if (!PrimPath.IsEmpty())
+										{
+											Context->InfoCache->LinkAssetToPrim(PrimPath, NewMID);
+										}
+										NewMaterials.Add(NewMID);
 									}
-									NewMaterials.Add( NewMID );
 								}
 							}
-						}
-						else
-						{
-							// This should never happen
-							ensure( false );
-						}
+							else
+							{
+								// This should never happen
+								ensure(false);
+							}
 					}
 				}
 			}
 		}
 	}
 
-	void UpgradeMaterialsAndTexturesToVT( TSet<UTexture*> TexturesToUpgrade, TSharedRef< FUsdSchemaTranslationContext >& Context )
+	void UpgradeMaterialsAndTexturesToVT(TSet<UTexture*> TexturesToUpgrade, TSharedRef<FUsdSchemaTranslationContext>& Context)
 	{
 		TSet<UMaterialInterface*> VisitedMaterials;
 		TSet<UMaterialInterface*> NewMaterials;
-		RecursiveUpgradeMaterialsAndTexturesToVT( TexturesToUpgrade, Context, VisitedMaterials, NewMaterials );
+		RecursiveUpgradeMaterialsAndTexturesToVT(TexturesToUpgrade, Context, VisitedMaterials, NewMaterials);
 
 		// When we "visit" a MID we'll create a brand new instance of it and discard the old one, so let's drop the old ones
 		// from Context->TextureToUserMaterials too
-		for ( UMaterialInterface* Material : VisitedMaterials )
+		for (UMaterialInterface* Material : VisitedMaterials)
 		{
-			if ( UMaterialInstanceDynamic* OldMID = Cast<UMaterialInstanceDynamic>( Material ) )
+			if (UMaterialInstanceDynamic* OldMID = Cast<UMaterialInstanceDynamic>(Material))
 			{
-				for ( TPair<UTexture*, TSet<UMaterialInterface*>>& Pair : Context->TextureToUserMaterials )
+				for (TPair<UTexture*, TSet<UMaterialInterface*>>& Pair : Context->TextureToUserMaterials)
 				{
-					Pair.Value.Remove( Material );
+					Pair.Value.Remove(Material);
 				}
 			}
 		}
 
 		// Additionally now we need to add those new MIDs we created back into Context->TextureToUserMaterials
-		for ( UMaterialInterface* Material : NewMaterials )
+		for (UMaterialInterface* Material : NewMaterials)
 		{
-			UMaterialInstanceDynamic* MID = Cast<UMaterialInstanceDynamic>( Material );
-			if ( ensure( MID ) )
+			UMaterialInstanceDynamic* MID = Cast<UMaterialInstanceDynamic>(Material);
+			if (ensure(MID))
 			{
-				for ( const FTextureParameterValue& TextureValue : MID->TextureParameterValues )
+				for (const FTextureParameterValue& TextureValue : MID->TextureParameterValues)
 				{
-					if ( UTexture* Texture = TextureValue.ParameterValue )
+					if (UTexture* Texture = TextureValue.ParameterValue)
 					{
-						Context->TextureToUserMaterials.FindOrAdd( Texture ).Add( MID );
+						Context->TextureToUserMaterials.FindOrAdd(Texture).Add(MID);
 					}
 				}
 			}
 		}
 	}
-}
+}	 // namespace UE::UsdShadeTranslator::Private
 
 void FUsdShadeMaterialTranslator::CreateAssets()
 {
-	TRACE_CPUPROFILER_EVENT_SCOPE( FUsdShadeMaterialTranslator::CreateAssets );
+	TRACE_CPUPROFILER_EVENT_SCOPE(FUsdShadeMaterialTranslator::CreateAssets);
 
-	pxr::UsdShadeMaterial ShadeMaterial( GetPrim() );
+	pxr::UsdShadeMaterial ShadeMaterial(GetPrim());
 
-	if ( !ShadeMaterial )
+	if (!ShadeMaterial)
 	{
 		return;
 	}
@@ -245,15 +238,18 @@ void FUsdShadeMaterialTranslator::CreateAssets()
 	{
 		if (!Context->InfoCache->IsMaterialUsed(PrimPath))
 		{
-			UE_LOG(LogUsd, Verbose, TEXT("Skipping creating assets for material prim '%s' as it is not currently bound by any prim."), *PrimPath.GetString());
+			UE_LOG(
+				LogUsd,
+				Verbose,
+				TEXT("Skipping creating assets for material prim '%s' as it is not currently bound by any prim."),
+				*PrimPath.GetString()
+			);
 			return;
 		}
 	}
 
-	const pxr::TfToken RenderContextToken =
-		Context->RenderContext.IsNone() ?
-			pxr::UsdShadeTokens->universalRenderContext :
-			UnrealToUsd::ConvertToken( *Context->RenderContext.ToString() ).Get();
+	const pxr::TfToken RenderContextToken = Context->RenderContext.IsNone() ? pxr::UsdShadeTokens->universalRenderContext
+																			: UnrealToUsd::ConvertToken(*Context->RenderContext.ToString()).Get();
 
 	// If this material has a valid surface output for the 'unreal' render context and we're using it, don't bother
 	// generating any new UMaterialInterface asset because when resolving material assignments for this material
@@ -277,10 +273,10 @@ void FUsdShadeMaterialTranslator::CreateAssets()
 		ConvertedMaterial = Cast<UMaterialInterface>(Context->AssetCache->GetCachedAsset(PrefixedMaterialHash));
 	}
 
-	if ( !ConvertedMaterial )
+	if (!ConvertedMaterial)
 	{
 		const FString PrimPathString = PrimPath.GetString();
-		const bool bIsTranslucent = UsdUtils::IsMaterialTranslucent( ShadeMaterial );
+		const bool bIsTranslucent = UsdUtils::IsMaterialTranslucent(ShadeMaterial);
 		const FName InstanceName = MakeUniqueObjectName(
 			GetTransientPackage(),
 			UMaterialInstance::StaticClass(),
@@ -288,9 +284,13 @@ void FUsdShadeMaterialTranslator::CreateAssets()
 		);
 
 #if WITH_EDITOR
-		if ( GIsEditor ) // Also have to prevent Standalone game from going with MaterialInstanceConstants
+		if (GIsEditor)	  // Also have to prevent Standalone game from going with MaterialInstanceConstants
 		{
-			if ( UMaterialInstanceConstant* NewMaterial = NewObject<UMaterialInstanceConstant>( GetTransientPackage(), InstanceName, Context->ObjectFlags | EObjectFlags::RF_Transient ) )
+			if (UMaterialInstanceConstant* NewMaterial = NewObject<UMaterialInstanceConstant>(
+					GetTransientPackage(),
+					InstanceName,
+					Context->ObjectFlags | EObjectFlags::RF_Transient
+				))
 			{
 				UUsdMaterialAssetUserData* UserData = NewObject<UUsdMaterialAssetUserData>(NewMaterial, TEXT("USDAssetUserData"));
 				UserData->PrimPaths = {PrimPath.GetString()};
@@ -303,93 +303,96 @@ void FUsdShadeMaterialTranslator::CreateAssets()
 					*Context->RenderContext.ToString(),
 					Context->bReuseIdenticalAssets
 				);
-				if ( !bSuccess )
+				if (!bSuccess)
 				{
 					return;
 				}
 
 				TSet<UTexture*> VTTextures;
 				TSet<UTexture*> NonVTTextures;
-				for ( const FTextureParameterValue& TextureValue : NewMaterial->TextureParameterValues )
+				for (const FTextureParameterValue& TextureValue : NewMaterial->TextureParameterValues)
 				{
-					if ( UTexture* Texture = TextureValue.ParameterValue )
+					if (UTexture* Texture = TextureValue.ParameterValue)
 					{
-						if ( Texture->VirtualTextureStreaming )
+						if (Texture->VirtualTextureStreaming)
 						{
-							UsdUtils::NotifyIfVirtualTexturesNeeded( Texture );
-							VTTextures.Add( Texture );
+							UsdUtils::NotifyIfVirtualTexturesNeeded(Texture);
+							VTTextures.Add(Texture);
 						}
 						else
 						{
-							NonVTTextures.Add( Texture );
+							NonVTTextures.Add(Texture);
 						}
 					}
 				}
 
 				// Our VT material only has VT texture samplers, so *all* of its textures must be VT
-				if ( VTTextures.Num() && NonVTTextures.Num() )
+				if (VTTextures.Num() && NonVTTextures.Num())
 				{
-					UE_LOG( LogUsd, Log, TEXT( "Upgrading textures used by material instance '%s' to VT as the material must be VT" ),
+					UE_LOG(
+						LogUsd,
+						Log,
+						TEXT("Upgrading textures used by material instance '%s' to VT as the material must be VT"),
 						*NewMaterial->GetName()
 					);
-					UE::UsdShadeTranslator::Private::UpgradeMaterialsAndTexturesToVT( NonVTTextures, Context );
+					UE::UsdShadeTranslator::Private::UpgradeMaterialsAndTexturesToVT(NonVTTextures, Context);
 				}
 
 				MeshTranslationImpl::EUsdReferenceMaterialProperties Properties = MeshTranslationImpl::EUsdReferenceMaterialProperties::None;
-				if ( bIsTranslucent )
+				if (bIsTranslucent)
 				{
 					Properties |= MeshTranslationImpl::EUsdReferenceMaterialProperties::Translucent;
 				}
-				if ( VTTextures.Num() > 0 )
+				if (VTTextures.Num() > 0)
 				{
 					Properties |= MeshTranslationImpl::EUsdReferenceMaterialProperties::VT;
 				}
-				UMaterialInterface* ReferenceMaterial = MeshTranslationImpl::GetReferencePreviewSurfaceMaterial( Properties );
+				UMaterialInterface* ReferenceMaterial = MeshTranslationImpl::GetReferencePreviewSurfaceMaterial(Properties);
 
-				if ( ensure( ReferenceMaterial ) )
+				if (ensure(ReferenceMaterial))
 				{
-					NewMaterial->SetParentEditorOnly( ReferenceMaterial );
+					NewMaterial->SetParentEditorOnly(ReferenceMaterial);
 
 					// We can't blindly recreate all component render states when a level is being added, because we may end up first creating
 					// render states for some components, and UWorld::AddToWorld calls FScene::AddPrimitive which expects the component to not have
 					// primitives yet
 					FMaterialUpdateContext::EOptions::Type Options = FMaterialUpdateContext::EOptions::Default;
-					if ( Context->Level && Context->Level->bIsAssociatingLevel )
+					if (Context->Level && Context->Level->bIsAssociatingLevel)
 					{
-						Options = ( FMaterialUpdateContext::EOptions::Type ) ( Options & ~FMaterialUpdateContext::EOptions::RecreateRenderStates );
+						Options = (FMaterialUpdateContext::EOptions::Type)(Options & ~FMaterialUpdateContext::EOptions::RecreateRenderStates);
 					}
 
-					FMaterialUpdateContext UpdateContext( Options, GMaxRHIShaderPlatform );
-					UpdateContext.AddMaterialInstance( NewMaterial );
-					NewMaterial->PreEditChange( nullptr );
+					FMaterialUpdateContext UpdateContext(Options, GMaxRHIShaderPlatform);
+					UpdateContext.AddMaterialInstance(NewMaterial);
+					NewMaterial->PreEditChange(nullptr);
 					NewMaterial->PostEditChange();
 
 					ConvertedMaterial = NewMaterial;
 
-					for ( UTexture* Texture : VTTextures.Union( NonVTTextures ) )
+					for (UTexture* Texture : VTTextures.Union(NonVTTextures))
 					{
-						Context->TextureToUserMaterials.FindOrAdd( Texture ).Add( NewMaterial );
+						Context->TextureToUserMaterials.FindOrAdd(Texture).Add(NewMaterial);
 					}
 				}
 			}
 		}
 		else
-#endif // WITH_EDITOR
+#endif	  // WITH_EDITOR
 		{
 			// At runtime we always start with a non-VT reference and if we discover we need one we just create a new MID
 			// using the VT reference and copy the overrides. Not much else we can do as we need a reference to call
 			// UMaterialInstanceDynamic::Create and get our instance, but we an instance to call UsdToUnreal::ConvertMaterial
 			// to create our textures and decide on our reference.
 			MeshTranslationImpl::EUsdReferenceMaterialProperties Properties = MeshTranslationImpl::EUsdReferenceMaterialProperties::None;
-			if ( bIsTranslucent )
+			if (bIsTranslucent)
 			{
 				Properties |= MeshTranslationImpl::EUsdReferenceMaterialProperties::Translucent;
 			}
-			UMaterialInterface* ReferenceMaterial = MeshTranslationImpl::GetReferencePreviewSurfaceMaterial( Properties );
+			UMaterialInterface* ReferenceMaterial = MeshTranslationImpl::GetReferencePreviewSurfaceMaterial(Properties);
 
-			if ( ensure( ReferenceMaterial ) )
+			if (ensure(ReferenceMaterial))
 			{
-				if ( UMaterialInstanceDynamic* NewMaterial = UMaterialInstanceDynamic::Create( ReferenceMaterial, GetTransientPackage(), InstanceName ) )
+				if (UMaterialInstanceDynamic* NewMaterial = UMaterialInstanceDynamic::Create(ReferenceMaterial, GetTransientPackage(), InstanceName))
 				{
 					UUsdMaterialAssetUserData* UserData = NewObject<UUsdMaterialAssetUserData>(NewMaterial, TEXT("USDAssetUserData"));
 					UserData->PrimPaths = {PrimPath.GetString()};
@@ -407,40 +410,43 @@ void FUsdShadeMaterialTranslator::CreateAssets()
 					{
 						TSet<UTexture*> VTTextures;
 						TSet<UTexture*> NonVTTextures;
-						for ( const FTextureParameterValue& TextureValue : NewMaterial->TextureParameterValues )
+						for (const FTextureParameterValue& TextureValue : NewMaterial->TextureParameterValues)
 						{
-							if ( UTexture* Texture = TextureValue.ParameterValue )
+							if (UTexture* Texture = TextureValue.ParameterValue)
 							{
-								if ( Texture->VirtualTextureStreaming )
+								if (Texture->VirtualTextureStreaming)
 								{
-									VTTextures.Add( Texture );
+									VTTextures.Add(Texture);
 								}
 								else
 								{
-									NonVTTextures.Add( Texture );
+									NonVTTextures.Add(Texture);
 								}
 							}
 						}
 
 						// We must stash our material and textures *before* we call UpgradeMaterialsAndTexturesToVT, as that
 						// is what will actually swap our reference with a VT one if needed
-						if ( Context->AssetCache && Context->InfoCache )
+						if (Context->AssetCache && Context->InfoCache)
 						{
 							Context->AssetCache->CacheAsset(PrefixedMaterialHash, NewMaterial);
 							Context->InfoCache->LinkAssetToPrim(PrimPath, NewMaterial);
 						}
-						for ( UTexture* Texture : VTTextures.Union( NonVTTextures ) )
+						for (UTexture* Texture : VTTextures.Union(NonVTTextures))
 						{
-							Context->TextureToUserMaterials.FindOrAdd( Texture ).Add( NewMaterial );
+							Context->TextureToUserMaterials.FindOrAdd(Texture).Add(NewMaterial);
 						}
 
 						// Our VT material only has VT texture samplers, so *all* of its textures must be VT
-						if ( VTTextures.Num() && NonVTTextures.Num() )
+						if (VTTextures.Num() && NonVTTextures.Num())
 						{
-							UE_LOG( LogUsd, Log, TEXT( "Upgrading textures used by material instance '%s' to VT as the material must be VT" ),
+							UE_LOG(
+								LogUsd,
+								Log,
+								TEXT("Upgrading textures used by material instance '%s' to VT as the material must be VT"),
 								*NewMaterial->GetName()
 							);
-							UE::UsdShadeTranslator::Private::UpgradeMaterialsAndTexturesToVT( NonVTTextures, Context );
+							UE::UsdShadeTranslator::Private::UpgradeMaterialsAndTexturesToVT(NonVTTextures, Context);
 						}
 
 						// We must go through the cache to fetch our result material here as UpgradeMaterialsAndTexturesToVT
@@ -455,12 +461,12 @@ void FUsdShadeMaterialTranslator::CreateAssets()
 	PostImportMaterial(PrefixedMaterialHash, ConvertedMaterial);
 }
 
-bool FUsdShadeMaterialTranslator::CollapsesChildren( ECollapsingType CollapsingType ) const
+bool FUsdShadeMaterialTranslator::CollapsesChildren(ECollapsingType CollapsingType) const
 {
 	return false;
 }
 
-bool FUsdShadeMaterialTranslator::CanBeCollapsed( ECollapsingType CollapsingType ) const
+bool FUsdShadeMaterialTranslator::CanBeCollapsed(ECollapsingType CollapsingType) const
 {
 	return false;
 }
@@ -517,24 +523,25 @@ void FUsdShadeMaterialTranslator::PostImportMaterial(const FString& PrefixedMate
 			{
 				FString FilePath;
 #if WITH_EDITOR
-				if(UAssetImportData* TextureImportData = Texture->AssetImportData.Get())
+				if (UAssetImportData* TextureImportData = Texture->AssetImportData.Get())
 				{
 					FilePath = TextureImportData->GetFirstFilename();
 				}
 				else
-#endif // WITH_EDITOR
+#endif	  // WITH_EDITOR
 				{
 					FilePath = Texture->GetName();
 				}
 
 				const FString HashPrefix = UsdUtils::GetAssetHashPrefix(GetPrim(), Context->bReuseIdenticalAssets);
-				const FString PrefixedTextureHash = HashPrefix + UsdUtils::GetTextureHash(
-					FilePath,
-					Texture->SRGB,
-					Texture->CompressionSettings,
-					Texture->GetTextureAddressX(),
-					Texture->GetTextureAddressY()
-				);
+				const FString PrefixedTextureHash = HashPrefix
+													+ UsdUtils::GetTextureHash(
+														FilePath,
+														Texture->SRGB,
+														Texture->CompressionSettings,
+														Texture->GetTextureAddressX(),
+														Texture->GetTextureAddressY()
+													);
 
 				// Some translators like FMaterialXUsdShadeMaterialTranslator will import many materials and textures
 				// at once and preload them all in the cache. In some complex scenarios when some materials are updated
@@ -625,9 +632,8 @@ TSet<UE::FSdfPath> FUsdShadeMaterialTranslator::CollectAuxiliaryPrims() const
 			return {};
 		}
 
-		const pxr::TfToken RenderContextToken = Context->RenderContext.IsNone()
-			? pxr::UsdShadeTokens->universalRenderContext
-			: UnrealToUsd::ConvertToken(*Context->RenderContext.ToString()).Get();
+		const pxr::TfToken RenderContextToken = Context->RenderContext.IsNone() ? pxr::UsdShadeTokens->universalRenderContext
+																				: UnrealToUsd::ConvertToken(*Context->RenderContext.ToString()).Get();
 		pxr::UsdShadeShader SurfaceShader = UsdShadeMaterial.ComputeSurfaceSource({RenderContextToken});
 		if (!SurfaceShader)
 		{
@@ -644,4 +650,4 @@ TSet<UE::FSdfPath> FUsdShadeMaterialTranslator::CollectAuxiliaryPrims() const
 	return Result;
 }
 
-#endif // #if USE_USD_SDK
+#endif	  // #if USE_USD_SDK
