USD: Reformat USDSchemas module.

diff --git a/Engine/Plugins/Importers/USDImporter/Source/USDSchemas/Private/USDSkelRootTranslator.cpp b/Engine/Plugins/Importers/USDImporter/Source/USDSchemas/Private/USDSkelRootTranslator.cpp
--- a/Engine/Plugins/Importers/USDImporter/Source/USDSchemas/Private/USDSkelRootTranslator.cpp
+++ b/Engine/Plugins/Importers/USDImporter/Source/USDSchemas/Private/USDSkelRootTranslator.cpp
@@ -5,11 +5,10 @@
 #if USE_USD_SDK
 
 #include "MeshTranslationImpl.h"
-#include "USDAssetCache.h"
 #include "USDAssetUserData.h"
-#include "USDDrawModeComponent.h"
 #include "USDClassesModule.h"
 #include "USDConversionUtils.h"
+#include "USDDrawModeComponent.h"
 #include "USDErrorUtils.h"
 #include "USDGeomMeshConversion.h"
 #include "USDGroomTranslatorUtils.h"
@@ -20,7 +19,6 @@
 #include "USDMemory.h"
 #include "USDSkeletalDataConversion.h"
 #include "USDTypesConversion.h"
-#include "USDValueConversion.h"
 
 #include "UsdWrappers/SdfLayer.h"
 #include "UsdWrappers/SdfPath.h"
@@ -30,7 +28,6 @@
 #include "Animation/AnimInstance.h"
 #include "Animation/AnimSequence.h"
 #include "Animation/Skeleton.h"
-#include "AnimationUtils.h"
 #include "Components/SkeletalMeshComponent.h"
 #include "Components/SkinnedMeshComponent.h"
 #include "Engine/SkeletalMesh.h"
@@ -38,35 +35,29 @@
 #include "GroomComponent.h"
 #include "MaterialDomain.h"
 #include "Materials/Material.h"
-#include "Materials/MaterialInstanceConstant.h"
 #include "Rendering/SkeletalMeshLODImporterData.h"
-#include "Serialization/BufferArchive.h"
 
 #if WITH_EDITOR
-	#include "AnimGraphNode_LiveLinkPose.h"
-	#include "AnimNode_LiveLinkPose.h"
-	#include "BlueprintCompilationManager.h"
-	#include "EdGraph/EdGraph.h"
-	#include "EdGraph/EdGraphNode.h"
-	#include "K2Node_VariableGet.h"
-	#include "Kismet2/BlueprintEditorUtils.h"
-	#include "Kismet2/CompilerResultsLog.h"
-	#include "MaterialEditingLibrary.h"
-	#include "ObjectTools.h"
-	#include "PhysicsAssetUtils.h"
-#endif // WITH_EDITOR
+#include "AnimGraphNode_LiveLinkPose.h"
+#include "AnimNode_LiveLinkPose.h"
+#include "BlueprintCompilationManager.h"
+#include "EdGraph/EdGraph.h"
+#include "EdGraph/EdGraphNode.h"
+#include "Kismet2/BlueprintEditorUtils.h"
+#include "Kismet2/CompilerResultsLog.h"
+#include "ObjectTools.h"
+#include "PhysicsAssetUtils.h"
+#endif	  // WITH_EDITOR
 
 #include "USDIncludesStart.h"
-	#include "pxr/usd/usd/primRange.h"
-	#include "pxr/usd/usdGeom/mesh.h"
-	#include "pxr/usd/usdShade/material.h"
-	#include "pxr/usd/usdSkel/binding.h"
-	#include "pxr/usd/usdSkel/bindingAPI.h"
-	#include "pxr/usd/usdSkel/blendShapeQuery.h"
-	#include "pxr/usd/usdSkel/cache.h"
-	#include "pxr/usd/usdSkel/root.h"
-	#include "pxr/usd/usdSkel/skeletonQuery.h"
-	#include "pxr/usd/usdSkel/tokens.h"
+#include "pxr/usd/usd/primRange.h"
+#include "pxr/usd/usdGeom/mesh.h"
+#include "pxr/usd/usdSkel/binding.h"
+#include "pxr/usd/usdSkel/bindingAPI.h"
+#include "pxr/usd/usdSkel/blendShapeQuery.h"
+#include "pxr/usd/usdSkel/cache.h"
+#include "pxr/usd/usdSkel/root.h"
+#include "pxr/usd/usdSkel/skeletonQuery.h"
 #include "USDIncludesEnd.h"
 
 #define LOCTEXT_NAMESPACE "UsdSkelRoot"
@@ -86,21 +77,23 @@ namespace UsdSkelRootTranslatorImpl
 		bool bReuseIdenticalAssets
 	)
 	{
-		TRACE_CPUPROFILER_EVENT_SCOPE( UsdSkelRootTranslatorImpl::ProcessMaterials );
+		TRACE_CPUPROFILER_EVENT_SCOPE(UsdSkelRootTranslatorImpl::ProcessMaterials);
 
-		if ( !SkeletalMesh )
+		if (!SkeletalMesh)
 		{
 			return false;
 		}
 
 		TArray<UMaterialInterface*> ExistingAssignments;
-		for ( const FSkeletalMaterial& SkeletalMaterial : SkeletalMesh->GetMaterials())
+		for (const FSkeletalMaterial& SkeletalMaterial : SkeletalMesh->GetMaterials())
 		{
-			ExistingAssignments.Add( SkeletalMaterial.MaterialInterface );
+			ExistingAssignments.Add(SkeletalMaterial.MaterialInterface);
 		}
 
 		UUsdMeshAssetUserData* UserData = SkeletalMesh->GetAssetUserData<UUsdMeshAssetUserData>();
-		ensureMsgf(UserData, TEXT("Skeletal Mesh '%s' generated for prim '%s' should have an UUsdMeshAssetUserData at this point!"),
+		ensureMsgf(
+			UserData,
+			TEXT("Skeletal Mesh '%s' generated for prim '%s' should have an UUsdMeshAssetUserData at this point!"),
 			*SkeletalMesh->GetPathName(),
 			*UsdToUnreal::ConvertPath(UsdPrim.GetPrimPath())
 		);
@@ -117,56 +110,69 @@ namespace UsdSkelRootTranslatorImpl
 		bool bMaterialsHaveChanged = false;
 
 		uint32 SkeletalMeshSlotIndex = 0;
-		for ( int32 LODIndex = 0; LODIndex < LODIndexToMaterialInfo.Num(); ++LODIndex )
+		for (int32 LODIndex = 0; LODIndex < LODIndexToMaterialInfo.Num(); ++LODIndex)
 		{
-			const TArray< UsdUtils::FUsdPrimMaterialSlot >& LODSlots = LODIndexToMaterialInfo[ LODIndex ].Slots;
+			const TArray<UsdUtils::FUsdPrimMaterialSlot>& LODSlots = LODIndexToMaterialInfo[LODIndex].Slots;
 
 			// We need to fill this in with the mapping from LOD material slots (i.e. sections) to the skeletal mesh's material slots
-			FSkeletalMeshLODInfo* LODInfo = SkeletalMesh->GetLODInfo( LODIndex );
-			if ( !LODInfo )
+			FSkeletalMeshLODInfo* LODInfo = SkeletalMesh->GetLODInfo(LODIndex);
+			if (!LODInfo)
 			{
-				UE_LOG( LogUsd, Error, TEXT( "When processing materials for SkeletalMesh '%s', encountered no LOD info for LOD index %d!" ), *SkeletalMesh->GetName(), LODIndex );
+				UE_LOG(
+					LogUsd,
+					Error,
+					TEXT("When processing materials for SkeletalMesh '%s', encountered no LOD info for LOD index %d!"),
+					*SkeletalMesh->GetName(),
+					LODIndex
+				);
 				continue;
 			}
 			TArray<int32>& LODMaterialMap = LODInfo->LODMaterialMap;
 			LODMaterialMap.Reserve(LODSlots.Num());
 
-			for ( int32 LODSlotIndex = 0; LODSlotIndex < LODSlots.Num(); ++LODSlotIndex, ++SkeletalMeshSlotIndex )
+			for (int32 LODSlotIndex = 0; LODSlotIndex < LODSlots.Num(); ++LODSlotIndex, ++SkeletalMeshSlotIndex)
 			{
-				const UsdUtils::FUsdPrimMaterialSlot& Slot = LODSlots[ LODSlotIndex ];
+				const UsdUtils::FUsdPrimMaterialSlot& Slot = LODSlots[LODSlotIndex];
 
 				UMaterialInterface* Material = UMaterial::GetDefaultMaterial(MD_Surface);
 
-				if ( UMaterialInterface** FoundMaterial = ResolvedMaterials.Find( &Slot ) )
+				if (UMaterialInterface** FoundMaterial = ResolvedMaterials.Find(&Slot))
 				{
 					Material = *FoundMaterial;
 				}
 				else
 				{
-					UE_LOG( LogUsd, Error, TEXT( "Failed to resolve material '%s' for slot '%d' of LOD '%d' for mesh '%s'" ), *Slot.MaterialSource, LODSlotIndex, LODIndex, *UsdToUnreal::ConvertPath( UsdPrim.GetPath() ) );
+					UE_LOG(
+						LogUsd,
+						Error,
+						TEXT("Failed to resolve material '%s' for slot '%d' of LOD '%d' for mesh '%s'"),
+						*Slot.MaterialSource,
+						LODSlotIndex,
+						LODIndex,
+						*UsdToUnreal::ConvertPath(UsdPrim.GetPath())
+					);
 					continue;
 				}
 
-				if ( Material )
+				if (Material)
 				{
 					bool bNeedsRecompile = false;
-					Material->GetMaterial()->SetMaterialUsage( bNeedsRecompile, MATUSAGE_SkeletalMesh );
-					if ( bSkeletalMeshHasMorphTargets )
+					Material->GetMaterial()->SetMaterialUsage(bNeedsRecompile, MATUSAGE_SkeletalMesh);
+					if (bSkeletalMeshHasMorphTargets)
 					{
-						Material->GetMaterial()->SetMaterialUsage( bNeedsRecompile, MATUSAGE_MorphTargets );
+						Material->GetMaterial()->SetMaterialUsage(bNeedsRecompile, MATUSAGE_MorphTargets);
 					}
 				}
 
-				FName MaterialSlotName = *LexToString( SkeletalMeshSlotIndex );
+				FName MaterialSlotName = *LexToString(SkeletalMeshSlotIndex);
 
 				// Already have a material at that skeletal mesh slot, need to reassign
-				if ( SkeletalMesh->GetMaterials().IsValidIndex( SkeletalMeshSlotIndex ) )
+				if (SkeletalMesh->GetMaterials().IsValidIndex(SkeletalMeshSlotIndex))
 				{
-					FSkeletalMaterial& ExistingMaterial = SkeletalMesh->GetMaterials()[ SkeletalMeshSlotIndex ];
+					FSkeletalMaterial& ExistingMaterial = SkeletalMesh->GetMaterials()[SkeletalMeshSlotIndex];
 
-					if ( ExistingMaterial.MaterialInterface != Material ||
-						 ExistingMaterial.MaterialSlotName != MaterialSlotName ||
-						 ExistingMaterial.ImportedMaterialSlotName != MaterialSlotName )
+					if (ExistingMaterial.MaterialInterface != Material || ExistingMaterial.MaterialSlotName != MaterialSlotName
+						|| ExistingMaterial.ImportedMaterialSlotName != MaterialSlotName)
 					{
 						ExistingMaterial.MaterialInterface = Material;
 						ExistingMaterial.MaterialSlotName = MaterialSlotName;
@@ -179,19 +185,21 @@ namespace UsdSkelRootTranslatorImpl
 				{
 					const bool bEnableShadowCasting = true;
 					const bool bRecomputeTangents = false;
-					SkeletalMesh->GetMaterials().Add( FSkeletalMaterial( Material, bEnableShadowCasting, bRecomputeTangents, MaterialSlotName, MaterialSlotName ) );
+					SkeletalMesh->GetMaterials().Add(
+						FSkeletalMaterial(Material, bEnableShadowCasting, bRecomputeTangents, MaterialSlotName, MaterialSlotName)
+					);
 					bMaterialsHaveChanged = true;
 				}
 
 				// Already have a material at that LOD remap slot, need to reassign
-				if ( LODMaterialMap.IsValidIndex( LODSlotIndex ) )
+				if (LODMaterialMap.IsValidIndex(LODSlotIndex))
 				{
-					LODMaterialMap[ LODSlotIndex ] = SkeletalMeshSlotIndex;
+					LODMaterialMap[LODSlotIndex] = SkeletalMeshSlotIndex;
 				}
 				// Add new material slot remap
 				else
 				{
-					LODMaterialMap.Add( SkeletalMeshSlotIndex );
+					LODMaterialMap.Add(SkeletalMeshSlotIndex);
 				}
 			}
 		}
@@ -207,23 +215,24 @@ namespace UsdSkelRootTranslatorImpl
 	{
 		FSHA1 HashState;
 
-		for ( const FSkeletalMeshImportData& ImportData : LODIndexToSkeletalMeshImportData )
+		for (const FSkeletalMeshImportData& ImportData : LODIndexToSkeletalMeshImportData)
 		{
-			HashState.Update( ( uint8* ) ImportData.Points.GetData(), ImportData.Points.Num() * ImportData.Points.GetTypeSize() );
-			HashState.Update( ( uint8* ) ImportData.Wedges.GetData(), ImportData.Wedges.Num() * ImportData.Wedges.GetTypeSize() );
-			HashState.Update( ( uint8* ) ImportData.Faces.GetData(), ImportData.Faces.Num() * ImportData.Faces.GetTypeSize() );
-			HashState.Update( ( uint8* ) ImportData.Influences.GetData(), ImportData.Influences.Num() * ImportData.Influences.GetTypeSize() );
+			HashState.Update((uint8*)ImportData.Points.GetData(), ImportData.Points.Num() * ImportData.Points.GetTypeSize());
+			HashState.Update((uint8*)ImportData.Wedges.GetData(), ImportData.Wedges.Num() * ImportData.Wedges.GetTypeSize());
+			HashState.Update((uint8*)ImportData.Faces.GetData(), ImportData.Faces.Num() * ImportData.Faces.GetTypeSize());
+			HashState.Update((uint8*)ImportData.Influences.GetData(), ImportData.Influences.Num() * ImportData.Influences.GetTypeSize());
 		}
 
-		// Hash the bones as well because it is possible for the mesh to be identical while only the bone configuration changed, and in that case we'd need new skeleton and ref skeleton
-		// Maybe in the future (as a separate feature) we could split off the skeleton import so that it could vary independently of the skeletal mesh
-		for ( const SkeletalMeshImportData::FBone& Bone : ImportedBones )
+		// Hash the bones as well because it is possible for the mesh to be identical while only the bone configuration changed, and in that case we'd
+		// need new skeleton and ref skeleton Maybe in the future (as a separate feature) we could split off the skeleton import so that it could vary
+		// independently of the skeletal mesh
+		for (const SkeletalMeshImportData::FBone& Bone : ImportedBones)
 		{
-			HashState.UpdateWithString( *Bone.Name, Bone.Name.Len() );
-			HashState.Update( ( uint8* ) &Bone.Flags, sizeof( Bone.Flags ) );
-			HashState.Update( ( uint8* ) &Bone.NumChildren, sizeof( Bone.NumChildren ) );
-			HashState.Update( ( uint8* ) &Bone.ParentIndex, sizeof( Bone.ParentIndex ) );
-			HashState.Update( ( uint8* ) &Bone.BonePos, sizeof( Bone.BonePos ) );
+			HashState.UpdateWithString(*Bone.Name, Bone.Name.Len());
+			HashState.Update((uint8*)&Bone.Flags, sizeof(Bone.Flags));
+			HashState.Update((uint8*)&Bone.NumChildren, sizeof(Bone.NumChildren));
+			HashState.Update((uint8*)&Bone.ParentIndex, sizeof(Bone.ParentIndex));
+			HashState.Update((uint8*)&Bone.BonePos, sizeof(Bone.BonePos));
 		}
 
 		if (BlendShapes)
@@ -248,7 +257,7 @@ namespace UsdSkelRootTranslatorImpl
 		FSHAHash OutHash;
 
 		HashState.Final();
-		HashState.GetHash( &OutHash.Hash[0] );
+		HashState.GetHash(&OutHash.Hash[0]);
 
 		return OutHash;
 	}
@@ -259,7 +268,7 @@ namespace UsdSkelRootTranslatorImpl
 		const FString& SkeletalMeshHashString
 	)
 	{
-		TRACE_CPUPROFILER_EVENT_SCOPE( UsdSkelRootTranslatorImpl::ComputeSHAHash_SkelQuery );
+		TRACE_CPUPROFILER_EVENT_SCOPE(UsdSkelRootTranslatorImpl::ComputeSHAHash_SkelQuery);
 
 		FSHAHash OutHash;
 		FSHA1 HashState;
@@ -267,25 +276,25 @@ namespace UsdSkelRootTranslatorImpl
 		FScopedUsdAllocs Allocs;
 
 		pxr::UsdSkelAnimQuery AnimQuery = InUsdSkeletonQuery.GetAnimQuery();
-		if ( !AnimQuery )
+		if (!AnimQuery)
 		{
 			return OutHash;
 		}
 
 		pxr::UsdPrim UsdPrim = InUsdSkeletonQuery.GetPrim();
-		if ( !UsdPrim )
+		if (!UsdPrim)
 		{
 			return OutHash;
 		}
 
 		pxr::UsdStageRefPtr Stage = UsdPrim.GetStage();
-		if ( !Stage )
+		if (!Stage)
 		{
 			return OutHash;
 		}
 
-		int32 InterpolationType = static_cast< int32 >( Stage->GetInterpolationType() );
-		HashState.Update( ( uint8* ) &InterpolationType, sizeof( int32 ) );
+		int32 InterpolationType = static_cast<int32>(Stage->GetInterpolationType());
+		HashState.Update((uint8*)&InterpolationType, sizeof(int32));
 
 		// Hash blend shape and joint order tokens
 		TFunction<void(const pxr::VtArray<pxr::TfToken>&)> HashTokens = [&HashState](const pxr::VtArray<pxr::TfToken>& Tokens)
@@ -302,55 +311,55 @@ namespace UsdSkelRootTranslatorImpl
 
 		// Time samples for joint transforms
 		std::vector<double> TimeData;
-		AnimQuery.GetJointTransformTimeSamples( &TimeData );
-		HashState.Update( ( uint8* ) TimeData.data(), TimeData.size() * sizeof( double ) );
+		AnimQuery.GetJointTransformTimeSamples(&TimeData);
+		HashState.Update((uint8*)TimeData.data(), TimeData.size() * sizeof(double));
 
 		// Joint transform values
 		pxr::VtArray<pxr::GfMatrix4d> JointTransforms;
-		for ( double JointTimeSample : TimeData )
+		for (double JointTimeSample : TimeData)
 		{
-			InUsdSkeletonQuery.ComputeJointLocalTransforms( &JointTransforms, JointTimeSample );
-			HashState.Update( ( uint8* ) JointTransforms.data(), JointTransforms.size() * sizeof( pxr::GfMatrix4d ) );
+			InUsdSkeletonQuery.ComputeJointLocalTransforms(&JointTransforms, JointTimeSample);
+			HashState.Update((uint8*)JointTransforms.data(), JointTransforms.size() * sizeof(pxr::GfMatrix4d));
 		}
 
 		// restTransforms
 		pxr::VtArray<pxr::GfMatrix4d> Transforms;
 		const bool bAtRest = true;
-		InUsdSkeletonQuery.ComputeJointLocalTransforms( &Transforms, pxr::UsdTimeCode::EarliestTime(), bAtRest );
-		HashState.Update( ( uint8* ) Transforms.data(), Transforms.size() * sizeof( pxr::GfMatrix4d ) );
+		InUsdSkeletonQuery.ComputeJointLocalTransforms(&Transforms, pxr::UsdTimeCode::EarliestTime(), bAtRest);
+		HashState.Update((uint8*)Transforms.data(), Transforms.size() * sizeof(pxr::GfMatrix4d));
 
 		// bindTransforms
-		InUsdSkeletonQuery.GetJointWorldBindTransforms( &Transforms );
-		HashState.Update( ( uint8* ) Transforms.data(), Transforms.size() * sizeof( pxr::GfMatrix4d ) );
+		InUsdSkeletonQuery.GetJointWorldBindTransforms(&Transforms);
+		HashState.Update((uint8*)Transforms.data(), Transforms.size() * sizeof(pxr::GfMatrix4d));
 
 		// Time samples for blend shape curves
-		AnimQuery.GetBlendShapeWeightTimeSamples( &TimeData );
-		HashState.Update( ( uint8* ) TimeData.data(), TimeData.size() * sizeof( double ) );
+		AnimQuery.GetBlendShapeWeightTimeSamples(&TimeData);
+		HashState.Update((uint8*)TimeData.data(), TimeData.size() * sizeof(double));
 
 		// Blend shape curve values
-		pxr::VtArray< float > WeightsForSample;
-		for ( double CurveTimeSample : TimeData )
+		pxr::VtArray<float> WeightsForSample;
+		for (double CurveTimeSample : TimeData)
 		{
-			AnimQuery.ComputeBlendShapeWeights( &WeightsForSample, pxr::UsdTimeCode( CurveTimeSample ) );
-			HashState.Update( ( uint8* ) WeightsForSample.data(), WeightsForSample.size() * sizeof( float ) );
+			AnimQuery.ComputeBlendShapeWeights(&WeightsForSample, pxr::UsdTimeCode(CurveTimeSample));
+			HashState.Update((uint8*)WeightsForSample.data(), WeightsForSample.size() * sizeof(float));
 		}
 
 		// If we're pulling root motion from anywhere, hash that too because if it changes we'll need to rebake
 		// the AnimSequence asset
-		if ( pxr::UsdGeomXformable Xformable{ RootMotionPrim } )
+		if (pxr::UsdGeomXformable Xformable{RootMotionPrim})
 		{
 			std::vector<double> TimeSamples;
-			Xformable.GetTimeSamples( &TimeSamples );
+			Xformable.GetTimeSamples(&TimeSamples);
 
 			pxr::GfMatrix4d Transform{};
 			bool bResetsXformSack = false;
-			for ( double TimeSample : TimeSamples )
+			for (double TimeSample : TimeSamples)
 			{
-				Xformable.GetLocalTransformation( &Transform, &bResetsXformSack, TimeSample );
-				HashState.Update( ( uint8* ) Transform.data(), 16 * sizeof( double ) );
+				Xformable.GetLocalTransformation(&Transform, &bResetsXformSack, TimeSample);
+				HashState.Update((uint8*)Transform.data(), 16 * sizeof(double));
 			}
 
-			HashState.Update( ( uint8* ) &bResetsXformSack, sizeof( bool ) );
+			HashState.Update((uint8*)&bResetsXformSack, sizeof(bool));
 		}
 
 		// An anim sequence matches a particular skeleton. If the skeleton is different, we'll likely
@@ -362,52 +371,52 @@ namespace UsdSkelRootTranslatorImpl
 		HashState.UpdateWithString(*SkeletalMeshHashString, SkeletalMeshHashString.Len());
 
 		HashState.Final();
-		HashState.GetHash( &OutHash.Hash[ 0 ] );
+		HashState.GetHash(&OutHash.Hash[0]);
 
 		return OutHash;
 	}
 
-	void SetMorphTargetWeight( USkeletalMeshComponent& SkeletalMeshComponent, const FString& MorphTargetName, float Weight )
+	void SetMorphTargetWeight(USkeletalMeshComponent& SkeletalMeshComponent, const FString& MorphTargetName, float Weight)
 	{
 		USkeletalMesh* SkeletalMesh = SkeletalMeshComponent.GetSkeletalMeshAsset();
 
 		// We try keeping a perfect correspondence between SkeletalMesh->GetMorphTargets() and SkeletalMeshComponent.ActiveMorphTargets
 		int32 IndexInSkeletalMesh = INDEX_NONE;
-		SkeletalMeshComponent.GetSkeletalMeshAsset()->FindMorphTargetAndIndex( *MorphTargetName, IndexInSkeletalMesh );
-		if ( IndexInSkeletalMesh == INDEX_NONE )
+		SkeletalMeshComponent.GetSkeletalMeshAsset()->FindMorphTargetAndIndex(*MorphTargetName, IndexInSkeletalMesh);
+		if (IndexInSkeletalMesh == INDEX_NONE)
 		{
 			return;
 		}
 
-		const UMorphTarget* MorphTarget = SkeletalMeshComponent.GetSkeletalMeshAsset()->GetMorphTargets()[ IndexInSkeletalMesh ];
-		if ( !MorphTarget )
+		const UMorphTarget* MorphTarget = SkeletalMeshComponent.GetSkeletalMeshAsset()->GetMorphTargets()[IndexInSkeletalMesh];
+		if (!MorphTarget)
 		{
 			return;
 		}
 
 		int32 WeightIndex = INDEX_NONE;
-		if ( SkeletalMeshComponent.ActiveMorphTargets.Contains( MorphTarget ) )
+		if (SkeletalMeshComponent.ActiveMorphTargets.Contains(MorphTarget))
 		{
-			WeightIndex = SkeletalMeshComponent.ActiveMorphTargets[ MorphTarget ];
+			WeightIndex = SkeletalMeshComponent.ActiveMorphTargets[MorphTarget];
 		}
 
 		// Morph target is not at expected location (i.e. after CreateComponents, duplicate for PIE or undo/redo) --> Rebuild ActiveMorphTargets
 		// This may lead to one frame of glitchiness, as we'll reset all weights to zero...
-		if ( WeightIndex == INDEX_NONE )
+		if (WeightIndex == INDEX_NONE)
 		{
 			SkeletalMeshComponent.ActiveMorphTargets.Reset();
 			SkeletalMeshComponent.MorphTargetWeights.Reset();
 			TArray<TObjectPtr<UMorphTarget>>& MorphTargets = SkeletalMesh->GetMorphTargets();
-			for ( int32 MorphTargetIndex = 0; MorphTargetIndex < MorphTargets.Num(); ++MorphTargetIndex )
+			for (int32 MorphTargetIndex = 0; MorphTargetIndex < MorphTargets.Num(); ++MorphTargetIndex)
 			{
-				SkeletalMeshComponent.ActiveMorphTargets.Add( MorphTargets[ MorphTargetIndex ], MorphTargetIndex );
-				SkeletalMeshComponent.MorphTargetWeights.Add( 0.0f ); // We'll update these right afterwards when we call UpdateComponents
+				SkeletalMeshComponent.ActiveMorphTargets.Add(MorphTargets[MorphTargetIndex], MorphTargetIndex);
+				SkeletalMeshComponent.MorphTargetWeights.Add(0.0f);	   // We'll update these right afterwards when we call UpdateComponents
 			}
 
 			WeightIndex = IndexInSkeletalMesh;
 		}
 
-		SkeletalMeshComponent.MorphTargetWeights[ WeightIndex ] = Weight;
+		SkeletalMeshComponent.MorphTargetWeights[WeightIndex] = Weight;
 	}
 
 	bool LoadAllSkeletalData(
@@ -424,7 +433,7 @@ namespace UsdSkelRootTranslatorImpl
 		const UsdToUnreal::FUsdMeshConversionOptions& Options
 	)
 	{
-		if ( !InSkeletonRoot )
+		if (!InSkeletonRoot)
 		{
 			return false;
 		}
@@ -432,16 +441,20 @@ namespace UsdSkelRootTranslatorImpl
 		FScopedUsdAllocs UsdAllocs;
 
 		pxr::UsdStageRefPtr Stage = InSkeletonRoot.GetPrim().GetStage();
-		const FUsdStageInfo StageInfo( Stage );
-
-		std::vector< pxr::UsdSkelBinding > SkeletonBindings;
-		InSkeletonCache.Populate( InSkeletonRoot, pxr::UsdTraverseInstanceProxies() );
-		InSkeletonCache.ComputeSkelBindings( InSkeletonRoot, &SkeletonBindings, pxr::UsdTraverseInstanceProxies() );
-		if ( SkeletonBindings.size() < 1 )
-		{
-			FUsdLogManager::LogMessage( EMessageSeverity::Warning,
-										FText::Format( LOCTEXT("InvalidBinding", "SkelRoot {0} doesn't have any binding. No skinned mesh will be generated."),
-										FText::FromString( UsdToUnreal::ConvertPath( InSkeletonRoot.GetPath() ) ) ) );
+		const FUsdStageInfo StageInfo(Stage);
+
+		std::vector<pxr::UsdSkelBinding> SkeletonBindings;
+		InSkeletonCache.Populate(InSkeletonRoot, pxr::UsdTraverseInstanceProxies());
+		InSkeletonCache.ComputeSkelBindings(InSkeletonRoot, &SkeletonBindings, pxr::UsdTraverseInstanceProxies());
+		if (SkeletonBindings.size() < 1)
+		{
+			FUsdLogManager::LogMessage(
+				EMessageSeverity::Warning,
+				FText::Format(
+					LOCTEXT("InvalidBinding", "SkelRoot {0} doesn't have any binding. No skinned mesh will be generated."),
+					FText::FromString(UsdToUnreal::ConvertPath(InSkeletonRoot.GetPath()))
+				)
+			);
 			return false;
 		}
 
@@ -449,25 +462,28 @@ namespace UsdSkelRootTranslatorImpl
 		// For now, extract just the first one
 		const pxr::UsdSkelBinding& SkeletonBinding = SkeletonBindings[0];
 		const pxr::UsdSkelSkeleton& Skeleton = SkeletonBinding.GetSkeleton();
-		if ( SkeletonBindings.size() > 1 )
-		{
-			UE_LOG(LogUsd, Warning, TEXT("Currently only a single skeleton is supported per UsdSkelRoot! '%s' will use skeleton '%s'"),
-				*UsdToUnreal::ConvertPath( InSkeletonRoot.GetPrim().GetPath() ),
-				*UsdToUnreal::ConvertPath( Skeleton.GetPrim().GetPath() )
+		if (SkeletonBindings.size() > 1)
+		{
+			UE_LOG(
+				LogUsd,
+				Warning,
+				TEXT("Currently only a single skeleton is supported per UsdSkelRoot! '%s' will use skeleton '%s'"),
+				*UsdToUnreal::ConvertPath(InSkeletonRoot.GetPrim().GetPath()),
+				*UsdToUnreal::ConvertPath(Skeleton.GetPrim().GetPath())
 			);
 		}
 
 		// Import skeleton data
-		pxr::UsdSkelSkeletonQuery SkelQuery = InSkeletonCache.GetSkelQuery( Skeleton );
+		pxr::UsdSkelSkeletonQuery SkelQuery = InSkeletonCache.GetSkelQuery(Skeleton);
 		{
 			FSkeletalMeshImportData DummyImportData;
-			const bool bSkeletonValid = UsdToUnreal::ConvertSkeleton( SkelQuery, DummyImportData );
-			if ( !bSkeletonValid )
+			const bool bSkeletonValid = UsdToUnreal::ConvertSkeleton(SkelQuery, DummyImportData);
+			if (!bSkeletonValid)
 			{
 				return false;
 			}
 			OutSkeletonBones = MoveTemp(DummyImportData.RefBonesBinary);
-			OutSkeletonName = *UsdToUnreal::ConvertString( SkelQuery.GetSkeleton().GetPrim().GetName() );
+			OutSkeletonName = *UsdToUnreal::ConvertString(SkelQuery.GetSkeleton().GetPrim().GetName());
 		}
 
 		TMap<int32, FSkeletalMeshImportData> LODIndexToSkeletalMeshImportDataMap;
@@ -477,10 +493,10 @@ namespace UsdSkelRootTranslatorImpl
 		// Since we may need to switch variants to parse LODs, we could invalidate references to SkinningQuery objects, so we need
 		// to keep track of these by path and construct one whenever we need them
 		TArray<pxr::SdfPath> PathsToSkinnedPrims;
-		for ( const pxr::UsdSkelSkinningQuery& SkinningQuery : SkeletonBinding.GetSkinningTargets() )
+		for (const pxr::UsdSkelSkinningQuery& SkinningQuery : SkeletonBinding.GetSkinningTargets())
 		{
 			// In USD, the skinning target need not be a mesh, but for Unreal we are only interested in skinning meshes
-			if ( pxr::UsdGeomMesh SkinningMesh = pxr::UsdGeomMesh( SkinningQuery.GetPrim() ) )
+			if (pxr::UsdGeomMesh SkinningMesh = pxr::UsdGeomMesh(SkinningQuery.GetPrim()))
 			{
 				// Let's only care about prims with the SkelBindingAPI for now as we'll *need* joint influences and weights
 				if (SkinningQuery.GetPrim().HasAPI<pxr::UsdSkelBindingAPI>())
@@ -492,7 +508,8 @@ namespace UsdSkelRootTranslatorImpl
 					UE_LOG(
 						LogUsd,
 						Log,
-						TEXT("Ignoring skinned prim '%s' when generating Skeletal Mesh for SkelRoot '%s' as the prim doesn't have the SkelBindingAPI"),
+						TEXT("Ignoring skinned prim '%s' when generating Skeletal Mesh for SkelRoot '%s' as the prim doesn't have the SkelBindingAPI"
+						),
 						*UsdToUnreal::ConvertPath(SkinningQuery.GetPrim().GetPrimPath()),
 						*UsdToUnreal::ConvertPath(InSkeletonRoot.GetPrim().GetPrimPath())
 					);
@@ -500,58 +517,48 @@ namespace UsdSkelRootTranslatorImpl
 			}
 		}
 
-		TFunction<bool( const pxr::UsdGeomMesh&, int32 )> ConvertLOD =
-		[
-			&LODIndexToSkeletalMeshImportDataMap,
-			&LODIndexToMaterialInfoMap,
-			&SkelQuery,
-			&InOutUsedMorphTargetNames,
-			OutBlendShapes,
-			&StageInfo,
-			Options
-		]
-		( const pxr::UsdGeomMesh& LODMesh, int32 LODIndex )
-		{
-			pxr::UsdSkelSkinningQuery SkinningQuery = UsdUtils::CreateSkinningQuery( LODMesh, SkelQuery );
-			if ( !SkinningQuery )
+		TFunction<bool(const pxr::UsdGeomMesh&, int32)> ConvertLOD = [&LODIndexToSkeletalMeshImportDataMap,
+																	  &LODIndexToMaterialInfoMap,
+																	  &SkelQuery,
+																	  &InOutUsedMorphTargetNames,
+																	  OutBlendShapes,
+																	  &StageInfo,
+																	  Options](const pxr::UsdGeomMesh& LODMesh, int32 LODIndex)
+		{
+			pxr::UsdSkelSkinningQuery SkinningQuery = UsdUtils::CreateSkinningQuery(LODMesh, SkelQuery);
+			if (!SkinningQuery)
 			{
-				return true; // Continue trying other LODs
+				return true;	// Continue trying other LODs
 			}
 
-			if ( LODMesh && LODMesh.ComputeVisibility() == pxr::UsdGeomTokens->invisible )
+			if (LODMesh && LODMesh.ComputeVisibility() == pxr::UsdGeomTokens->invisible)
 			{
 				return true;
 			}
 
-			FSkeletalMeshImportData& LODImportData = LODIndexToSkeletalMeshImportDataMap.FindOrAdd( LODIndex );
-			UsdUtils::FUsdPrimMaterialAssignmentInfo& LODMaterialInfo = LODIndexToMaterialInfoMap.FindOrAdd( LODIndex );
+			FSkeletalMeshImportData& LODImportData = LODIndexToSkeletalMeshImportDataMap.FindOrAdd(LODIndex);
+			UsdUtils::FUsdPrimMaterialAssignmentInfo& LODMaterialInfo = LODIndexToMaterialInfoMap.FindOrAdd(LODIndex);
 
 			// BlendShape data is respective to point indices for each mesh in isolation, but we combine all points
 			// into one FSkeletalMeshImportData per LOD, so we need to remap the indices using this
-			uint32 NumPointsBeforeThisMesh = static_cast< uint32 >( LODImportData.Points.Num() );
-
-			bool bSuccess = UsdToUnreal::ConvertSkinnedMesh(
-				SkinningQuery,
-				SkelQuery,
-				LODImportData,
-				LODMaterialInfo,
-				Options
-			);
-			if ( !bSuccess )
+			uint32 NumPointsBeforeThisMesh = static_cast<uint32>(LODImportData.Points.Num());
+
+			bool bSuccess = UsdToUnreal::ConvertSkinnedMesh(SkinningQuery, SkelQuery, LODImportData, LODMaterialInfo, Options);
+			if (!bSuccess)
 			{
 				return true;
 			}
 
-			if ( OutBlendShapes )
+			if (OutBlendShapes)
 			{
-				pxr::UsdSkelBindingAPI SkelBindingAPI{ LODMesh.GetPrim() };
-				pxr::UsdSkelBlendShapeQuery BlendShapeQuery{ SkelBindingAPI };
-				if ( BlendShapeQuery )
+				pxr::UsdSkelBindingAPI SkelBindingAPI{LODMesh.GetPrim()};
+				pxr::UsdSkelBlendShapeQuery BlendShapeQuery{SkelBindingAPI};
+				if (BlendShapeQuery)
 				{
-					for ( int32 BlendShapeIndex = 0; BlendShapeIndex < BlendShapeQuery.GetNumBlendShapes(); ++BlendShapeIndex )
+					for (int32 BlendShapeIndex = 0; BlendShapeIndex < BlendShapeQuery.GetNumBlendShapes(); ++BlendShapeIndex)
 					{
 						UsdToUnreal::ConvertBlendShape(
-							BlendShapeQuery.GetBlendShape( BlendShapeIndex ),
+							BlendShapeQuery.GetBlendShape(BlendShapeIndex),
 							StageInfo,
 							LODIndex,
 							NumPointsBeforeThisMesh,
@@ -566,10 +573,10 @@ namespace UsdSkelRootTranslatorImpl
 		};
 
 		// Actually parse all mesh data
-		for ( const pxr::SdfPath& SkinnedPrimPath : PathsToSkinnedPrims )
+		for (const pxr::SdfPath& SkinnedPrimPath : PathsToSkinnedPrims)
 		{
-			pxr::UsdGeomMesh SkinnedMesh{ Stage->GetPrimAtPath( SkinnedPrimPath ) };
-			if ( !SkinnedMesh )
+			pxr::UsdGeomMesh SkinnedMesh{Stage->GetPrimAtPath(SkinnedPrimPath)};
+			if (!SkinnedMesh)
 			{
 				continue;
 			}
@@ -578,7 +585,7 @@ namespace UsdSkelRootTranslatorImpl
 			FString ParentPrimPath = UsdToUnreal::ConvertPath(ParentPrim.GetPath());
 
 			bool bInterpretedLODs = false;
-			if ( bInInterpretLODs && ParentPrim && !ProcessedLODParentPaths.Contains(ParentPrimPath))
+			if (bInInterpretLODs && ParentPrim && !ProcessedLODParentPaths.Contains(ParentPrimPath))
 			{
 				// At the moment we only consider a single mesh per variant, so if multiple meshes tell us to process the same parent prim, we skip.
 				// This check would also prevent us from getting in here in case we just have many meshes children of a same prim, outside
@@ -589,52 +596,52 @@ namespace UsdSkelRootTranslatorImpl
 				bInterpretedLODs = UsdUtils::IterateLODMeshes(ParentPrim, ConvertLOD);
 			}
 
-			if ( !bInterpretedLODs )
+			if (!bInterpretedLODs)
 			{
 				// Refresh reference to this prim as it could have been inside a variant that was temporarily switched by IterateLODMeshes
-				ConvertLOD( pxr::UsdGeomMesh{ Stage->GetPrimAtPath( SkinnedPrimPath ) }, 0);
+				ConvertLOD(pxr::UsdGeomMesh{Stage->GetPrimAtPath(SkinnedPrimPath)}, 0);
 			}
 		}
 
 		// Place the LODs in order as we can't have e.g. LOD0 and LOD2 without LOD1, and there's no reason downstream code needs to care about
 		// what LOD number these data originally wanted to be
 		TMap<int32, int32> OldLODIndexToNewLODIndex;
-		LODIndexToSkeletalMeshImportDataMap.KeySort( TLess<int32>() );
-		OutLODIndexToSkeletalMeshImportData.Reset( LODIndexToSkeletalMeshImportDataMap.Num() );
-		OutLODIndexToMaterialInfo.Reset( LODIndexToMaterialInfoMap.Num() );
-		for ( TPair<int32, FSkeletalMeshImportData>& Entry : LODIndexToSkeletalMeshImportDataMap )
+		LODIndexToSkeletalMeshImportDataMap.KeySort(TLess<int32>());
+		OutLODIndexToSkeletalMeshImportData.Reset(LODIndexToSkeletalMeshImportDataMap.Num());
+		OutLODIndexToMaterialInfo.Reset(LODIndexToMaterialInfoMap.Num());
+		for (TPair<int32, FSkeletalMeshImportData>& Entry : LODIndexToSkeletalMeshImportDataMap)
 		{
 			FSkeletalMeshImportData& ImportData = Entry.Value;
-			if ( Entry.Value.Points.Num() == 0 )
+			if (Entry.Value.Points.Num() == 0)
 			{
 				continue;
 			}
 
 			const int32 OldLODIndex = Entry.Key;
-			const int32 NewLODIndex = OutLODIndexToSkeletalMeshImportData.Add( MoveTemp( ImportData ) );
-			OutLODIndexToMaterialInfo.Add( LODIndexToMaterialInfoMap[ OldLODIndex ] );
+			const int32 NewLODIndex = OutLODIndexToSkeletalMeshImportData.Add(MoveTemp(ImportData));
+			OutLODIndexToMaterialInfo.Add(LODIndexToMaterialInfoMap[OldLODIndex]);
 
 			// Keep track of these to remap blendshapes
-			OldLODIndexToNewLODIndex.Add( OldLODIndex, NewLODIndex );
+			OldLODIndexToNewLODIndex.Add(OldLODIndex, NewLODIndex);
 		}
-		if ( OutBlendShapes )
+		if (OutBlendShapes)
 		{
-			for ( auto& Pair : *OutBlendShapes )
+			for (auto& Pair : *OutBlendShapes)
 			{
 				UsdUtils::FUsdBlendShape& BlendShape = Pair.Value;
 
 				TSet<int32> NewLODIndexUsers;
 				NewLODIndexUsers.Reserve(BlendShape.LODIndicesThatUseThis.Num());
 
-				for ( int32 OldLODIndexUser : BlendShape.LODIndicesThatUseThis )
+				for (int32 OldLODIndexUser : BlendShape.LODIndicesThatUseThis)
 				{
-					if ( int32* FoundNewLODIndex = OldLODIndexToNewLODIndex.Find( OldLODIndexUser ) )
+					if (int32* FoundNewLODIndex = OldLODIndexToNewLODIndex.Find(OldLODIndexUser))
 					{
 						NewLODIndexUsers.Add(*FoundNewLODIndex);
 					}
 					else
 					{
-						UE_LOG(LogUsd, Error, TEXT("Failed to remap blend shape '%s's LOD index '%d'"), *BlendShape.Name, OldLODIndexUser );
+						UE_LOG(LogUsd, Error, TEXT("Failed to remap blend shape '%s's LOD index '%d'"), *BlendShape.Name, OldLODIndexUser);
 					}
 				}
 
@@ -662,8 +669,8 @@ namespace UsdSkelRootTranslatorImpl
 	{
 		FScopedUsdAllocs Allocs;
 
-		pxr::UsdSkelRoot SkelRoot{ SkelRootPrim };
-		if ( !SkelRoot )
+		pxr::UsdSkelRoot SkelRoot{SkelRootPrim};
+		if (!SkelRoot)
 		{
 			return;
 		}
@@ -671,53 +678,50 @@ namespace UsdSkelRootTranslatorImpl
 		pxr::UsdStageRefPtr Stage = SkelRoot.GetPrim().GetStage();
 
 		pxr::TfToken RenderContextToken = pxr::UsdShadeTokens->universalRenderContext;
-		if ( !RenderContext.IsNone() )
+		if (!RenderContext.IsNone())
 		{
-			RenderContextToken = UnrealToUsd::ConvertToken( *RenderContext.ToString() ).Get();
+			RenderContextToken = UnrealToUsd::ConvertToken(*RenderContext.ToString()).Get();
 		}
 
 		pxr::TfToken MaterialPurposeToken = pxr::UsdShadeTokens->allPurpose;
-		if ( !MaterialPurpose.IsNone() )
+		if (!MaterialPurpose.IsNone())
 		{
-			MaterialPurposeToken = UnrealToUsd::ConvertToken( *MaterialPurpose.ToString() ).Get();
+			MaterialPurposeToken = UnrealToUsd::ConvertToken(*MaterialPurpose.ToString()).Get();
 		}
 
 		TMap<int32, UsdUtils::FUsdPrimMaterialAssignmentInfo> LODIndexToMaterialInfoMap;
 		TMap<int32, TSet<UsdUtils::FUsdPrimMaterialSlot>> CombinedSlotsForLODIndex;
-		TFunction<bool( const pxr::UsdGeomMesh&, int32 )> IterateLODsLambda =
-		[
-			&LODIndexToMaterialInfoMap,
-			&CombinedSlotsForLODIndex,
-			Time,
-			RenderContextToken,
-			MaterialPurposeToken
-		]
-		( const pxr::UsdGeomMesh& LODMesh, int32 LODIndex )
-		{
-			if ( LODMesh && LODMesh.ComputeVisibility() == pxr::UsdGeomTokens->invisible )
+		TFunction<bool(const pxr::UsdGeomMesh&, int32)> IterateLODsLambda = [&LODIndexToMaterialInfoMap,
+																			 &CombinedSlotsForLODIndex,
+																			 Time,
+																			 RenderContextToken,
+																			 MaterialPurposeToken](const pxr::UsdGeomMesh& LODMesh, int32 LODIndex)
+		{
+			if (LODMesh && LODMesh.ComputeVisibility() == pxr::UsdGeomTokens->invisible)
 			{
 				return true;
 			}
 
-			TArray<UsdUtils::FUsdPrimMaterialSlot>& CombinedLODSlots = LODIndexToMaterialInfoMap.FindOrAdd( LODIndex ).Slots;
-			TSet<UsdUtils::FUsdPrimMaterialSlot>& CombinedLODSlotsSet = CombinedSlotsForLODIndex.FindOrAdd( LODIndex );
+			TArray<UsdUtils::FUsdPrimMaterialSlot>& CombinedLODSlots = LODIndexToMaterialInfoMap.FindOrAdd(LODIndex).Slots;
+			TSet<UsdUtils::FUsdPrimMaterialSlot>& CombinedLODSlotsSet = CombinedSlotsForLODIndex.FindOrAdd(LODIndex);
 
-			const bool bProvideMaterialIndices = false; // We have no use for material indices and it can be slow to retrieve, as it will iterate all faces
+			const bool bProvideMaterialIndices = false;	   // We have no use for material indices and it can be slow to retrieve, as it will iterate
+														   // all faces
 			UsdUtils::FUsdPrimMaterialAssignmentInfo LocalInfo = UsdUtils::GetPrimMaterialAssignments(
 				LODMesh.GetPrim(),
-				pxr::UsdTimeCode( Time ),
+				pxr::UsdTimeCode(Time),
 				bProvideMaterialIndices,
 				RenderContextToken,
 				MaterialPurposeToken
 			);
 
 			// Combine material slots in the same order that UsdToUnreal::ConvertSkinnedMesh does
-			for ( UsdUtils::FUsdPrimMaterialSlot& LocalSlot : LocalInfo.Slots )
+			for (UsdUtils::FUsdPrimMaterialSlot& LocalSlot : LocalInfo.Slots)
 			{
-				if ( !CombinedLODSlotsSet.Contains( LocalSlot ) )
+				if (!CombinedLODSlotsSet.Contains(LocalSlot))
 				{
-					CombinedLODSlots.Add( LocalSlot );
-					CombinedLODSlotsSet.Add( LocalSlot );
+					CombinedLODSlots.Add(LocalSlot);
+					CombinedLODSlotsSet.Add(LocalSlot);
 				}
 			}
 
@@ -729,36 +733,36 @@ namespace UsdSkelRootTranslatorImpl
 		// Because we combine all skinning target meshes into a single skeletal mesh, we'll have to reconstruct the combined
 		// material assignment info that this SkelRoot wants in order to compare with the existing assignments.
 		pxr::UsdSkelCache SkeletonCache;
-		SkeletonCache.Populate( SkelRoot, pxr::UsdTraverseInstanceProxies() );
-		std::vector< pxr::UsdSkelBinding > SkeletonBindings;
-		SkeletonCache.ComputeSkelBindings( SkelRoot, &SkeletonBindings, pxr::UsdTraverseInstanceProxies() );
-		for ( const pxr::UsdSkelBinding& Binding : SkeletonBindings )
+		SkeletonCache.Populate(SkelRoot, pxr::UsdTraverseInstanceProxies());
+		std::vector<pxr::UsdSkelBinding> SkeletonBindings;
+		SkeletonCache.ComputeSkelBindings(SkelRoot, &SkeletonBindings, pxr::UsdTraverseInstanceProxies());
+		for (const pxr::UsdSkelBinding& Binding : SkeletonBindings)
 		{
-			for ( const pxr::UsdSkelSkinningQuery& SkinningQuery : Binding.GetSkinningTargets() )
+			for (const pxr::UsdSkelSkinningQuery& SkinningQuery : Binding.GetSkinningTargets())
 			{
 				pxr::UsdPrim MeshPrim = SkinningQuery.GetPrim();
-				pxr::UsdGeomMesh Mesh{ MeshPrim };
-				if ( !Mesh )
+				pxr::UsdGeomMesh Mesh{MeshPrim};
+				if (!Mesh)
 				{
 					continue;
 				}
 				pxr::SdfPath MeshPrimPath = MeshPrim.GetPath();
 
 				pxr::UsdPrim ParentPrim = MeshPrim.GetParent();
-				FString ParentPrimPath = UsdToUnreal::ConvertPath( ParentPrim.GetPath() );
+				FString ParentPrimPath = UsdToUnreal::ConvertPath(ParentPrim.GetPath());
 
 				bool bInterpretedLODs = false;
-				if ( bInterpretLODs && UsdUtils::IsGeomMeshALOD( MeshPrim ) && !ProcessedLODParentPaths.Contains( ParentPrimPath ) )
+				if (bInterpretLODs && UsdUtils::IsGeomMeshALOD(MeshPrim) && !ProcessedLODParentPaths.Contains(ParentPrimPath))
 				{
-					ProcessedLODParentPaths.Add( ParentPrimPath );
+					ProcessedLODParentPaths.Add(ParentPrimPath);
 
-					bInterpretedLODs = UsdUtils::IterateLODMeshes( ParentPrim, IterateLODsLambda );
+					bInterpretedLODs = UsdUtils::IterateLODMeshes(ParentPrim, IterateLODsLambda);
 				}
 
-				if ( !bInterpretedLODs )
+				if (!bInterpretedLODs)
 				{
 					// Refresh reference to this prim as it could have been inside a variant that was temporarily switched by IterateLODMeshes
-					IterateLODsLambda( pxr::UsdGeomMesh{ Stage->GetPrimAtPath( MeshPrimPath ) }, 0 );
+					IterateLODsLambda(pxr::UsdGeomMesh{Stage->GetPrimAtPath(MeshPrimPath)}, 0);
 				}
 			}
 
@@ -768,15 +772,15 @@ namespace UsdSkelRootTranslatorImpl
 		}
 
 		// Refresh reference to SkelRootPrim because variant switching potentially invalidated it
-		pxr::UsdPrim ValidSkelRootPrim = Stage->GetPrimAtPath( SkelRootPrimPath );
+		pxr::UsdPrim ValidSkelRootPrim = Stage->GetPrimAtPath(SkelRootPrimPath);
 
 		// Place the LODs in order as we can't have e.g. LOD0 and LOD2 without LOD1, and there's no reason downstream code needs to care about
 		// what LOD number these data originally wanted to be
 		TArray<UsdUtils::FUsdPrimMaterialAssignmentInfo> LODIndexToAssignments;
-		LODIndexToMaterialInfoMap.KeySort( TLess<int32>() );
-		for ( TPair<int32, UsdUtils::FUsdPrimMaterialAssignmentInfo>& Entry : LODIndexToMaterialInfoMap )
+		LODIndexToMaterialInfoMap.KeySort(TLess<int32>());
+		for (TPair<int32, UsdUtils::FUsdPrimMaterialAssignmentInfo>& Entry : LODIndexToMaterialInfoMap)
 		{
-			LODIndexToAssignments.Add( MoveTemp( Entry.Value ) );
+			LODIndexToAssignments.Add(MoveTemp(Entry.Value));
 		}
 
 		// Stash our mesh PrimvarToUVIndex into the assignment info, as that's where ResolveMaterialAssignmentInfo will look for it
@@ -812,52 +816,59 @@ namespace UsdSkelRootTranslatorImpl
 
 		// Compare resolved materials with existing assignments, and create overrides if we need to
 		uint32 SkeletalMeshSlotIndex = 0;
-		for ( int32 LODIndex = 0; LODIndex < LODIndexToAssignments.Num(); ++LODIndex )
+		for (int32 LODIndex = 0; LODIndex < LODIndexToAssignments.Num(); ++LODIndex)
 		{
-			const TArray< UsdUtils::FUsdPrimMaterialSlot >& LODSlots = LODIndexToAssignments[ LODIndex ].Slots;
-			for ( int32 LODSlotIndex = 0; LODSlotIndex < LODSlots.Num(); ++LODSlotIndex, ++SkeletalMeshSlotIndex )
+			const TArray<UsdUtils::FUsdPrimMaterialSlot>& LODSlots = LODIndexToAssignments[LODIndex].Slots;
+			for (int32 LODSlotIndex = 0; LODSlotIndex < LODSlots.Num(); ++LODSlotIndex, ++SkeletalMeshSlotIndex)
 			{
-				const UsdUtils::FUsdPrimMaterialSlot& Slot = LODSlots[ LODSlotIndex ];
+				const UsdUtils::FUsdPrimMaterialSlot& Slot = LODSlots[LODSlotIndex];
 
 				UMaterialInterface* Material = nullptr;
-				if ( UMaterialInterface** FoundMaterial = ResolvedMaterials.Find( &Slot ) )
+				if (UMaterialInterface** FoundMaterial = ResolvedMaterials.Find(&Slot))
 				{
 					Material = *FoundMaterial;
 				}
 				else
 				{
-					UE_LOG( LogUsd, Error, TEXT( "Lost track of resolved material for slot '%d' of LOD '%d' for mesh '%s'" ), LODSlotIndex, LODIndex, *UsdToUnreal::ConvertPath( ValidSkelRootPrim.GetPath() ) );
+					UE_LOG(
+						LogUsd,
+						Error,
+						TEXT("Lost track of resolved material for slot '%d' of LOD '%d' for mesh '%s'"),
+						LODSlotIndex,
+						LODIndex,
+						*UsdToUnreal::ConvertPath(ValidSkelRootPrim.GetPath())
+					);
 					continue;
 				}
 
-				UMaterialInterface* ExistingMaterial = ExistingAssignments[ SkeletalMeshSlotIndex ];
-				if ( ExistingMaterial == Material )
+				UMaterialInterface* ExistingMaterial = ExistingAssignments[SkeletalMeshSlotIndex];
+				if (ExistingMaterial == Material)
 				{
 					continue;
 				}
 				else
 				{
-					MeshComponent.SetMaterial( SkeletalMeshSlotIndex, Material );
+					MeshComponent.SetMaterial(SkeletalMeshSlotIndex, Material);
 				}
 			}
 		}
 	}
 
-	bool HasLODSkinningTargets( const pxr::UsdSkelRoot& SkelRoot )
+	bool HasLODSkinningTargets(const pxr::UsdSkelRoot& SkelRoot)
 	{
 		FScopedUsdAllocs Allocs;
 
 		pxr::UsdSkelCache SkeletonCache;
-		SkeletonCache.Populate( SkelRoot, pxr::UsdTraverseInstanceProxies() );
+		SkeletonCache.Populate(SkelRoot, pxr::UsdTraverseInstanceProxies());
 
-		std::vector< pxr::UsdSkelBinding > SkeletonBindings;
-		SkeletonCache.ComputeSkelBindings( SkelRoot, &SkeletonBindings, pxr::UsdTraverseInstanceProxies() );
+		std::vector<pxr::UsdSkelBinding> SkeletonBindings;
+		SkeletonCache.ComputeSkelBindings(SkelRoot, &SkeletonBindings, pxr::UsdTraverseInstanceProxies());
 
-		for ( const pxr::UsdSkelBinding& Binding : SkeletonBindings )
+		for (const pxr::UsdSkelBinding& Binding : SkeletonBindings)
 		{
-			for ( const pxr::UsdSkelSkinningQuery& SkinningQuery : Binding.GetSkinningTargets() )
+			for (const pxr::UsdSkelSkinningQuery& SkinningQuery : Binding.GetSkinningTargets())
 			{
-				if ( UsdUtils::IsGeomMeshALOD( SkinningQuery.GetPrim() ) )
+				if (UsdUtils::IsGeomMeshALOD(SkinningQuery.GetPrim()))
 				{
 					return true;
 				}
@@ -867,20 +878,16 @@ namespace UsdSkelRootTranslatorImpl
 		return false;
 	}
 
-	UPhysicsAsset* GenerateAndAssignPhysicsAsset( USkeletalMesh* SkeletalMesh, EObjectFlags Flags )
+	UPhysicsAsset* GenerateAndAssignPhysicsAsset(USkeletalMesh* SkeletalMesh, EObjectFlags Flags)
 	{
 #if WITH_EDITOR
 		FName AssetName = MakeUniqueObjectName(
 			GetTransientPackage(),
 			UPhysicsAsset::StaticClass(),
-			*IUsdClassesModule::SanitizeObjectName(
-				FPaths::GetBaseFilename(
-					TEXT("PHYS_") + SkeletalMesh->GetName()
-				)
-			)
+			*IUsdClassesModule::SanitizeObjectName(FPaths::GetBaseFilename(TEXT("PHYS_") + SkeletalMesh->GetName()))
 		);
 
-		UPhysicsAsset* Result = NewObject< UPhysicsAsset >(
+		UPhysicsAsset* Result = NewObject<UPhysicsAsset>(
 			GetTransientPackage(),
 			AssetName,
 			Flags | EObjectFlags::RF_Public | EObjectFlags::RF_Transient
@@ -888,15 +895,18 @@ namespace UsdSkelRootTranslatorImpl
 
 		FPhysAssetCreateParams NewBodyData;
 		FText CreationErrorMessage;
-		bool bSuccess = FPhysicsAssetUtils::CreateFromSkeletalMesh( Result, SkeletalMesh, NewBodyData, CreationErrorMessage );
-		if ( !bSuccess )
+		bool bSuccess = FPhysicsAssetUtils::CreateFromSkeletalMesh(Result, SkeletalMesh, NewBodyData, CreationErrorMessage);
+		if (!bSuccess)
 		{
-			UE_LOG( LogUsd, Error, TEXT( "Couldn't create PhysicsAsset for SkeletalMesh '%s': %s" ),
+			UE_LOG(
+				LogUsd,
+				Error,
+				TEXT("Couldn't create PhysicsAsset for SkeletalMesh '%s': %s"),
 				*SkeletalMesh->GetName(),
 				*CreationErrorMessage.ToString()
 			);
 
-			ensure( ObjectTools::DeleteSingleObject( Result ) );
+			ensure(ObjectTools::DeleteSingleObject(Result));
 
 			Result = nullptr;
 		}
@@ -904,14 +914,14 @@ namespace UsdSkelRootTranslatorImpl
 		return Result;
 #else
 		return nullptr;
-#endif // WITH_EDITOR
+#endif	  // WITH_EDITOR
 	}
 
 	UAnimBlueprint* CreateAnimBlueprint(
 		const FUsdSchemaTranslationContext& Context,
 		const pxr::UsdPrim& Prim,
-		bool bDelayRecompilation = false,  // Whether we should recompile within this function or not
-		bool* bOutNeedsRecompile = nullptr // Whether this function wants the returned AnimBP to be recompiled
+		bool bDelayRecompilation = false,	  // Whether we should recompile within this function or not
+		bool* bOutNeedsRecompile = nullptr	  // Whether this function wants the returned AnimBP to be recompiled
 	)
 	{
 		if (!Prim || !Context.InfoCache || Context.AssetCache)
@@ -919,18 +929,16 @@ namespace UsdSkelRootTranslatorImpl
 			return nullptr;
 		}
 
-		FString PrimName = UsdToUnreal::ConvertString( Prim.GetName() );
+		FString PrimName = UsdToUnreal::ConvertString(Prim.GetName());
 
-		USkeletalMesh* SkeletalMesh = Context.InfoCache->GetSingleAssetForPrim<USkeletalMesh>(
-			UE::FSdfPath{Prim.GetPath()}
-		);
-		if ( !SkeletalMesh )
+		USkeletalMesh* SkeletalMesh = Context.InfoCache->GetSingleAssetForPrim<USkeletalMesh>(UE::FSdfPath{Prim.GetPath()});
+		if (!SkeletalMesh)
 		{
 			return nullptr;
 		}
 
 		USkeleton* Skeleton = SkeletalMesh->GetSkeleton();
-		if ( !Skeleton )
+		if (!Skeleton)
 		{
 			return nullptr;
 		}
@@ -940,31 +948,31 @@ namespace UsdSkelRootTranslatorImpl
 		{
 			FScopedUsdAllocs Allocs;
 
-			if ( pxr::UsdAttribute Attr = Prim.GetAttribute( UnrealIdentifiers::UnrealAnimBlueprintPath ) )
+			if (pxr::UsdAttribute Attr = Prim.GetAttribute(UnrealIdentifiers::UnrealAnimBlueprintPath))
 			{
 				std::string PathString;
-				if ( Attr.Get( &PathString ) )
+				if (Attr.Get(&PathString))
 				{
-					AnimBPPath = UsdToUnreal::ConvertString( PathString );
+					AnimBPPath = UsdToUnreal::ConvertString(PathString);
 				}
 			}
 		}
-		if ( AnimBPPath.IsEmpty() )
+		if (AnimBPPath.IsEmpty())
 		{
 			return nullptr;
 		}
 
 		bool bNeedRecompile = false;
 
-		UAnimBlueprint* AnimBP = Cast<UAnimBlueprint>( FSoftObjectPath( AnimBPPath ).TryLoad() );
-		if ( !AnimBP )
+		UAnimBlueprint* AnimBP = Cast<UAnimBlueprint>(FSoftObjectPath(AnimBPPath).TryLoad());
+		if (!AnimBP)
 		{
 			return nullptr;
 		}
 
 		// Create transient AnimBP based on our template, so that we can assign it a proper skeleton
-		const static FString DefaultAnimBPPath = TEXT( "/USDImporter/Blueprint/DefaultLiveLinkAnimBP.DefaultLiveLinkAnimBP" );
-		if ( DefaultAnimBPPath == AnimBPPath )
+		const static FString DefaultAnimBPPath = TEXT("/USDImporter/Blueprint/DefaultLiveLinkAnimBP.DefaultLiveLinkAnimBP");
+		if (DefaultAnimBPPath == AnimBPPath)
 		{
 			const FString PrimPath = UsdToUnreal::ConvertPath(Prim.GetPrimPath());
 
@@ -984,7 +992,7 @@ namespace UsdSkelRootTranslatorImpl
 			bool bReusedAnimBP = false;
 			if (UAnimBlueprint* CachedAnimBP = Cast<UAnimBlueprint>(Context.AssetCache->GetCachedAsset(PrefixedAnimBPHash)))
 			{
-				if ( CachedAnimBP->TargetSkeleton == Skeleton )
+				if (CachedAnimBP->TargetSkeleton == Skeleton)
 				{
 					AnimBP = CachedAnimBP;
 					bReusedAnimBP = true;
@@ -992,46 +1000,48 @@ namespace UsdSkelRootTranslatorImpl
 			}
 
 			// We have to generate a new transient AnimBP
-			if ( !bReusedAnimBP )
+			if (!bReusedAnimBP)
 			{
 				FName UniqueName = MakeUniqueObjectName(
 					GetTransientPackage(),
 					UAnimBlueprint::StaticClass(),
-					*IUsdClassesModule::SanitizeObjectName(
-						PrimName + TEXT( "_DefaultAnimBlueprint" )
-					)
+					*IUsdClassesModule::SanitizeObjectName(PrimName + TEXT("_DefaultAnimBlueprint"))
 				);
 
 				// Duplicate and never reuse these so that they can be assigned independent subject names if desired.
 				// Its not as if scenes will have thousands of these anyway.
-				AnimBP = DuplicateObject( AnimBP, GetTransientPackage(), UniqueName );
+				AnimBP = DuplicateObject(AnimBP, GetTransientPackage(), UniqueName);
 
 				// Patch up the flags here or else the rest of the engine code (and our plugin code) will get confused
 				// as to what a non-transient asset in the transient package even means. It can lead to some crashes
 				// when saving after import if we don't do this
-				AnimBP->ClearFlags( AnimBP->GetFlags() );
-				AnimBP->SetFlags( Skeleton->GetFlags() );
+				AnimBP->ClearFlags(AnimBP->GetFlags());
+				AnimBP->SetFlags(Skeleton->GetFlags());
 
 				AnimBP->TargetSkeleton = Skeleton;
 				AnimBP->bIsTemplate = false;
 
 				bNeedRecompile = true;
 
-				Context.AssetCache->CacheAsset( PrefixedAnimBPHash, AnimBP );
+				Context.AssetCache->CacheAsset(PrefixedAnimBPHash, AnimBP);
 			}
 		}
 		// Path is pointing to an existing, persistent AnimBP
 		else
 		{
 			// Force skeletons to be compatible if they aren't (we need both ways!)
-			if ( !Skeleton->IsCompatibleForEditor( AnimBP->TargetSkeleton ) )
+			if (!Skeleton->IsCompatibleForEditor(AnimBP->TargetSkeleton))
 			{
-				AnimBP->TargetSkeleton->AddCompatibleSkeleton( Skeleton );
-				Skeleton->AddCompatibleSkeleton( AnimBP->TargetSkeleton );
+				AnimBP->TargetSkeleton->AddCompatibleSkeleton(Skeleton);
+				Skeleton->AddCompatibleSkeleton(AnimBP->TargetSkeleton);
 
-				if ( AnimBP->TargetSkeleton->GetReferenceSkeleton().GetRefBoneInfo() != Skeleton->GetReferenceSkeleton().GetRefBoneInfo() )
+				if (AnimBP->TargetSkeleton->GetReferenceSkeleton().GetRefBoneInfo() != Skeleton->GetReferenceSkeleton().GetRefBoneInfo())
 				{
-					UE_LOG( LogUsd, Warning, TEXT( "Forcing AnimBlueprint '%s's Skeleton '%s' to be compatible with the Skeleton generated for prim '%s', but they may be different!" ),
+					UE_LOG(
+						LogUsd,
+						Warning,
+						TEXT("Forcing AnimBlueprint '%s's Skeleton '%s' to be compatible with the Skeleton generated for prim '%s', but they may be "
+							 "different!"),
 						*AnimBP->GetPathName(),
 						*AnimBP->TargetSkeleton->GetPathName(),
 						*PrimName
@@ -1040,15 +1050,15 @@ namespace UsdSkelRootTranslatorImpl
 			}
 		}
 
-		if ( bNeedRecompile && !bDelayRecompilation )
+		if (bNeedRecompile && !bDelayRecompilation)
 		{
 			FCompilerResultsLog Results;
-			FBPCompileRequest Request( AnimBP, EBlueprintCompileOptions::None, &Results );
-			FBlueprintCompilationManager::CompileSynchronously( Request );
+			FBPCompileRequest Request(AnimBP, EBlueprintCompileOptions::None, &Results);
+			FBlueprintCompilationManager::CompileSynchronously(Request);
 			bNeedRecompile = false;
 		}
 
-		if ( bOutNeedsRecompile )
+		if (bOutNeedsRecompile)
 		{
 			*bOutNeedsRecompile = bNeedRecompile;
 		}
@@ -1056,25 +1066,25 @@ namespace UsdSkelRootTranslatorImpl
 		return AnimBP;
 	}
 
-	void UpdateLiveLinkProperties( const FUsdSchemaTranslationContext& Context, USkeletalMeshComponent* Component, const pxr::UsdPrim& Prim )
+	void UpdateLiveLinkProperties(const FUsdSchemaTranslationContext& Context, USkeletalMeshComponent* Component, const pxr::UsdPrim& Prim)
 	{
-		if ( !Component || !Component->GetSkeletalMeshAsset() || !Prim )
+		if (!Component || !Component->GetSkeletalMeshAsset() || !Prim)
 		{
 			return;
 		}
 
-		FString PrimName = UsdToUnreal::ConvertString( Prim.GetName() );
+		FString PrimName = UsdToUnreal::ConvertString(Prim.GetName());
 
 		USkeleton* Skeleton = Component->GetSkeletalMeshAsset()->GetSkeleton();
-		if ( !Skeleton )
+		if (!Skeleton)
 		{
 			return;
 		}
 
 		UAnimBlueprint* ExistingAnimBP = nullptr;
-		if ( Component->AnimClass )
+		if (Component->AnimClass)
 		{
-			ExistingAnimBP = Cast<UAnimBlueprint>( Component->AnimClass->ClassGeneratedBy );
+			ExistingAnimBP = Cast<UAnimBlueprint>(Component->AnimClass->ClassGeneratedBy);
 		}
 
 		// Fetch relevant attributes from prim, since we know it has the schema
@@ -1083,90 +1093,88 @@ namespace UsdSkelRootTranslatorImpl
 		{
 			FScopedUsdAllocs Allocs;
 
-			if ( pxr::UsdAttribute Attr = Prim.GetAttribute( UnrealIdentifiers::UnrealLiveLinkSubjectName ) )
+			if (pxr::UsdAttribute Attr = Prim.GetAttribute(UnrealIdentifiers::UnrealLiveLinkSubjectName))
 			{
 				std::string SubjectNameString;
-				if ( Attr.Get( &SubjectNameString ) )
+				if (Attr.Get(&SubjectNameString))
 				{
-					SubjectName = UsdToUnreal::ConvertString( SubjectNameString );
+					SubjectName = UsdToUnreal::ConvertString(SubjectNameString);
 				}
 			}
 
-			if ( pxr::UsdAttribute Attr = Prim.GetAttribute( UnrealIdentifiers::UnrealAnimBlueprintPath ) )
+			if (pxr::UsdAttribute Attr = Prim.GetAttribute(UnrealIdentifiers::UnrealAnimBlueprintPath))
 			{
 				std::string PathString;
-				if ( Attr.Get( &PathString ) )
+				if (Attr.Get(&PathString))
 				{
-					AnimBPPath = UsdToUnreal::ConvertString( PathString );
+					AnimBPPath = UsdToUnreal::ConvertString(PathString);
 				}
 			}
 		}
 
-		UAnimBlueprint* AnimBP = Cast<UAnimBlueprint>( FSoftObjectPath( AnimBPPath ).TryLoad() );
+		UAnimBlueprint* AnimBP = Cast<UAnimBlueprint>(FSoftObjectPath(AnimBPPath).TryLoad());
 
 		// Check if we need to change the AnimBP
 		bool bNeedRecompile = false;
-		if ( ExistingAnimBP != AnimBP )
+		if (ExistingAnimBP != AnimBP)
 		{
 			const bool bDelayRecompilation = true;
-			AnimBP = UsdSkelRootTranslatorImpl::CreateAnimBlueprint(
-				Context,
-				Prim,
-				bDelayRecompilation,
-				&bNeedRecompile
-			);
+			AnimBP = UsdSkelRootTranslatorImpl::CreateAnimBlueprint(Context, Prim, bDelayRecompilation, &bNeedRecompile);
 		}
 
 		// Apply subject name to live link pose AnimBlueprint node
 		// Reference: UAnimationBlueprintLibrary::AddNodeAssetOverride
-		if ( AnimBP )
+		if (AnimBP)
 		{
 			TArray<UBlueprint*> BlueprintHierarchy;
-			AnimBP->GetBlueprintHierarchyFromClass( AnimBP->GetAnimBlueprintGeneratedClass(), BlueprintHierarchy );
+			AnimBP->GetBlueprintHierarchyFromClass(AnimBP->GetAnimBlueprintGeneratedClass(), BlueprintHierarchy);
 
 			TArray<UAnimGraphNode_LiveLinkPose*> LiveLinkNodes;
 
-			for ( int32 BlueprintIndex = 0; BlueprintIndex < BlueprintHierarchy.Num(); ++BlueprintIndex )
+			for (int32 BlueprintIndex = 0; BlueprintIndex < BlueprintHierarchy.Num(); ++BlueprintIndex)
 			{
-				UBlueprint* CurrentBlueprint = BlueprintHierarchy[ BlueprintIndex ];
+				UBlueprint* CurrentBlueprint = BlueprintHierarchy[BlueprintIndex];
 
 				TArray<UEdGraph*> Graphs;
-				CurrentBlueprint->GetAllGraphs( Graphs );
+				CurrentBlueprint->GetAllGraphs(Graphs);
 
-				for ( UEdGraph* Graph : Graphs )
+				for (UEdGraph* Graph : Graphs)
 				{
-					for ( UEdGraphNode* Node : Graph->Nodes )
+					for (UEdGraphNode* Node : Graph->Nodes)
 					{
-						if ( UAnimGraphNode_LiveLinkPose* AnimNode = Cast<UAnimGraphNode_LiveLinkPose>( Node ) )
+						if (UAnimGraphNode_LiveLinkPose* AnimNode = Cast<UAnimGraphNode_LiveLinkPose>(Node))
 						{
-							LiveLinkNodes.Add( AnimNode );
+							LiveLinkNodes.Add(AnimNode);
 						}
 					}
 				}
 			}
 
-			if ( LiveLinkNodes.Num() > 1 && !ExistingAnimBP )
+			if (LiveLinkNodes.Num() > 1 && !ExistingAnimBP)
 			{
-				UE_LOG( LogUsd, Warning, TEXT( "Found more than one LiveLinkPose blueprint node on AnimBlueprint '%s's graphs."
-					"Note that all of those nodes will have their LiveLink SubjectName's updated to '%s', as described for prim '%s'!" ),
+				UE_LOG(
+					LogUsd,
+					Warning,
+					TEXT("Found more than one LiveLinkPose blueprint node on AnimBlueprint '%s's graphs."
+						 "Note that all of those nodes will have their LiveLink SubjectName's updated to '%s', as described for prim '%s'!"),
 					*AnimBP->GetPathName(),
 					*SubjectName,
-					*UsdToUnreal::ConvertPath( Prim.GetPrimPath() )
+					*UsdToUnreal::ConvertPath(Prim.GetPrimPath())
 				);
 			}
 
-			for ( UAnimGraphNode_LiveLinkPose* Node : LiveLinkNodes )
+			for (UAnimGraphNode_LiveLinkPose* Node : LiveLinkNodes)
 			{
 				const UEdGraphSchema* Schema = Node->GetSchema();
-				if ( !Schema )
+				if (!Schema)
 				{
 					continue;
 				}
 
 				UEdGraphPin* SubjectNamePin = nullptr;
-				for ( UEdGraphPin* Pin : Node->Pins )
+				for (UEdGraphPin* Pin : Node->Pins)
 				{
-					if ( Pin->GetName() == GET_MEMBER_NAME_STRING_CHECKED( FAnimNode_LiveLinkPose, LiveLinkSubjectName ) )
+					if (Pin->GetName() == GET_MEMBER_NAME_STRING_CHECKED(FAnimNode_LiveLinkPose, LiveLinkSubjectName))
 					{
 						SubjectNamePin = Pin;
 						break;
@@ -1174,12 +1182,16 @@ namespace UsdSkelRootTranslatorImpl
 				}
 
 				// The subject name pin is already connected to something...
-				if ( SubjectNamePin->LinkedTo.Num() != 0 )
+				if (SubjectNamePin->LinkedTo.Num() != 0)
 				{
-					if ( !ExistingAnimBP )
+					if (!ExistingAnimBP)
 					{
-						UE_LOG( LogUsd, Warning, TEXT( "Failed to update a LiveLinkPose node's 'Subject Name' to '%s' on AnimBlueprint '%s', "
-							"because the pin is already connected to some other node. Disconnect it if you want it to be updated automatically." ),
+						UE_LOG(
+							LogUsd,
+							Warning,
+							TEXT("Failed to update a LiveLinkPose node's 'Subject Name' to '%s' on AnimBlueprint '%s', "
+								 "because the pin is already connected to some other node. Disconnect it if you want it to be updated automatically."
+							),
 							*SubjectName,
 							*AnimBP->GetPathName()
 						);
@@ -1188,7 +1200,7 @@ namespace UsdSkelRootTranslatorImpl
 					continue;
 				}
 
-				if ( SubjectNamePin )
+				if (SubjectNamePin)
 				{
 					// The pin type is FLiveLinkSubjectName, so we must create an instance of it and serialize it using
 					// UScriptStruct::ExportText to generate a proper default value string
@@ -1200,50 +1212,50 @@ namespace UsdSkelRootTranslatorImpl
 					const UObject* OwnerObject = nullptr;
 					const int32 PortFlags = EPropertyPortFlags::PPF_None;
 					const UObject* ExportRootScope = nullptr;
-					FLiveLinkSubjectName::StaticStruct()->ExportText( ValueString, &Dummy, nullptr, nullptr, EPropertyPortFlags::PPF_None, nullptr );
+					FLiveLinkSubjectName::StaticStruct()->ExportText(ValueString, &Dummy, nullptr, nullptr, EPropertyPortFlags::PPF_None, nullptr);
 
-					if ( !Schema->DoesDefaultValueMatch( *SubjectNamePin, ValueString ) )
+					if (!Schema->DoesDefaultValueMatch(*SubjectNamePin, ValueString))
 					{
 						SubjectNamePin->Modify();
-						Schema->TrySetDefaultValue( *SubjectNamePin, ValueString );
+						Schema->TrySetDefaultValue(*SubjectNamePin, ValueString);
 
-						FBlueprintEditorUtils::MarkBlueprintAsModified( AnimBP );
+						FBlueprintEditorUtils::MarkBlueprintAsModified(AnimBP);
 						bNeedRecompile = true;
 					}
 				}
 			}
 		}
 
-		if ( bNeedRecompile )
+		if (bNeedRecompile)
 		{
 			FCompilerResultsLog Results;
-			FBPCompileRequest Request( AnimBP, EBlueprintCompileOptions::None, &Results );
-			FBlueprintCompilationManager::CompileSynchronously( Request );
+			FBPCompileRequest Request(AnimBP, EBlueprintCompileOptions::None, &Results);
+			FBlueprintCompilationManager::CompileSynchronously(Request);
 
 			// We need to force the component to update its anim after we regenerate the blueprint class
 			Component->ClearAnimScriptInstance();
-			Component->InitAnim( true );
+			Component->InitAnim(true);
 		}
 
-		if ( AnimBP != ExistingAnimBP )
+		if (AnimBP != ExistingAnimBP)
 		{
 			// This can internally change AnimationMode, but lets revert it to what it was so that we can control it from
 			// that single place in ::UpdateComponents
 			EAnimationMode::Type OldMode = Component->GetAnimationMode();
-			Component->SetAnimInstanceClass( AnimBP->GeneratedClass );
-			Component->SetAnimationMode( OldMode );
+			Component->SetAnimInstanceClass(AnimBP->GeneratedClass);
+			Component->SetAnimationMode(OldMode);
 		}
 	}
 
 	class FSkelRootCreateAssetsTaskChain : public FUsdSchemaTranslatorTaskChain
 	{
 	public:
-		explicit FSkelRootCreateAssetsTaskChain( const TSharedRef< FUsdSchemaTranslationContext >& InContext, const UE::FSdfPath& InPrimPath );
+		explicit FSkelRootCreateAssetsTaskChain(const TSharedRef<FUsdSchemaTranslationContext>& InContext, const UE::FSdfPath& InPrimPath);
 
 	protected:
 		// Inputs
 		UE::FSdfPath PrimPath;
-		TSharedRef< FUsdSchemaTranslationContext > Context;
+		TSharedRef<FUsdSchemaTranslationContext> Context;
 
 		// Outputs
 		TArray<FSkeletalMeshImportData> LODIndexToSkeletalMeshImportData;
@@ -1254,18 +1266,25 @@ namespace UsdSkelRootTranslatorImpl
 
 		// Note that we want this to be case insensitive so that our UMorphTarget FNames are unique not only due to case differences
 		TSet<FString> UsedMorphTargetNames;
-		TUsdStore< pxr::UsdSkelCache > SkeletonCache;
+		TUsdStore<pxr::UsdSkelCache> SkeletonCache;
 		FString PrefixedSkelMeshHash;
 
-		// Don't keep a live reference to the prim because other translators may mutate the stage in an ExclusiveSync translation step, invalidating the reference
-		UE::FUsdPrim GetPrim() const { return Context->Stage.GetPrimAtPath( PrimPath ); }
+		// Don't keep a live reference to the prim because other translators may mutate the stage in an ExclusiveSync translation step, invalidating
+		// the reference
+		UE::FUsdPrim GetPrim() const
+		{
+			return Context->Stage.GetPrimAtPath(PrimPath);
+		}
 
 		void SetupTasks();
 	};
 
-	FSkelRootCreateAssetsTaskChain::FSkelRootCreateAssetsTaskChain( const TSharedRef< FUsdSchemaTranslationContext >& InContext, const UE::FSdfPath& InPrimPath )
-		: PrimPath( InPrimPath )
-		, Context( InContext )
+	FSkelRootCreateAssetsTaskChain::FSkelRootCreateAssetsTaskChain(
+		const TSharedRef<FUsdSchemaTranslationContext>& InContext,
+		const UE::FSdfPath& InPrimPath
+	)
+		: PrimPath(InPrimPath)
+		, Context(InContext)
 	{
 		SetupTasks();
 	}
@@ -1273,7 +1292,7 @@ namespace UsdSkelRootTranslatorImpl
 	void FSkelRootCreateAssetsTaskChain::SetupTasks()
 	{
 		// Ignore prims from disabled purposes
-		if ( !EnumHasAllFlags( Context->PurposesToLoad, IUsdPrim::GetPurpose( GetPrim() ) ) )
+		if (!EnumHasAllFlags(Context->PurposesToLoad, IUsdPrim::GetPurpose(GetPrim())))
 		{
 			return;
 		}
@@ -1281,56 +1300,57 @@ namespace UsdSkelRootTranslatorImpl
 		// To parse all LODs we need to actively switch variant sets to other variants (triggering prim loading/unloading and notices),
 		// which could cause race conditions if other async translation tasks are trying to access those prims
 		ESchemaTranslationLaunchPolicy LaunchPolicy = ESchemaTranslationLaunchPolicy::Async;
-		if ( Context->bAllowInterpretingLODs && HasLODSkinningTargets( pxr::UsdSkelRoot( GetPrim() ) ) )
+		if (Context->bAllowInterpretingLODs && HasLODSkinningTargets(pxr::UsdSkelRoot(GetPrim())))
 		{
 			LaunchPolicy = ESchemaTranslationLaunchPolicy::ExclusiveSync;
 		}
 
 		// Create SkeletalMeshImportData (Async or ExclusiveSync)
-		Do( LaunchPolicy,
-			[ this ]()
-			{
-				// No point in importing blend shapes if the import context doesn't want them
-				UsdUtils::FBlendShapeMap* OutBlendShapes = Context->BlendShapesByPath ? &NewBlendShapes : nullptr;
-
-				pxr::TfToken RenderContextToken = pxr::UsdShadeTokens->universalRenderContext;
-				if ( !Context->RenderContext.IsNone() )
-				{
-					RenderContextToken = UnrealToUsd::ConvertToken( *Context->RenderContext.ToString() ).Get();
-				}
-
-				pxr::TfToken MaterialPurposeToken = pxr::UsdShadeTokens->allPurpose;
-				if ( !Context->MaterialPurpose.IsNone() )
-				{
-					MaterialPurposeToken = UnrealToUsd::ConvertToken( *Context->MaterialPurpose.ToString() ).Get();
-				}
-
-				UsdToUnreal::FUsdMeshConversionOptions Options;
-				Options.TimeCode = Context->Time;
-				Options.RenderContext = RenderContextToken;
-				Options.MaterialPurpose = MaterialPurposeToken;
-				Options.SubdivisionLevel = Context->SubdivisionLevel;
-
-				bool bContinueTaskChain = UsdSkelRootTranslatorImpl::LoadAllSkeletalData(
-					SkeletonCache.Get(),
-					pxr::UsdSkelRoot( GetPrim() ),
-					LODIndexToSkeletalMeshImportData,
-					LODIndexToMaterialInfo,
-					SkeletonBones,
-					SkeletonName,
-					OutBlendShapes,
-					UsedMorphTargetNames,
-					*Context->AssetCache.Get(),
-					Context->bAllowInterpretingLODs,
-					Options
-				);
-
-				return bContinueTaskChain;
-			} );
+		Do(LaunchPolicy,
+		   [this]()
+		   {
+			   // No point in importing blend shapes if the import context doesn't want them
+			   UsdUtils::FBlendShapeMap* OutBlendShapes = Context->BlendShapesByPath ? &NewBlendShapes : nullptr;
+
+			   pxr::TfToken RenderContextToken = pxr::UsdShadeTokens->universalRenderContext;
+			   if (!Context->RenderContext.IsNone())
+			   {
+				   RenderContextToken = UnrealToUsd::ConvertToken(*Context->RenderContext.ToString()).Get();
+			   }
+
+			   pxr::TfToken MaterialPurposeToken = pxr::UsdShadeTokens->allPurpose;
+			   if (!Context->MaterialPurpose.IsNone())
+			   {
+				   MaterialPurposeToken = UnrealToUsd::ConvertToken(*Context->MaterialPurpose.ToString()).Get();
+			   }
+
+			   UsdToUnreal::FUsdMeshConversionOptions Options;
+			   Options.TimeCode = Context->Time;
+			   Options.RenderContext = RenderContextToken;
+			   Options.MaterialPurpose = MaterialPurposeToken;
+			   Options.SubdivisionLevel = Context->SubdivisionLevel;
+
+			   bool bContinueTaskChain = UsdSkelRootTranslatorImpl::LoadAllSkeletalData(
+				   SkeletonCache.Get(),
+				   pxr::UsdSkelRoot(GetPrim()),
+				   LODIndexToSkeletalMeshImportData,
+				   LODIndexToMaterialInfo,
+				   SkeletonBones,
+				   SkeletonName,
+				   OutBlendShapes,
+				   UsedMorphTargetNames,
+				   *Context->AssetCache.Get(),
+				   Context->bAllowInterpretingLODs,
+				   Options
+			   );
+
+			   return bContinueTaskChain;
+		   });
 
 		// Create USkeletalMesh (Main thread)
-		Then( ESchemaTranslationLaunchPolicy::Sync,
-			[ this ]()
+		Then(
+			ESchemaTranslationLaunchPolicy::Sync,
+			[this]()
 			{
 				if (!Context->InfoCache || !Context->AssetCache)
 				{
@@ -1340,13 +1360,13 @@ namespace UsdSkelRootTranslatorImpl
 				UsdUtils::FBlendShapeMap* BlendShapes = Context->BlendShapesByPath ? &NewBlendShapes : nullptr;
 
 				FString SkelRootPath = PrimPath.GetString();
-				FSHAHash SkeletalMeshHash = UsdSkelRootTranslatorImpl::ComputeSHAHash( LODIndexToSkeletalMeshImportData, SkeletonBones, BlendShapes );
+				FSHAHash SkeletalMeshHash = UsdSkelRootTranslatorImpl::ComputeSHAHash(LODIndexToSkeletalMeshImportData, SkeletonBones, BlendShapes);
 				PrefixedSkelMeshHash = UsdUtils::GetAssetHashPrefix(GetPrim(), Context->bReuseIdenticalAssets) + SkeletalMeshHash.ToString();
 
 				USkeletalMesh* SkeletalMesh = Cast<USkeletalMesh>(Context->AssetCache->GetCachedAsset(PrefixedSkelMeshHash));
 
 				bool bIsNew = false;
-				if ( !SkeletalMesh )
+				if (!SkeletalMesh)
 				{
 					bIsNew = true;
 					SkeletalMesh = UsdToUnreal::GetSkeletalMeshFromImportData(
@@ -1354,7 +1374,7 @@ namespace UsdSkelRootTranslatorImpl
 						SkeletonBones,
 						NewBlendShapes,
 						Context->ObjectFlags,
-						*FPaths::GetBaseFilename( SkelRootPath ),
+						*FPaths::GetBaseFilename(SkelRootPath),
 						SkeletonName
 					);
 				}
@@ -1392,8 +1412,8 @@ namespace UsdSkelRootTranslatorImpl
 							SkeletalMesh->UpdateUVChannelData(bRebuildAll);
 						}
 
-						Context->AssetCache->CacheAsset( PrefixedSkelMeshHash, SkeletalMesh );
-						Context->AssetCache->CacheAsset( PrefixedSkelMeshHash + TEXT( "_Skeleton" ), SkeletalMesh->GetSkeleton() );
+						Context->AssetCache->CacheAsset(PrefixedSkelMeshHash, SkeletalMesh);
+						Context->AssetCache->CacheAsset(PrefixedSkelMeshHash + TEXT("_Skeleton"), SkeletalMesh->GetSkeleton());
 
 						// The PreviewSkeletalMesh property on the Skeleton is a soft object path, that is set within
 						// UsdToUnreal::GetSkeletalMeshFromImportData before the SkeletalMesh is part of the cache. When we cache the
@@ -1407,15 +1427,12 @@ namespace UsdSkelRootTranslatorImpl
 					static IConsoleVariable* PhysicsAssetsCvar = IConsoleManager::Get().FindConsoleVariable(TEXT("USD.GeneratePhysicsAssets"));
 					const bool bGeneratePhysicsAssets = PhysicsAssetsCvar && PhysicsAssetsCvar->GetBool();
 
-					if ( bGeneratePhysicsAssets )
+					if (bGeneratePhysicsAssets)
 					{
 						UPhysicsAsset* PhysicsAsset = SkeletalMesh->GetPhysicsAsset();
 						if (!PhysicsAsset)
 						{
-							PhysicsAsset = UsdSkelRootTranslatorImpl::GenerateAndAssignPhysicsAsset(
-								SkeletalMesh,
-								Context->ObjectFlags
-							);
+							PhysicsAsset = UsdSkelRootTranslatorImpl::GenerateAndAssignPhysicsAsset(SkeletalMesh, Context->ObjectFlags);
 
 							if (PhysicsAsset)
 							{
@@ -1431,23 +1448,23 @@ namespace UsdSkelRootTranslatorImpl
 					else
 					{
 						// Actively clear this so that if we toggle the cvar and reload we'll clear our physics assets
-						SkeletalMesh->SetPhysicsAsset( nullptr );
+						SkeletalMesh->SetPhysicsAsset(nullptr);
 					}
 
 					Context->InfoCache->LinkAssetToPrim(PrimPath, SkeletalMesh);
 
 					// Track our Skeleton by the source skeleton prim path
-					if ( USkeleton* Skeleton = SkeletalMesh->GetSkeleton() )
+					if (USkeleton* Skeleton = SkeletalMesh->GetSkeleton())
 					{
 						FScopedUsdAllocs Allocs;
 
-						// Our SkeletonCache should already have been populated by our call to UsdSkelRootTranslatorImpl::LoadAllSkeletalData, so we should
-						// be able to quickly query our skeleton again to find its path
-						std::vector< pxr::UsdSkelBinding > SkeletonBindings;
-						SkeletonCache.Get().ComputeSkelBindings( pxr::UsdSkelRoot( GetPrim() ), &SkeletonBindings, pxr::UsdTraverseInstanceProxies() );
+						// Our SkeletonCache should already have been populated by our call to UsdSkelRootTranslatorImpl::LoadAllSkeletalData, so we
+						// should be able to quickly query our skeleton again to find its path
+						std::vector<pxr::UsdSkelBinding> SkeletonBindings;
+						SkeletonCache.Get().ComputeSkelBindings(pxr::UsdSkelRoot(GetPrim()), &SkeletonBindings, pxr::UsdTraverseInstanceProxies());
 						if (SkeletonBindings.size() > 0 && Context->InfoCache)
 						{
-							pxr::UsdSkelBinding& SkeletonBinding = SkeletonBindings[ 0 ];
+							pxr::UsdSkelBinding& SkeletonBinding = SkeletonBindings[0];
 							const pxr::UsdSkelSkeleton& UsdSkeleton = SkeletonBinding.GetSkeleton();
 							Context->InfoCache->LinkAssetToPrim(PrimPath, Skeleton);
 						}
@@ -1455,125 +1472,125 @@ namespace UsdSkelRootTranslatorImpl
 
 					// We may be reusing a skeletal mesh we got in the cache, but we always need the BlendShapesByPath stored on the
 					// actor to be up-to-date with the Skeletal Mesh that is actually being displayed
-					if ( Context->BlendShapesByPath )
+					if (Context->BlendShapesByPath)
 					{
-						Context->BlendShapesByPath->Append( NewBlendShapes );
+						Context->BlendShapesByPath->Append(NewBlendShapes);
 					}
 				}
 
 				// Continuing even if the mesh is not new as we currently don't add the SkelAnimation info to the mesh hash, so the animations
 				// may have changed
 				return true;
-			} );
+			}
+		);
 
 		// Create AnimBP asset if we need to
-		Then( ESchemaTranslationLaunchPolicy::Sync,
-			[ this ]()
+		Then(
+			ESchemaTranslationLaunchPolicy::Sync,
+			[this]()
 			{
 				UE::FUsdPrim Prim = GetPrim();
 
 				// We need to *also* create the AnimBP within CreateAssets so that we will still generate it even if
 				// we never call Create/UpdateComponents (e.g. importing without importing actors).
-				const bool bPrimHasLiveLinkSchema = UsdUtils::PrimHasSchema( Prim, UnrealIdentifiers::LiveLinkAPI );
-				if ( bPrimHasLiveLinkSchema )
+				const bool bPrimHasLiveLinkSchema = UsdUtils::PrimHasSchema(Prim, UnrealIdentifiers::LiveLinkAPI);
+				if (bPrimHasLiveLinkSchema)
 				{
-					TSharedPtr< FUsdSchemaTranslationContext > ContextPtr = Context;
+					TSharedPtr<FUsdSchemaTranslationContext> ContextPtr = Context;
 
 					// When importing, we can't delay creating the AnimBP to the next frame as that will be after the
 					// import. We don't need to worry about deadlocks though, because we will never *import* as a
 					// response to a USD event: It's always an intentional function call, where the Python GIL is
 					// properly handled
-					if ( ContextPtr && ContextPtr->bIsImporting )
+					if (ContextPtr && ContextPtr->bIsImporting)
 					{
-						UsdSkelRootTranslatorImpl::CreateAnimBlueprint( *ContextPtr.Get(), Prim );
+						UsdSkelRootTranslatorImpl::CreateAnimBlueprint(*ContextPtr.Get(), Prim);
 					}
 					else
 					{
 						// HACK. c.f. the large comment on the analogous timer inside
 						// FUsdSkelRootTranslator::UpdateComponents
-						FTSTicker::GetCoreTicker().AddTicker(
-							FTickerDelegate::CreateLambda( [ContextPtr, Prim]( float Time )
+						FTSTicker::GetCoreTicker().AddTicker(FTickerDelegate::CreateLambda(
+							[ContextPtr, Prim](float Time)
 							{
-								if ( ContextPtr )
+								if (ContextPtr)
 								{
-									UsdSkelRootTranslatorImpl::CreateAnimBlueprint( *ContextPtr.Get(), Prim );
+									UsdSkelRootTranslatorImpl::CreateAnimBlueprint(*ContextPtr.Get(), Prim);
 								}
 
 								// Returning false means this is a one-off, and won't repeat
 								return false;
-							})
-						);
+							}
+						));
 					}
 				}
 
 				return true;
-			});
+			}
+		);
 
 		// Create UAnimSequences (requires a completed USkeleton. Main thread as some steps of the animation compression require it)
-		Then( ESchemaTranslationLaunchPolicy::Sync,
-			[ this ]()
+		Then(
+			ESchemaTranslationLaunchPolicy::Sync,
+			[this]()
 			{
 				if (!Context->bAllowParsingSkeletalAnimations || !Context->InfoCache)
 				{
 					return false;
 				}
 
-				USkeletalMesh* SkeletalMesh = Context->InfoCache->GetSingleAssetForPrim<USkeletalMesh>(
-					PrimPath
-				);
-				if ( !SkeletalMesh )
+				USkeletalMesh* SkeletalMesh = Context->InfoCache->GetSingleAssetForPrim<USkeletalMesh>(PrimPath);
+				if (!SkeletalMesh)
 				{
 					return false;
 				}
 
 				FScopedUsdAllocs Allocs;
 
-				if ( pxr::UsdSkelRoot SkeletonRoot{ GetPrim() } )
+				if (pxr::UsdSkelRoot SkeletonRoot{GetPrim()})
 				{
-					std::vector< pxr::UsdSkelBinding > SkeletonBindings;
-					SkeletonCache.Get().Populate( SkeletonRoot, pxr::UsdTraverseInstanceProxies() );
-					SkeletonCache.Get().ComputeSkelBindings( SkeletonRoot, &SkeletonBindings, pxr::UsdTraverseInstanceProxies() );
+					std::vector<pxr::UsdSkelBinding> SkeletonBindings;
+					SkeletonCache.Get().Populate(SkeletonRoot, pxr::UsdTraverseInstanceProxies());
+					SkeletonCache.Get().ComputeSkelBindings(SkeletonRoot, &SkeletonBindings, pxr::UsdTraverseInstanceProxies());
 
-					for ( const pxr::UsdSkelBinding& Binding : SkeletonBindings )
+					for (const pxr::UsdSkelBinding& Binding : SkeletonBindings)
 					{
 						const pxr::UsdSkelSkeleton& Skeleton = Binding.GetSkeleton();
-						pxr::UsdSkelSkeletonQuery SkelQuery = SkeletonCache.Get().GetSkelQuery( Skeleton );
+						pxr::UsdSkelSkeletonQuery SkelQuery = SkeletonCache.Get().GetSkelQuery(Skeleton);
 						pxr::UsdSkelAnimQuery AnimQuery = SkelQuery.GetAnimQuery();
-						if ( !AnimQuery )
+						if (!AnimQuery)
 						{
 							continue;
 						}
 
 						pxr::UsdPrim SkelAnimationPrim = AnimQuery.GetPrim();
-						if ( !SkelAnimationPrim )
+						if (!SkelAnimationPrim)
 						{
 							continue;
 						}
-						FString SkelAnimationPrimPath = UsdToUnreal::ConvertPath( SkelAnimationPrim.GetPath() );
+						FString SkelAnimationPrimPath = UsdToUnreal::ConvertPath(SkelAnimationPrim.GetPath());
 
 						std::vector<double> JointTimeSamples;
 						std::vector<double> BlendShapeTimeSamples;
-						if ( ( !AnimQuery.GetJointTransformTimeSamples( &JointTimeSamples ) || JointTimeSamples.size() == 0 ) &&
-							( NewBlendShapes.Num() == 0 || ( !AnimQuery.GetBlendShapeWeightTimeSamples( &BlendShapeTimeSamples ) || BlendShapeTimeSamples.size() == 0 ) ) )
+						if ((!AnimQuery.GetJointTransformTimeSamples(&JointTimeSamples) || JointTimeSamples.size() == 0)
+							&& (NewBlendShapes.Num() == 0
+								|| (!AnimQuery.GetBlendShapeWeightTimeSamples(&BlendShapeTimeSamples) || BlendShapeTimeSamples.size() == 0)))
 						{
 							continue;
 						}
 
 						pxr::UsdPrim RootMotionPrim = Context->RootMotionHandling == EUsdRootMotionHandling::UseMotionFromSkelRoot
-							? SkeletonRoot.GetPrim()
-							: Context->RootMotionHandling == EUsdRootMotionHandling::UseMotionFromSkeleton
-								? Skeleton.GetPrim()
-								: pxr::UsdPrim{};
-
-						FSHAHash Hash = UsdSkelRootTranslatorImpl::ComputeSHAHash(
-							SkelQuery,
-							RootMotionPrim,
-							PrefixedSkelMeshHash
-						);
-						FString PrefixedSkelAnimHash = UsdUtils::GetAssetHashPrefix(SkelAnimationPrim, Context->bReuseIdenticalAssets) + Hash.ToString();
+														  ? SkeletonRoot.GetPrim()
+													  : Context->RootMotionHandling == EUsdRootMotionHandling::UseMotionFromSkeleton
+														  ? Skeleton.GetPrim()
+														  : pxr::UsdPrim{};
+
+						FSHAHash Hash = UsdSkelRootTranslatorImpl::ComputeSHAHash(SkelQuery, RootMotionPrim, PrefixedSkelMeshHash);
+						FString PrefixedSkelAnimHash = UsdUtils::GetAssetHashPrefix(SkelAnimationPrim, Context->bReuseIdenticalAssets)
+													   + Hash.ToString();
 						UAnimSequence* AnimSequence = Cast<UAnimSequence>(Context->AssetCache->GetCachedAsset(PrefixedSkelAnimHash));
 
-						if ( !AnimSequence || AnimSequence->GetSkeleton() != SkeletalMesh->GetSkeleton() )
+						if (!AnimSequence || AnimSequence->GetSkeleton() != SkeletalMesh->GetSkeleton())
 						{
 							FScopedUnrealAllocs UEAllocs;
 
@@ -1584,18 +1601,19 @@ namespace UsdSkelRootTranslatorImpl
 							);
 
 							// The UAnimSequence can't be created with the RF_Transactional flag, or else it will be serialized without
-							// Bone/CurveCompressionSettings. Undoing that transaction would call UAnimSequence::Serialize with nullptr values for both, which crashes.
-							// Besides, this particular asset type is only ever created when we import to content folder assets (so never for realtime), and
-							// in that case we don't need it to be transactional anyway
+							// Bone/CurveCompressionSettings. Undoing that transaction would call UAnimSequence::Serialize with nullptr values for
+							// both, which crashes. Besides, this particular asset type is only ever created when we import to content folder assets
+							// (so never for realtime), and in that case we don't need it to be transactional anyway
 							AnimSequence = NewObject<UAnimSequence>(
 								GetTransientPackage(),
 								AnimSequenceName,
 								(Context->ObjectFlags & ~EObjectFlags::RF_Transactional) | EObjectFlags::RF_Transient
 							);
-							AnimSequence->SetSkeleton( SkeletalMesh->GetSkeleton() );
+							AnimSequence->SetSkeleton(SkeletalMesh->GetSkeleton());
 
-							// This is read back in the USDImporter, so that if we ever import this AnimSequence we will always also import the SkeletalMesh for it
-							AnimSequence->SetPreviewMesh( SkeletalMesh );
+							// This is read back in the USDImporter, so that if we ever import this AnimSequence we will always also import the
+							// SkeletalMesh for it
+							AnimSequence->SetPreviewMesh(SkeletalMesh);
 
 							TUsdStore<pxr::VtArray<pxr::UsdSkelSkinningQuery>> SkinningTargets = Binding.GetSkinningTargets();
 							float LayerStartOffsetSeconds = 0.0f;
@@ -1609,15 +1627,18 @@ namespace UsdSkelRootTranslatorImpl
 								&LayerStartOffsetSeconds
 							);
 
-							if (AnimSequence->GetDataModel()->GetNumBoneTracks() != 0 || AnimSequence->GetDataModel()->GetNumberOfFloatCurves() != 0 )
+							if (AnimSequence->GetDataModel()->GetNumBoneTracks() != 0 || AnimSequence->GetDataModel()->GetNumberOfFloatCurves() != 0)
 							{
-								UUsdAnimSequenceAssetUserData* UserData = NewObject<UUsdAnimSequenceAssetUserData>(AnimSequence, TEXT("USDAssetUserData"));
+								UUsdAnimSequenceAssetUserData* UserData = NewObject<UUsdAnimSequenceAssetUserData>(
+									AnimSequence,
+									TEXT("USDAssetUserData")
+								);
 								UserData->PrimPaths = {SkelAnimationPrimPath};
 								UserData->LayerStartOffsetSeconds = LayerStartOffsetSeconds;
 
 								AnimSequence->AddAssetUserData(UserData);
 
-								Context->AssetCache->CacheAsset( PrefixedSkelAnimHash, AnimSequence );
+								Context->AssetCache->CacheAsset(PrefixedSkelAnimHash, AnimSequence);
 							}
 							else
 							{
@@ -1639,11 +1660,12 @@ namespace UsdSkelRootTranslatorImpl
 				}
 
 				return true;
-			});
+			}
+		);
 	}
 
-#endif // WITH_EDITOR
-}
+#endif	  // WITH_EDITOR
+}	 // namespace UsdSkelRootTranslatorImpl
 
 PRAGMA_DISABLE_DEPRECATION_WARNINGS
 void FUsdSkelRootTranslator::CreateAssets()
@@ -1661,11 +1683,11 @@ void FUsdSkelRootTranslator::CreateAssets()
 #if WITH_EDITOR
 	// Importing skeletal meshes actually works in Standalone mode, but we intentionally block it here
 	// to not confuse users as to why it doesn't work at runtime
-	TSharedRef< UsdSkelRootTranslatorImpl::FSkelRootCreateAssetsTaskChain > AssetsTaskChain =
-		MakeShared< UsdSkelRootTranslatorImpl::FSkelRootCreateAssetsTaskChain >( Context, PrimPath );
+	TSharedRef<UsdSkelRootTranslatorImpl::FSkelRootCreateAssetsTaskChain>
+		AssetsTaskChain = MakeShared<UsdSkelRootTranslatorImpl::FSkelRootCreateAssetsTaskChain>(Context, PrimPath);
 
-	Context->TranslatorTasks.Add( MoveTemp( AssetsTaskChain ) );
-#endif // WITH_EDITOR
+	Context->TranslatorTasks.Add(MoveTemp(AssetsTaskChain));
+#endif	  // WITH_EDITOR
 }
 
 USceneComponent* FUsdSkelRootTranslator::CreateComponents()
@@ -1676,13 +1698,8 @@ USceneComponent* FUsdSkelRootTranslator::CreateComponents()
 
 #if WITH_EDITOR
 	// Check if the prim has the GroomBinding schema and setup the component and assets necessary to bind the groom to the SkeletalMesh
-	if (SceneComponent &&
-		SceneComponent->IsA<USkeletalMeshComponent>() &&
-		UsdUtils::PrimHasSchema(GetPrim(), UnrealIdentifiers::GroomBindingAPI) &&
-		Context->AssetCache &&
-		Context->InfoCache &&
-		Context->bAllowParsingGroomAssets
-	)
+	if (SceneComponent && SceneComponent->IsA<USkeletalMeshComponent>() && UsdUtils::PrimHasSchema(GetPrim(), UnrealIdentifiers::GroomBindingAPI)
+		&& Context->AssetCache && Context->InfoCache && Context->bAllowParsingGroomAssets)
 	{
 		UsdGroomTranslatorUtils::CreateGroomBindingAsset(
 			GetPrim(),
@@ -1694,20 +1711,20 @@ USceneComponent* FUsdSkelRootTranslator::CreateComponents()
 
 		// For the groom binding to work, the GroomComponent must be a child of the SceneComponent
 		// so the Context ParentComponent is set to the SceneComponent temporarily
-		TGuardValue< USceneComponent* > ParentComponentGuard{ Context->ParentComponent, SceneComponent };
+		TGuardValue<USceneComponent*> ParentComponentGuard{Context->ParentComponent, SceneComponent};
 		const bool bNeedsActor = false;
 		UGroomComponent* GroomComponent = Cast<UGroomComponent>(CreateComponentsEx({UGroomComponent::StaticClass()}, bNeedsActor));
-		if ( GroomComponent )
+		if (GroomComponent)
 		{
-			UpdateComponents( SceneComponent );
+			UpdateComponents(SceneComponent);
 		}
 	}
-#endif // WITH_EDITOR
+#endif	  // WITH_EDITOR
 
 	return SceneComponent;
 }
 
-void FUsdSkelRootTranslator::UpdateComponents( USceneComponent* SceneComponent )
+void FUsdSkelRootTranslator::UpdateComponents(USceneComponent* SceneComponent)
 {
 	USkeletalMeshComponent* SkeletalMeshComponent = Cast<USkeletalMeshComponent>(SceneComponent);
 	if (!SkeletalMeshComponent)
@@ -1719,17 +1736,17 @@ void FUsdSkelRootTranslator::UpdateComponents( USceneComponent* SceneComponent )
 
 	UE::FUsdPrim Prim = GetPrim();
 
-	const bool bPrimHasLiveLinkSchema = UsdUtils::PrimHasSchema( Prim, UnrealIdentifiers::LiveLinkAPI );
+	const bool bPrimHasLiveLinkSchema = UsdUtils::PrimHasSchema(Prim, UnrealIdentifiers::LiveLinkAPI);
 
 	bool bPrimHasLiveLinkEnabled = bPrimHasLiveLinkSchema;
-	if ( bPrimHasLiveLinkSchema )
+	if (bPrimHasLiveLinkSchema)
 	{
 		FScopedUsdAllocs Allocs;
 
-		if ( pxr::UsdAttribute Attr = pxr::UsdPrim{ Prim }.GetAttribute( UnrealIdentifiers::UnrealLiveLinkEnabled ) )
+		if (pxr::UsdAttribute Attr = pxr::UsdPrim{Prim}.GetAttribute(UnrealIdentifiers::UnrealLiveLinkEnabled))
 		{
 			bool bEnabled = true;
-			if ( Attr.Get( &bEnabled ) )
+			if (Attr.Get(&bEnabled))
 			{
 				bPrimHasLiveLinkEnabled = bEnabled;
 			}
@@ -1739,50 +1756,44 @@ void FUsdSkelRootTranslator::UpdateComponents( USceneComponent* SceneComponent )
 	SkeletalMeshComponent->Modify();
 
 	SkeletalMeshComponent->SetAnimationMode(
-		bPrimHasLiveLinkEnabled
-			? EAnimationMode::AnimationBlueprint
-			: Context->bSequencerIsAnimating
-				? EAnimationMode::AnimationCustomMode
-				: EAnimationMode::AnimationSingleNode
+		bPrimHasLiveLinkEnabled			 ? EAnimationMode::AnimationBlueprint
+		: Context->bSequencerIsAnimating ? EAnimationMode::AnimationCustomMode
+										 : EAnimationMode::AnimationSingleNode
 	);
 
-	UE::FUsdPrim SkelAnimPrim = UsdUtils::FindFirstAnimationSource( Prim );
+	UE::FUsdPrim SkelAnimPrim = UsdUtils::FindFirstAnimationSource(Prim);
 	if (SkelAnimPrim)
 	{
-		UAnimSequence* TargetAnimSequence = Context->InfoCache->GetSingleAssetForPrim<UAnimSequence>(
-			Prim.GetPrimPath()
-		);
-		if ( TargetAnimSequence != SkeletalMeshComponent->AnimationData.AnimToPlay )
+		UAnimSequence* TargetAnimSequence = Context->InfoCache->GetSingleAssetForPrim<UAnimSequence>(Prim.GetPrimPath());
+		if (TargetAnimSequence != SkeletalMeshComponent->AnimationData.AnimToPlay)
 		{
 			SkeletalMeshComponent->AnimationData.AnimToPlay = TargetAnimSequence;
 			SkeletalMeshComponent->AnimationData.bSavedLooping = false;
 			SkeletalMeshComponent->AnimationData.bSavedPlaying = false;
-			SkeletalMeshComponent->SetAnimation( TargetAnimSequence );
+			SkeletalMeshComponent->SetAnimation(TargetAnimSequence);
 		}
 	}
 
-	Super::UpdateComponents( SceneComponent );
+	Super::UpdateComponents(SceneComponent);
 
 	// We always want this, but need to be registered for this to work (Super::UpdateComponents should register us)
 	const bool bNewUpdateState = true;
-	SkeletalMeshComponent->SetUpdateAnimationInEditor( bNewUpdateState );
+	SkeletalMeshComponent->SetUpdateAnimationInEditor(bNewUpdateState);
 
 #if WITH_EDITOR
 	// Re-set the skeletal mesh if we created a new one (maybe the hash changed, a skinned UsdGeomMesh was hidden, etc.)
-	USkeletalMesh* TargetSkeletalMesh = Context->InfoCache->GetSingleAssetForPrim<USkeletalMesh>(
-		PrimPath
-	);
-	if ( SkeletalMeshComponent->GetSkeletalMeshAsset() != TargetSkeletalMesh )
+	USkeletalMesh* TargetSkeletalMesh = Context->InfoCache->GetSingleAssetForPrim<USkeletalMesh>(PrimPath);
+	if (SkeletalMeshComponent->GetSkeletalMeshAsset() != TargetSkeletalMesh)
 	{
 		SkeletalMeshComponent->SetSkeletalMesh(TargetSkeletalMesh);
 
 		// Handle material overrides
-		if ( TargetSkeletalMesh )
+		if (TargetSkeletalMesh)
 		{
 			TArray<UMaterialInterface*> ExistingAssignments;
-			for ( FSkeletalMaterial& SkeletalMaterial : TargetSkeletalMesh->GetMaterials())
+			for (FSkeletalMaterial& SkeletalMaterial : TargetSkeletalMesh->GetMaterials())
 			{
-				ExistingAssignments.Add( SkeletalMaterial.MaterialInterface );
+				ExistingAssignments.Add(SkeletalMaterial.MaterialInterface);
 			}
 
 			UsdSkelRootTranslatorImpl::SetMaterialOverrides(
@@ -1801,14 +1812,14 @@ void FUsdSkelRootTranslator::UpdateComponents( USceneComponent* SceneComponent )
 		}
 	}
 
-	if ( bPrimHasLiveLinkSchema )
+	if (bPrimHasLiveLinkSchema)
 	{
-		TSharedPtr< FUsdSchemaTranslationContext > ContextPtr = Context;
-		TStrongObjectPtr< USkeletalMeshComponent > PinnedSkelMeshComponent{ SkeletalMeshComponent };
+		TSharedPtr<FUsdSchemaTranslationContext> ContextPtr = Context;
+		TStrongObjectPtr<USkeletalMeshComponent> PinnedSkelMeshComponent{SkeletalMeshComponent};
 
-		if ( ContextPtr && ContextPtr->bIsImporting )
+		if (ContextPtr && ContextPtr->bIsImporting)
 		{
-			UsdSkelRootTranslatorImpl::UpdateLiveLinkProperties( *ContextPtr.Get(), PinnedSkelMeshComponent.Get(), Prim );
+			UsdSkelRootTranslatorImpl::UpdateLiveLinkProperties(*ContextPtr.Get(), PinnedSkelMeshComponent.Get(), Prim);
 		}
 		else
 		{
@@ -1819,36 +1830,36 @@ void FUsdSkelRootTranslator::UpdateComponents( USceneComponent* SceneComponent )
 			// a deadlock, as the game thread still holds the GIL and a background reference collector thread may want
 			// to acquire it too. What this does is run this part of the UpdateComponents on tick, that has a callstack
 			// that doesn't originate from Python, and so doesn't have the GIL locked
-			FTSTicker::GetCoreTicker().AddTicker(
-				FTickerDelegate::CreateLambda( [ContextPtr, PinnedSkelMeshComponent, Prim]( float Time )
+			FTSTicker::GetCoreTicker().AddTicker(FTickerDelegate::CreateLambda(
+				[ContextPtr, PinnedSkelMeshComponent, Prim](float Time)
 				{
-					if ( ContextPtr )
+					if (ContextPtr)
 					{
-						UsdSkelRootTranslatorImpl::UpdateLiveLinkProperties( *ContextPtr.Get(), PinnedSkelMeshComponent.Get(), Prim );
+						UsdSkelRootTranslatorImpl::UpdateLiveLinkProperties(*ContextPtr.Get(), PinnedSkelMeshComponent.Get(), Prim);
 					}
 
 					// Returning false means this is a one-off, and won't repeat
 					return false;
-				})
-			);
+				}
+			));
 		}
 	}
 
 	// Update the animation state
 	// Don't try animating ourselves if the sequencer is animating as it will just overwrite the animation state on next
 	// tick anyway, and all this would do is lead to flickering and other issues
-	if ( !Context->bSequencerIsAnimating && SkeletalMeshComponent->GetSkeletalMeshAsset() && !bPrimHasLiveLinkEnabled )
+	if (!Context->bSequencerIsAnimating && SkeletalMeshComponent->GetSkeletalMeshAsset() && !bPrimHasLiveLinkEnabled)
 	{
-		if ( UAnimSequence* AnimSequence = Cast<UAnimSequence>( SkeletalMeshComponent->AnimationData.AnimToPlay.Get() ) )
+		if (UAnimSequence* AnimSequence = Cast<UAnimSequence>(SkeletalMeshComponent->AnimationData.AnimToPlay.Get()))
 		{
 			UE::FSdfLayerOffset CombinedOffset;
-			if ( !SkelAnimPrim )
+			if (!SkelAnimPrim)
 			{
-				SkelAnimPrim = UsdUtils::FindFirstAnimationSource( Prim );
+				SkelAnimPrim = UsdUtils::FindFirstAnimationSource(Prim);
 			}
-			if ( SkelAnimPrim )
+			if (SkelAnimPrim)
 			{
-				CombinedOffset = UsdUtils::GetPrimToStageOffset( SkelAnimPrim );
+				CombinedOffset = UsdUtils::GetPrimToStageOffset(SkelAnimPrim);
 			}
 
 			double LayerStartOffsetSeconds = 0.0f;
@@ -1858,7 +1869,7 @@ void FUsdSkelRootTranslator::UpdateComponents( USceneComponent* SceneComponent )
 			}
 
 			// Always change the mode here because the sequencer will change it back to AnimationCustomMode when animating
-			SkeletalMeshComponent->SetAnimationMode( EAnimationMode::AnimationSingleNode );
+			SkeletalMeshComponent->SetAnimationMode(EAnimationMode::AnimationSingleNode);
 
 			// Part of the CombinedOffset will be due to a framerate difference. We don't care about that part here though, so remove it
 			const double TimeCodesPerSecondDifference = Context->Stage.GetTimeCodesPerSecond() / AnimSequence->ImportFileFramerate;
@@ -1868,11 +1879,11 @@ void FUsdSkelRootTranslator::UpdateComponents( USceneComponent* SceneComponent )
 			// asset may have been created when the stage had a different framesPerSecond (and was reused by the assets cache)
 			// Use the import framerate here because we will need to change the sampling framerate of the sequence in order to get it
 			// to match the target duration in seconds and the number of source frames.
-			const double LayerTimeCode = ( ( Context->Time - CombinedOffset.Offset ) / CombinedOffset.Scale );
+			const double LayerTimeCode = ((Context->Time - CombinedOffset.Offset) / CombinedOffset.Scale);
 			const double AnimSequenceTime = LayerTimeCode / AnimSequence->ImportFileFramerate;
-			SkeletalMeshComponent->SetPosition( AnimSequenceTime - LayerStartOffsetSeconds );
+			SkeletalMeshComponent->SetPosition(AnimSequenceTime - LayerStartOffsetSeconds);
 
-			SkeletalMeshComponent->TickAnimation( 0.f, false );
+			SkeletalMeshComponent->TickAnimation(0.f, false);
 			SkeletalMeshComponent->RefreshBoneTransforms();
 			SkeletalMeshComponent->RefreshFollowerComponents();
 			SkeletalMeshComponent->UpdateComponentToWorld();
@@ -1885,9 +1896,9 @@ void FUsdSkelRootTranslator::UpdateComponents( USceneComponent* SceneComponent )
 	// If the prim has a GroomBinding schema, apply the target groom to its associated GroomComponent
 	if (UsdUtils::PrimHasSchema(Prim, UnrealIdentifiers::GroomBindingAPI) && Context->InfoCache)
 	{
-		UsdGroomTranslatorUtils::SetGroomFromPrim( Prim, *Context->InfoCache, SceneComponent );
+		UsdGroomTranslatorUtils::SetGroomFromPrim(Prim, *Context->InfoCache, SceneComponent);
 	}
-#endif // WITH_EDITOR
+#endif	  // WITH_EDITOR
 }
 
 bool FUsdSkelRootTranslator::CollapsesChildren(ECollapsingType CollapsingType) const
@@ -1956,7 +1967,7 @@ TSet<UE::FSdfPath> FUsdSkelRootTranslator::CollectAuxiliaryPrims() const
 		{
 			pxr::UsdSkelCache SkeletonCache;
 
-			std::vector< pxr::UsdSkelBinding > SkeletonBindings;
+			std::vector<pxr::UsdSkelBinding> SkeletonBindings;
 			SkeletonCache.Populate(SkeletonRoot, pxr::UsdTraverseInstanceProxies());
 			SkeletonCache.ComputeSkelBindings(SkeletonRoot, &SkeletonBindings, pxr::UsdTraverseInstanceProxies());
 
@@ -1984,4 +1995,4 @@ PRAGMA_ENABLE_DEPRECATION_WARNINGS
 
 #undef LOCTEXT_NAMESPACE
 
-#endif // #if USE_USD_SDK
+#endif	  // #if USE_USD_SDK
