USD: Reformat USDSchemas module.

diff --git a/Engine/Plugins/Importers/USDImporter/Source/USDSchemas/Private/USDSkelSkeletonTranslator.cpp b/Engine/Plugins/Importers/USDImporter/Source/USDSchemas/Private/USDSkelSkeletonTranslator.cpp
--- a/Engine/Plugins/Importers/USDImporter/Source/USDSchemas/Private/USDSkelSkeletonTranslator.cpp
+++ b/Engine/Plugins/Importers/USDImporter/Source/USDSchemas/Private/USDSkelSkeletonTranslator.cpp
@@ -50,14 +50,14 @@
 #endif	  // WITH_EDITOR
 
 #include "USDIncludesStart.h"
-	#include "pxr/usd/usd/primRange.h"
-	#include "pxr/usd/usdGeom/mesh.h"
-	#include "pxr/usd/usdSkel/binding.h"
-	#include "pxr/usd/usdSkel/bindingAPI.h"
-	#include "pxr/usd/usdSkel/blendShapeQuery.h"
-	#include "pxr/usd/usdSkel/cache.h"
-	#include "pxr/usd/usdSkel/root.h"
-	#include "pxr/usd/usdSkel/skeletonQuery.h"
+#include "pxr/usd/usd/primRange.h"
+#include "pxr/usd/usdGeom/mesh.h"
+#include "pxr/usd/usdSkel/binding.h"
+#include "pxr/usd/usdSkel/bindingAPI.h"
+#include "pxr/usd/usdSkel/blendShapeQuery.h"
+#include "pxr/usd/usdSkel/cache.h"
+#include "pxr/usd/usdSkel/root.h"
+#include "pxr/usd/usdSkel/skeletonQuery.h"
 #include "USDIncludesEnd.h"
 
 #define LOCTEXT_NAMESPACE "UsdSkelRoot"
@@ -530,21 +530,17 @@ namespace UsdSkelSkeletonTranslatorImpl
 			}
 		}
 
-		TFunction<bool(const pxr::UsdGeomMesh&, int32)> ConvertLOD =
-		[
-			&LODIndexToSkeletalMeshImportDataMap,
-			&LODIndexToMaterialInfoMap,
-			&LODMetadata,
-			&InSkelCache,
-			&Stage,
-			&SkeletonPrimPath,
-			&InOutUsedMorphTargetNames,
-			OutBlendShapes,
-			&StageInfo,
-			Options,
-			&MetadataOptions
-		]
-		(const pxr::UsdGeomMesh& LODMesh, int32 LODIndex)
+		TFunction<bool(const pxr::UsdGeomMesh&, int32)> ConvertLOD = [&LODIndexToSkeletalMeshImportDataMap,
+																	  &LODIndexToMaterialInfoMap,
+																	  &LODMetadata,
+																	  &InSkelCache,
+																	  &Stage,
+																	  &SkeletonPrimPath,
+																	  &InOutUsedMorphTargetNames,
+																	  OutBlendShapes,
+																	  &StageInfo,
+																	  Options,
+																	  &MetadataOptions](const pxr::UsdGeomMesh& LODMesh, int32 LODIndex)
 		{
 			// Construct this and SkinningQuery every time so as to survive the prim reference invalidation caused by flipping LODs
 			pxr::UsdSkelSkeletonQuery SkeletonQuery = InSkelCache.GetSkelQuery(pxr::UsdSkelSkeleton{Stage->GetPrimAtPath(SkeletonPrimPath)});
@@ -753,46 +749,48 @@ namespace UsdSkelSkeletonTranslatorImpl
 		TMap<int32, UsdUtils::FUsdPrimMaterialAssignmentInfo> LODIndexToMaterialInfoMap;
 		TMap<int32, TSet<UsdUtils::FUsdPrimMaterialSlot>> CombinedSlotsForLODIndex;
 		TFunction<bool(const pxr::UsdGeomMesh&, int32)> IterateLODsLambda =
-			[&LODIndexToMaterialInfoMap, &CombinedSlotsForLODIndex, Time, RenderContextToken, MaterialPurposeToken, PurposesToLoad]
-			(const pxr::UsdGeomMesh& LODMesh, int32 LODIndex)
+			[&LODIndexToMaterialInfoMap, &CombinedSlotsForLODIndex, Time, RenderContextToken, MaterialPurposeToken, PurposesToLoad](
+				const pxr::UsdGeomMesh& LODMesh,
+				int32 LODIndex
+			)
+		{
+			if (LODMesh && LODMesh.ComputeVisibility() == pxr::UsdGeomTokens->invisible)
 			{
-				if (LODMesh && LODMesh.ComputeVisibility() == pxr::UsdGeomTokens->invisible)
-				{
-					return true;
-				}
-
-				// Ignore prims with disabled purposes: We need to match the material slot ordering that was used
-				// to generate the mesh in the first place, so this is important
-				if (!EnumHasAllFlags(PurposesToLoad, IUsdPrim::GetPurpose(LODMesh.GetPrim())))
-				{
-					return true;
-				}
+				return true;
+			}
 
-				TArray<UsdUtils::FUsdPrimMaterialSlot>& CombinedLODSlots = LODIndexToMaterialInfoMap.FindOrAdd(LODIndex).Slots;
-				TSet<UsdUtils::FUsdPrimMaterialSlot>& CombinedLODSlotsSet = CombinedSlotsForLODIndex.FindOrAdd(LODIndex);
+			// Ignore prims with disabled purposes: We need to match the material slot ordering that was used
+			// to generate the mesh in the first place, so this is important
+			if (!EnumHasAllFlags(PurposesToLoad, IUsdPrim::GetPurpose(LODMesh.GetPrim())))
+			{
+				return true;
+			}
 
-				const bool bProvideMaterialIndices = false;	   // We have no use for material indices and it can be slow to retrieve, as it will iterate
-															// all faces
-				UsdUtils::FUsdPrimMaterialAssignmentInfo LocalInfo = UsdUtils::GetPrimMaterialAssignments(
-					LODMesh.GetPrim(),
-					pxr::UsdTimeCode(Time),
-					bProvideMaterialIndices,
-					RenderContextToken,
-					MaterialPurposeToken
-				);
+			TArray<UsdUtils::FUsdPrimMaterialSlot>& CombinedLODSlots = LODIndexToMaterialInfoMap.FindOrAdd(LODIndex).Slots;
+			TSet<UsdUtils::FUsdPrimMaterialSlot>& CombinedLODSlotsSet = CombinedSlotsForLODIndex.FindOrAdd(LODIndex);
+
+			const bool bProvideMaterialIndices = false;	   // We have no use for material indices and it can be slow to retrieve, as it will iterate
+														   // all faces
+			UsdUtils::FUsdPrimMaterialAssignmentInfo LocalInfo = UsdUtils::GetPrimMaterialAssignments(
+				LODMesh.GetPrim(),
+				pxr::UsdTimeCode(Time),
+				bProvideMaterialIndices,
+				RenderContextToken,
+				MaterialPurposeToken
+			);
 
-				// Combine material slots in the same order that UsdToUnreal::ConvertSkinnedMesh does
-				for (UsdUtils::FUsdPrimMaterialSlot& LocalSlot : LocalInfo.Slots)
+			// Combine material slots in the same order that UsdToUnreal::ConvertSkinnedMesh does
+			for (UsdUtils::FUsdPrimMaterialSlot& LocalSlot : LocalInfo.Slots)
+			{
+				if (!CombinedLODSlotsSet.Contains(LocalSlot))
 				{
-					if (!CombinedLODSlotsSet.Contains(LocalSlot))
-					{
-						CombinedLODSlots.Add(LocalSlot);
-						CombinedLODSlotsSet.Add(LocalSlot);
-					}
+					CombinedLODSlots.Add(LocalSlot);
+					CombinedLODSlotsSet.Add(LocalSlot);
 				}
+			}
 
-				return true;
-			};
+			return true;
+		};
 
 		TSet<FString> ProcessedLODParentPaths;
 
@@ -1447,59 +1445,57 @@ namespace UsdSkelSkeletonTranslatorImpl
 																	 : ESchemaTranslationLaunchPolicy::Async;
 
 		// Create SkeletalMeshImportData (Async or ExclusiveSync)
-		Do(
-			LaunchPolicy,
-			[this, bTryLODParsing]()
-			{
-				RefreshSkelReferencesIfNeeded(
-					ClosestParentSkelRoot.Get(),
-					pxr::UsdSkelSkeleton{GetSkeletonPrim()},
-					SkelCache.Get(),
-					SkeletonBinding.Get(),
-					SkeletonQuery.Get()
-				);
-
-				// No point in importing blend shapes if the import context doesn't want them
-				UsdUtils::FBlendShapeMap* OutBlendShapes = Context->BlendShapesByPath ? &NewBlendShapes : nullptr;
-
-				pxr::TfToken RenderContextToken = pxr::UsdShadeTokens->universalRenderContext;
-				if (!Context->RenderContext.IsNone())
-				{
-					RenderContextToken = UnrealToUsd::ConvertToken(*Context->RenderContext.ToString()).Get();
-				}
-
-				pxr::TfToken MaterialPurposeToken = pxr::UsdShadeTokens->allPurpose;
-				if (!Context->MaterialPurpose.IsNone())
-				{
-					MaterialPurposeToken = UnrealToUsd::ConvertToken(*Context->MaterialPurpose.ToString()).Get();
-				}
-
-				UsdToUnreal::FUsdMeshConversionOptions Options;
-				Options.TimeCode = Context->Time;
-				Options.RenderContext = RenderContextToken;
-				Options.MaterialPurpose = MaterialPurposeToken;
-				Options.SubdivisionLevel = Context->SubdivisionLevel;
-				Options.PurposesToLoad = Context->PurposesToLoad;
-				Options.bMergeIdenticalMaterialSlots = Context->bMergeIdenticalMaterialSlots;
-
-				const bool bContinueTaskChain = UsdSkelSkeletonTranslatorImpl::LoadAllSkeletalData(
-					SkeletonBinding.Get(),
-					SkelCache.Get(),
-					LODIndexToSkeletalMeshImportData,
-					LODIndexToMaterialInfo,
-					LODMetadata,
-					SkeletonBones,
-					SkeletonName,
-					OutBlendShapes,
-					UsedMorphTargetNames,
-					bTryLODParsing,
-					Options,
-					Context->MetadataOptions
-				);
-
-				return bContinueTaskChain;
-			}
-		);
+		Do(LaunchPolicy,
+		   [this, bTryLODParsing]()
+		   {
+			   RefreshSkelReferencesIfNeeded(
+				   ClosestParentSkelRoot.Get(),
+				   pxr::UsdSkelSkeleton{GetSkeletonPrim()},
+				   SkelCache.Get(),
+				   SkeletonBinding.Get(),
+				   SkeletonQuery.Get()
+			   );
+
+			   // No point in importing blend shapes if the import context doesn't want them
+			   UsdUtils::FBlendShapeMap* OutBlendShapes = Context->BlendShapesByPath ? &NewBlendShapes : nullptr;
+
+			   pxr::TfToken RenderContextToken = pxr::UsdShadeTokens->universalRenderContext;
+			   if (!Context->RenderContext.IsNone())
+			   {
+				   RenderContextToken = UnrealToUsd::ConvertToken(*Context->RenderContext.ToString()).Get();
+			   }
+
+			   pxr::TfToken MaterialPurposeToken = pxr::UsdShadeTokens->allPurpose;
+			   if (!Context->MaterialPurpose.IsNone())
+			   {
+				   MaterialPurposeToken = UnrealToUsd::ConvertToken(*Context->MaterialPurpose.ToString()).Get();
+			   }
+
+			   UsdToUnreal::FUsdMeshConversionOptions Options;
+			   Options.TimeCode = Context->Time;
+			   Options.RenderContext = RenderContextToken;
+			   Options.MaterialPurpose = MaterialPurposeToken;
+			   Options.SubdivisionLevel = Context->SubdivisionLevel;
+			   Options.PurposesToLoad = Context->PurposesToLoad;
+			   Options.bMergeIdenticalMaterialSlots = Context->bMergeIdenticalMaterialSlots;
+
+			   const bool bContinueTaskChain = UsdSkelSkeletonTranslatorImpl::LoadAllSkeletalData(
+				   SkeletonBinding.Get(),
+				   SkelCache.Get(),
+				   LODIndexToSkeletalMeshImportData,
+				   LODIndexToMaterialInfo,
+				   LODMetadata,
+				   SkeletonBones,
+				   SkeletonName,
+				   OutBlendShapes,
+				   UsedMorphTargetNames,
+				   bTryLODParsing,
+				   Options,
+				   Context->MetadataOptions
+			   );
+
+			   return bContinueTaskChain;
+		   });
 
 		// Create USkeletalMesh (Main thread)
 		Then(
@@ -1633,10 +1629,7 @@ namespace UsdSkelSkeletonTranslatorImpl
 						UPhysicsAsset* PhysicsAsset = SkeletalMesh->GetPhysicsAsset();
 						if (!PhysicsAsset)
 						{
-							PhysicsAsset = UsdSkelSkeletonTranslatorImpl::GenerateAndAssignPhysicsAsset(
-								SkeletalMesh,
-								Context->ObjectFlags
-							);
+							PhysicsAsset = UsdSkelSkeletonTranslatorImpl::GenerateAndAssignPhysicsAsset(SkeletalMesh, Context->ObjectFlags);
 
 							if (PhysicsAsset)
 							{
@@ -1770,7 +1763,7 @@ namespace UsdSkelSkeletonTranslatorImpl
 				}
 
 				pxr::UsdPrim RootMotionPrim;
-				switch(Context->RootMotionHandling)
+				switch (Context->RootMotionHandling)
 				{
 					case EUsdRootMotionHandling::UseMotionFromSkelRoot:
 					{
@@ -1792,11 +1785,7 @@ namespace UsdSkelSkeletonTranslatorImpl
 					}
 				}
 
-				FSHAHash Hash = UsdSkelSkeletonTranslatorImpl::ComputeSHAHash(
-					SkeletonQuery.Get(),
-					RootMotionPrim,
-					PrefixedSkelMeshHash
-				);
+				FSHAHash Hash = UsdSkelSkeletonTranslatorImpl::ComputeSHAHash(SkeletonQuery.Get(), RootMotionPrim, PrefixedSkelMeshHash);
 				FString PrefixedSkelAnimHash = UsdUtils::GetAssetHashPrefix(SkelAnimationPrim, Context->bReuseIdenticalAssets) + Hash.ToString();
 				UAnimSequence* AnimSequence = Cast<UAnimSequence>(Context->AssetCache->GetCachedAsset(PrefixedSkelAnimHash));
 
@@ -1837,7 +1826,8 @@ namespace UsdSkelSkeletonTranslatorImpl
 						&LayerStartOffsetSeconds
 					);
 
-					if (bSuccess && (AnimSequence->GetDataModel()->GetNumBoneTracks() != 0 || AnimSequence->GetDataModel()->GetNumberOfFloatCurves() != 0))
+					if (bSuccess
+						&& (AnimSequence->GetDataModel()->GetNumBoneTracks() != 0 || AnimSequence->GetDataModel()->GetNumberOfFloatCurves() != 0))
 					{
 						if (UUsdAnimSequenceAssetUserData* UserData = UsdUtils::GetOrCreateAssetUserData<UUsdAnimSequenceAssetUserData>(AnimSequence))
 						{
@@ -1923,9 +1913,7 @@ USceneComponent* FUsdSkelSkeletonTranslator::CreateComponents()
 	}
 
 	// Check if the prim has the GroomBinding schema and setup the component and assets necessary to bind the groom to the SkeletalMesh
-	if (Context->bAllowParsingGroomAssets &&
-		Context->AssetCache &&
-		Context->InfoCache)
+	if (Context->bAllowParsingGroomAssets && Context->AssetCache && Context->InfoCache)
 	{
 		UE::FUsdPrim PrimWithSchema;
 		if (UsdUtils::PrimHasSchema(SkeletonPrim, UnrealIdentifiers::GroomBindingAPI))
diff --git a/Engine/Plugins/Importers/USDImporter/Source/USDSchemas/Public/USDGeomCameraTranslator.h b/Engine/Plugins/Importers/USDImporter/Source/USDSchemas/Public/USDGeomCameraTranslator.h
--- a/Engine/Plugins/Importers/USDImporter/Source/USDSchemas/Public/USDGeomCameraTranslator.h
+++ b/Engine/Plugins/Importers/USDImporter/Source/USDSchemas/Public/USDGeomCameraTranslator.h
@@ -6,24 +6,16 @@
 
 #if USE_USD_SDK
 
-#include "USDIncludesStart.h"
-	#include "pxr/pxr.h"
-#include "USDIncludesEnd.h"
-
-PXR_NAMESPACE_OPEN_SCOPE
-	class UsdGeomCamera;
-PXR_NAMESPACE_CLOSE_SCOPE
-
 class USDSCHEMAS_API FUsdGeomCameraTranslator : public FUsdGeomXformableTranslator
 {
 public:
 	using FUsdGeomXformableTranslator::FUsdGeomXformableTranslator;
 
 	virtual USceneComponent* CreateComponents() override;
-	virtual void UpdateComponents( USceneComponent* SceneComponent ) override;
+	virtual void UpdateComponents(USceneComponent* SceneComponent) override;
 
-	virtual bool CollapsesChildren( ECollapsingType CollapsingType ) const override;
-	virtual bool CanBeCollapsed( ECollapsingType CollapsingType ) const override;
+	virtual bool CollapsesChildren(ECollapsingType CollapsingType) const override;
+	virtual bool CanBeCollapsed(ECollapsingType CollapsingType) const override;
 };
 
-#endif // #if USE_USD_SDK
+#endif	  // #if USE_USD_SDK
diff --git a/Engine/Plugins/Importers/USDImporter/Source/USDSchemas/Public/USDGeomMeshTranslator.h b/Engine/Plugins/Importers/USDImporter/Source/USDSchemas/Public/USDGeomMeshTranslator.h
--- a/Engine/Plugins/Importers/USDImporter/Source/USDSchemas/Public/USDGeomMeshTranslator.h
+++ b/Engine/Plugins/Importers/USDImporter/Source/USDSchemas/Public/USDGeomMeshTranslator.h
@@ -7,23 +7,14 @@
 
 #include "UsdWrappers/SdfPath.h"
 #include "UsdWrappers/UsdPrim.h"
-#include "UsdWrappers/UsdStage.h"
 
 #include "MeshDescription.h"
 
 #if USE_USD_SDK
 
-#include "USDIncludesStart.h"
-	#include "pxr/pxr.h"
-#include "USDIncludesEnd.h"
-
 class UStaticMesh;
 class FStaticMeshComponentRecreateRenderStateContext;
 
-PXR_NAMESPACE_OPEN_SCOPE
-	class UsdGeomMesh;
-PXR_NAMESPACE_CLOSE_SCOPE
-
 namespace UsdUtils
 {
 	struct FUsdPrimMaterialAssignmentInfo;
@@ -33,7 +24,7 @@ class FBuildStaticMeshTaskChain : public FUsdSchemaTranslatorTaskChain
 {
 public:
 	explicit FBuildStaticMeshTaskChain(
-		const TSharedRef< FUsdSchemaTranslationContext >& InContext,
+		const TSharedRef<FUsdSchemaTranslationContext>& InContext,
 		const UE::FSdfPath& InPrimPath,
 		const TOptional<UE::FSdfPath>& AlternativePrimToLinkAssetsTo = {}
 	);
@@ -42,7 +33,7 @@ protected:
 	// Inputs
 	// When multiple meshes are collapsed together, this Schema might not be the same as the Context schema, which is the root schema
 	UE::FSdfPath PrimPath;
-	TSharedRef< FUsdSchemaTranslationContext > Context;
+	TSharedRef<FUsdSchemaTranslationContext> Context;
 	TArray<FMeshDescription> LODIndexToMeshDescription;
 	TArray<UsdUtils::FUsdPrimMaterialAssignmentInfo> LODIndexToMaterialInfo;
 	// We collect metadata early (during LOD parsing) so that we don't have to flip through
@@ -58,7 +49,10 @@ protected:
 	bool bCollectedMetadata = false;
 
 protected:
-	UE::FUsdPrim GetPrim() const { return Context->Stage.GetPrimAtPath( PrimPath ); }
+	UE::FUsdPrim GetPrim() const
+	{
+		return Context->Stage.GetPrimAtPath(PrimPath);
+	}
 
 	virtual void SetupTasks();
 };
@@ -67,7 +61,7 @@ class FGeomMeshCreateAssetsTaskChain : public FBuildStaticMeshTaskChain
 {
 public:
 	explicit FGeomMeshCreateAssetsTaskChain(
-		const TSharedRef< FUsdSchemaTranslationContext >& InContext,
+		const TSharedRef<FUsdSchemaTranslationContext>& InContext,
 		const UE::FSdfPath& PrimPath,
 		const TOptional<UE::FSdfPath>& AlternativePrimToLinkAssetsTo = {},
 		const FTransform& AdditionalTransform = FTransform::Identity
@@ -88,15 +82,15 @@ public:
 
 	using FUsdGeomXformableTranslator::FUsdGeomXformableTranslator;
 
-	FUsdGeomMeshTranslator( const FUsdGeomMeshTranslator& Other ) = delete;
-	FUsdGeomMeshTranslator& operator=( const FUsdGeomMeshTranslator& Other ) = delete;
+	FUsdGeomMeshTranslator(const FUsdGeomMeshTranslator& Other) = delete;
+	FUsdGeomMeshTranslator& operator=(const FUsdGeomMeshTranslator& Other) = delete;
 
 	virtual void CreateAssets() override;
 	virtual USceneComponent* CreateComponents() override;
-	virtual void UpdateComponents( USceneComponent* SceneComponent ) override;
+	virtual void UpdateComponents(USceneComponent* SceneComponent) override;
 
-	virtual bool CollapsesChildren( ECollapsingType CollapsingType ) const override;
-	virtual bool CanBeCollapsed( ECollapsingType CollapsingType ) const override;
+	virtual bool CollapsesChildren(ECollapsingType CollapsingType) const override;
+	virtual bool CanBeCollapsed(ECollapsingType CollapsingType) const override;
 
 	virtual TSet<UE::FSdfPath> CollectAuxiliaryPrims() const override;
 
@@ -117,4 +111,4 @@ protected:
 	bool ShouldSkipSkinnablePrim(bool bCheckForComponent = false) const;
 };
 
-#endif // #if USE_USD_SDK
+#endif	  // #if USE_USD_SDK
diff --git a/Engine/Plugins/Importers/USDImporter/Source/USDSchemas/Public/USDGeomXformableTranslator.h b/Engine/Plugins/Importers/USDImporter/Source/USDSchemas/Public/USDGeomXformableTranslator.h
--- a/Engine/Plugins/Importers/USDImporter/Source/USDSchemas/Public/USDGeomXformableTranslator.h
+++ b/Engine/Plugins/Importers/USDImporter/Source/USDSchemas/Public/USDGeomXformableTranslator.h
@@ -8,35 +8,31 @@
 
 #if USE_USD_SDK
 
-#include "USDIncludesStart.h"
-	#include "pxr/pxr.h"
-#include "USDIncludesEnd.h"
-
 class USceneComponent;
 enum class EUsdDrawMode : int32;
 
-PXR_NAMESPACE_OPEN_SCOPE
-	class UsdGeomXformable;
-PXR_NAMESPACE_CLOSE_SCOPE
-
 class USDSCHEMAS_API FUsdGeomXformableTranslator : public FUsdSchemaTranslator
 {
 public:
 	using FUsdSchemaTranslator::FUsdSchemaTranslator;
 
-	explicit FUsdGeomXformableTranslator( TSubclassOf< USceneComponent > InComponentTypeOverride, TSharedRef< FUsdSchemaTranslationContext > InContext, const UE::FUsdTyped& InSchema );
+	explicit FUsdGeomXformableTranslator(
+		TSubclassOf<USceneComponent> InComponentTypeOverride,
+		TSharedRef<FUsdSchemaTranslationContext> InContext,
+		const UE::FUsdTyped& InSchema
+	);
 
 	virtual void CreateAssets() override;
 	virtual USceneComponent* CreateComponents() override;
-	virtual void UpdateComponents( USceneComponent* SceneComponent ) override;
+	virtual void UpdateComponents(USceneComponent* SceneComponent) override;
 
-	virtual bool CollapsesChildren( ECollapsingType CollapsingType ) const override;
-	virtual bool CanBeCollapsed( ECollapsingType CollapsingType ) const override;
+	virtual bool CollapsesChildren(ECollapsingType CollapsingType) const override;
+	virtual bool CanBeCollapsed(ECollapsingType CollapsingType) const override;
 
 	virtual TSet<UE::FSdfPath> CollectAuxiliaryPrims() const override;
 
 	// If the optional parameters are not set, we'll figure them out automatically.
-	USceneComponent* CreateComponentsEx( TOptional< TSubclassOf< USceneComponent > > ComponentType, TOptional< bool > bNeedsActor );
+	USceneComponent* CreateComponentsEx(TOptional<TSubclassOf<USceneComponent>> ComponentType, TOptional<bool> bNeedsActor);
 
 protected:
 	// Creates actors, components and assets in case we have an alt draw mode like bounds or cards.
@@ -48,7 +44,7 @@ protected:
 	void CreateAlternativeDrawModeAssets(EUsdDrawMode DrawMode);
 
 private:
-	TOptional< TSubclassOf< USceneComponent > > ComponentTypeOverride;
+	TOptional<TSubclassOf<USceneComponent>> ComponentTypeOverride;
 };
 
-#endif // #if USE_USD_SDK
+#endif	  // #if USE_USD_SDK
diff --git a/Engine/Plugins/Importers/USDImporter/Source/USDSchemas/Public/USDGroomTranslator.h b/Engine/Plugins/Importers/USDImporter/Source/USDSchemas/Public/USDGroomTranslator.h
--- a/Engine/Plugins/Importers/USDImporter/Source/USDSchemas/Public/USDGroomTranslator.h
+++ b/Engine/Plugins/Importers/USDImporter/Source/USDSchemas/Public/USDGroomTranslator.h
@@ -5,7 +5,6 @@
 #if USE_USD_SDK && WITH_EDITOR
 
 #include "USDGeometryCacheTranslator.h"
-#include "USDGeomXformableTranslator.h"
 
 class USDSCHEMAS_API FUsdGroomTranslator : public FUsdGeometryCacheTranslator
 {
@@ -27,4 +26,4 @@ private:
 	bool IsGroomPrim() const;
 };
 
-#endif // #if USE_USD_SDK
\ No newline at end of file
+#endif	  // #if USE_USD_SDK
diff --git a/Engine/Plugins/Importers/USDImporter/Source/USDSchemas/Public/USDLuxLightTranslator.h b/Engine/Plugins/Importers/USDImporter/Source/USDSchemas/Public/USDLuxLightTranslator.h
--- a/Engine/Plugins/Importers/USDImporter/Source/USDSchemas/Public/USDLuxLightTranslator.h
+++ b/Engine/Plugins/Importers/USDImporter/Source/USDSchemas/Public/USDLuxLightTranslator.h
@@ -6,16 +6,8 @@
 
 #if USE_USD_SDK
 
-#include "USDIncludesStart.h"
-	#include "pxr/pxr.h"
-#include "USDIncludesEnd.h"
-
 class USceneComponent;
 
-PXR_NAMESPACE_OPEN_SCOPE
-class UsdGeomXformable;
-PXR_NAMESPACE_CLOSE_SCOPE
-
 class USDSCHEMAS_API FUsdLuxLightTranslator : public FUsdGeomXformableTranslator
 {
 	using Super = FUsdGeomXformableTranslator;
@@ -25,10 +17,10 @@ public:
 	using FUsdGeomXformableTranslator::FUsdGeomXformableTranslator;
 
 	virtual USceneComponent* CreateComponents() override;
-	virtual void UpdateComponents( USceneComponent* SceneComponent );
+	virtual void UpdateComponents(USceneComponent* SceneComponent);
 
-	virtual bool CollapsesChildren( ECollapsingType CollapsingType ) const override;
-	virtual bool CanBeCollapsed( ECollapsingType CollapsingType ) const override;
+	virtual bool CollapsesChildren(ECollapsingType CollapsingType) const override;
+	virtual bool CanBeCollapsed(ECollapsingType CollapsingType) const override;
 };
 
-#endif // #if USE_USD_SDK
+#endif	  // #if USE_USD_SDK
diff --git a/Engine/Plugins/Importers/USDImporter/Source/USDSchemas/Public/USDShadeMaterialTranslator.h b/Engine/Plugins/Importers/USDImporter/Source/USDSchemas/Public/USDShadeMaterialTranslator.h
--- a/Engine/Plugins/Importers/USDImporter/Source/USDSchemas/Public/USDShadeMaterialTranslator.h
+++ b/Engine/Plugins/Importers/USDImporter/Source/USDSchemas/Public/USDShadeMaterialTranslator.h
@@ -6,10 +6,6 @@
 
 #if USE_USD_SDK
 
-#include "USDIncludesStart.h"
-	#include "pxr/pxr.h"
-#include "USDIncludesEnd.h"
-
 class USDSCHEMAS_API FUsdShadeMaterialTranslator : public FUsdSchemaTranslator
 {
 public:
@@ -17,8 +13,8 @@ public:
 
 	virtual void CreateAssets() override;
 
-	virtual bool CollapsesChildren( ECollapsingType CollapsingType ) const override;
-	virtual bool CanBeCollapsed( ECollapsingType CollapsingType ) const override;
+	virtual bool CollapsesChildren(ECollapsingType CollapsingType) const override;
+	virtual bool CanBeCollapsed(ECollapsingType CollapsingType) const override;
 
 	virtual TSet<UE::FSdfPath> CollectAuxiliaryPrims() const override;
 
@@ -26,4 +22,4 @@ protected:
 	virtual void PostImportMaterial(const FString& PrefixedMaterialHash, UMaterialInterface* ImportedMaterial);
 };
 
-#endif // #if USE_USD_SDK
+#endif	  // #if USE_USD_SDK
