USD: Reformat USDStage module.

#rnx

diff --git a/Engine/Plugins/Importers/USDImporter/Source/USDStage/Private/USDLevelSequenceHelper.cpp b/Engine/Plugins/Importers/USDImporter/Source/USDStage/Private/USDLevelSequenceHelper.cpp
--- a/Engine/Plugins/Importers/USDImporter/Source/USDStage/Private/USDLevelSequenceHelper.cpp
+++ b/Engine/Plugins/Importers/USDImporter/Source/USDStage/Private/USDLevelSequenceHelper.cpp
@@ -4,15 +4,14 @@
 
 #include "USDAssetUserData.h"
 #include "USDAttributeUtils.h"
-#include "USDDrawModeComponent.h"
 #include "USDClassesModule.h"
 #include "USDConversionUtils.h"
+#include "USDDrawModeComponent.h"
 #include "USDInfoCache.h"
 #include "USDIntegrationUtils.h"
 #include "USDLayerUtils.h"
 #include "USDListener.h"
 #include "USDLog.h"
-#include "USDMemory.h"
 #include "USDPrimConversion.h"
 #include "USDPrimTwin.h"
 #include "USDProjectSettings.h"
@@ -45,11 +44,9 @@
 #include "Components/SkeletalMeshComponent.h"
 #include "Components/SpotLightComponent.h"
 #include "ControlRigObjectBinding.h"
-#include "CoreMinimal.h"
 #include "Framework/Notifications/NotificationManager.h"
 #include "GeometryCache.h"
 #include "GeometryCacheComponent.h"
-#include "GeometryCacheTrackUSD.h"
 #include "GroomCache.h"
 #include "GroomComponent.h"
 #include "HAL/IConsoleManager.h"
@@ -89,14 +86,12 @@
 #include "ControlRigBlueprint.h"
 #include "Editor.h"
 #include "Editor/TransBuffer.h"
-#include "EditorModeManager.h"
-#include "EditorModes.h"
 #include "Exporters/AnimSeqExportOption.h"
 #include "ILevelSequenceEditorToolkit.h"
 #include "ISequencer.h"
 #include "MovieSceneToolHelpers.h"
 #include "Subsystems/AssetEditorSubsystem.h"
-#endif // WITH_EDITOR
+#endif	  // WITH_EDITOR
 
 #define LOCTEXT_NAMESPACE "USDLevelSequenceHelper"
 
@@ -199,14 +194,14 @@ namespace UsdLevelSequenceHelperImpl
 		{
 			MovieScene->GetMuteNodes().Remove(MuteNode);
 		}
-#endif // WITH_EDITOR
+#endif	  // WITH_EDITOR
 
 		Track->Modify();
 		Track->SetEvalDisabled(bMute);
 	}
 
 #if WITH_EDITOR
-	TSharedPtr< ISequencer > GetOpenedSequencerForLevelSequence(ULevelSequence* LevelSequence)
+	TSharedPtr<ISequencer> GetOpenedSequencerForLevelSequence(ULevelSequence* LevelSequence)
 	{
 		const bool bFocusIfOpen = false;
 		IAssetEditorInstance* AssetEditor = GEditor->GetEditorSubsystem<UAssetEditorSubsystem>()->FindEditorForAsset(LevelSequence, bFocusIfOpen);
@@ -255,7 +250,11 @@ namespace UsdLevelSequenceHelperImpl
 			}
 
 			// Evaluate at the beginning of the subscene time to ensure that spawnables are created before export
-			FFrameTime StartTime = FFrameRate::TransformTime(UE::MovieScene::DiscreteInclusiveLower(MovieScene->GetPlaybackRange()).Value, MovieScene->GetTickResolution(), MovieScene->GetDisplayRate());
+			FFrameTime StartTime = FFrameRate::TransformTime(
+				UE::MovieScene::DiscreteInclusiveLower(MovieScene->GetPlaybackRange()).Value,
+				MovieScene->GetTickResolution(),
+				MovieScene->GetDisplayRate()
+			);
 			LevelPlayer->SetPlaybackPosition(FMovieSceneSequencePlaybackParams(StartTime, EUpdatePositionMethod::Play));
 		}
 
@@ -273,7 +272,15 @@ namespace UsdLevelSequenceHelperImpl
 
 			FMovieSceneSequenceIDRef Template = MovieSceneSequenceID::Root;
 			FMovieSceneSequenceTransform RootToLocalTransform;
-			bResult = MovieSceneToolHelpers::ExportToAnimSequence(AnimSequence, ExportOptions, MovieScene, Player, SkeletalMeshComp, Template, RootToLocalTransform);
+			bResult = MovieSceneToolHelpers::ExportToAnimSequence(
+				AnimSequence,
+				ExportOptions,
+				MovieScene,
+				Player,
+				SkeletalMeshComp,
+				Template,
+				RootToLocalTransform
+			);
 			if (!bResult)
 			{
 				goto cleanup;
@@ -283,7 +290,8 @@ namespace UsdLevelSequenceHelperImpl
 		// Disable any extra existing control rig tracks for this binding.
 		// Reuse one of the control rig parameter tracks if we can
 		{
-			TArray<UMovieSceneTrack*> Tracks = MovieScene->FindTracks(UMovieSceneControlRigParameterTrack::StaticClass(), ComponentBinding, NAME_None);
+			TArray<UMovieSceneTrack*> Tracks = MovieScene
+												   ->FindTracks(UMovieSceneControlRigParameterTrack::StaticClass(), ComponentBinding, NAME_None);
 			for (UMovieSceneTrack* AnyOleTrack : Tracks)
 			{
 				UMovieSceneControlRigParameterTrack* ValidTrack = Cast<UMovieSceneControlRigParameterTrack>(AnyOleTrack);
@@ -300,7 +308,9 @@ namespace UsdLevelSequenceHelperImpl
 
 			if (!Track)
 			{
-				Track = Cast<UMovieSceneControlRigParameterTrack>(MovieScene->AddTrack(UMovieSceneControlRigParameterTrack::StaticClass(), ComponentBinding));
+				Track = Cast<UMovieSceneControlRigParameterTrack>(
+					MovieScene->AddTrack(UMovieSceneControlRigParameterTrack::StaticClass(), ComponentBinding)
+				);
 				Track->Modify();
 			}
 		}
@@ -332,7 +342,7 @@ namespace UsdLevelSequenceHelperImpl
 			FFrameNumber ControlRigSectionStartFrame = TNumericLimits<int32>::Max();
 			UMovieSceneSkeletalAnimationTrack* SkelTrack = Cast<UMovieSceneSkeletalAnimationTrack>(
 				MovieScene->FindTrack(UMovieSceneSkeletalAnimationTrack::StaticClass(), ComponentBinding, NAME_None)
-				);
+			);
 			if (SkelTrack)
 			{
 				for (const UMovieSceneSection* Section : SkelTrack->GetAllSections())
@@ -371,8 +381,15 @@ namespace UsdLevelSequenceHelperImpl
 			// of the playback range, which we do here
 			ControlRigSectionStartFrame -= UE::MovieScene::DiscreteInclusiveLower(MovieScene->GetPlaybackRange());
 
-			ParamSection->LoadAnimSequenceIntoThisSection(AnimSequence, MovieScene, SkeletalMeshComp, bReduceKeys, Tolerance, ControlRigSectionStartFrame,
-				EMovieSceneKeyInterpolation::SmartAuto);
+			ParamSection->LoadAnimSequenceIntoThisSection(
+				AnimSequence,
+				MovieScene,
+				SkeletalMeshComp,
+				bReduceKeys,
+				Tolerance,
+				ControlRigSectionStartFrame,
+				EMovieSceneKeyInterpolation::SmartAuto
+			);
 
 			// Disable Skeletal Animation Tracks
 			if (SkelTrack)
@@ -392,7 +409,7 @@ namespace UsdLevelSequenceHelperImpl
 			bResult = true;
 		}
 
-	cleanup:
+cleanup:
 		if (bCreatedTempSequence && AnimSequence)
 		{
 			AnimSequence->MarkAsGarbage();
@@ -411,15 +428,15 @@ namespace UsdLevelSequenceHelperImpl
 		return bResult;
 	}
 
-	void ShowTransformTrackOnCameraComponentWarning( const USceneComponent* Component )
+	void ShowTransformTrackOnCameraComponentWarning(const USceneComponent* Component)
 	{
-		const UCineCameraComponent* CameraComponent = Cast<const UCineCameraComponent>( Component );
-		if( !CameraComponent )
+		const UCineCameraComponent* CameraComponent = Cast<const UCineCameraComponent>(Component);
+		if (!CameraComponent)
 		{
 			return;
 		}
 		const AActor* OwnerActor = CameraComponent->GetOwner();
-		if ( !OwnerActor )
+		if (!OwnerActor)
 		{
 			return;
 		}
@@ -432,25 +449,29 @@ namespace UsdLevelSequenceHelperImpl
 		}
 		WarnedComponents.Add(NewComponentKey);
 
-		const FText Text = LOCTEXT( "TransformTrackOnCameraComponentText", "USD: Transform track on camera component" );
+		const FText Text = LOCTEXT("TransformTrackOnCameraComponentText", "USD: Transform track on camera component");
 
 		const FText SubText = FText::Format(
-			LOCTEXT( "TransformTrackOnCameraComponentSubText", "The LevelSequence binding to the camera component '{0}' has a transform track, which is not supported.\n\nFor animating camera transforms, please bind a transform track to the '{1}' CameraActor directly, or to its root scene component instead." ),
-			FText::FromString( Component->GetName() ),
-			FText::FromString( OwnerActor->GetActorLabel() )
+			LOCTEXT(
+				"TransformTrackOnCameraComponentSubText",
+				"The LevelSequence binding to the camera component '{0}' has a transform track, which is not supported.\n\nFor animating camera "
+				"transforms, please bind a transform track to the '{1}' CameraActor directly, or to its root scene component instead."
+			),
+			FText::FromString(Component->GetName()),
+			FText::FromString(OwnerActor->GetActorLabel())
 		);
 
-		UE_LOG( LogUsd, Warning, TEXT( "%s" ), *SubText.ToString().Replace( TEXT( "\n\n" ), TEXT( " " ) ) );
+		UE_LOG(LogUsd, Warning, TEXT("%s"), *SubText.ToString().Replace(TEXT("\n\n"), TEXT(" ")));
 
 		const UUsdProjectSettings* Settings = GetDefault<UUsdProjectSettings>();
-		if ( Settings && Settings->bShowTransformTrackOnCameraComponentWarning )
+		if (Settings && Settings->bShowTransformTrackOnCameraComponentWarning)
 		{
 			static TWeakPtr<SNotificationItem> Notification;
 
-			FNotificationInfo Toast( Text );
+			FNotificationInfo Toast(Text);
 			Toast.SubText = SubText;
-			Toast.Image = FCoreStyle::Get().GetBrush( TEXT( "MessageLog.Warning" ) );
-			Toast.CheckBoxText = LOCTEXT( "DontAskAgain", "Don't prompt again" );
+			Toast.Image = FCoreStyle::Get().GetBrush(TEXT("MessageLog.Warning"));
+			Toast.CheckBoxText = LOCTEXT("DontAskAgain", "Don't prompt again");
 			Toast.bUseLargeFont = false;
 			Toast.bFireAndForget = false;
 			Toast.FadeOutDuration = 0.0f;
@@ -458,42 +479,47 @@ namespace UsdLevelSequenceHelperImpl
 			Toast.bUseThrobber = false;
 			Toast.bUseSuccessFailIcons = false;
 			Toast.ButtonDetails.Emplace(
-				LOCTEXT( "OverridenOpinionMessageOk", "Ok" ),
+				LOCTEXT("OverridenOpinionMessageOk", "Ok"),
 				FText::GetEmpty(),
-				FSimpleDelegate::CreateLambda( []() {
-				if ( TSharedPtr<SNotificationItem> PinnedNotification = Notification.Pin() )
-				{
-					PinnedNotification->SetCompletionState( SNotificationItem::CS_Success );
-					PinnedNotification->ExpireAndFadeout();
-				}
-			} )
+				FSimpleDelegate::CreateLambda(
+					[]()
+					{
+						if (TSharedPtr<SNotificationItem> PinnedNotification = Notification.Pin())
+						{
+							PinnedNotification->SetCompletionState(SNotificationItem::CS_Success);
+							PinnedNotification->ExpireAndFadeout();
+						}
+					}
+				)
 			);
 			// This is flipped because the default checkbox message is "Don't prompt again"
 			Toast.CheckBoxState = Settings->bShowTransformTrackOnCameraComponentWarning ? ECheckBoxState::Unchecked : ECheckBoxState::Checked;
-			Toast.CheckBoxStateChanged = FOnCheckStateChanged::CreateStatic( []( ECheckBoxState NewState )
-			{
-				if ( UUsdProjectSettings* Settings = GetMutableDefault<UUsdProjectSettings>() )
+			Toast.CheckBoxStateChanged = FOnCheckStateChanged::CreateStatic(
+				[](ECheckBoxState NewState)
 				{
-					// This is flipped because the default checkbox message is "Don't prompt again"
-					Settings->bShowTransformTrackOnCameraComponentWarning = NewState == ECheckBoxState::Unchecked;
-					Settings->SaveConfig();
+					if (UUsdProjectSettings* Settings = GetMutableDefault<UUsdProjectSettings>())
+					{
+						// This is flipped because the default checkbox message is "Don't prompt again"
+						Settings->bShowTransformTrackOnCameraComponentWarning = NewState == ECheckBoxState::Unchecked;
+						Settings->SaveConfig();
+					}
 				}
-			} );
+			);
 
 			// Only show one at a time
-			if ( !Notification.IsValid() )
+			if (!Notification.IsValid())
 			{
-				Notification = FSlateNotificationManager::Get().AddNotification( Toast );
+				Notification = FSlateNotificationManager::Get().AddNotification(Toast);
 			}
 
-			if ( TSharedPtr<SNotificationItem> PinnedNotification = Notification.Pin() )
+			if (TSharedPtr<SNotificationItem> PinnedNotification = Notification.Pin())
 			{
-				PinnedNotification->SetCompletionState( SNotificationItem::CS_Pending );
+				PinnedNotification->SetCompletionState(SNotificationItem::CS_Pending);
 			}
 		}
 	}
-#endif // WITH_EDITOR
-}
+#endif	  // WITH_EDITOR
+}	 // namespace UsdLevelSequenceHelperImpl
 
 class FUsdLevelSequenceHelperImpl : private FGCObject
 {
@@ -519,7 +545,7 @@ private:
 		FString Identifier;
 		FString FilePath;
 
-		TArray< FLayerOffsetInfo > SubLayersOffsets;
+		TArray<FLayerOffsetInfo> SubLayersOffsets;
 
 		TOptional<double> StartTimeCode;
 		TOptional<double> EndTimeCode;
@@ -549,10 +575,13 @@ public:
 	void SetRootMotionHandling(EUsdRootMotionHandling NewValue);
 	void OnStageActorRenamed();
 
-	ULevelSequence* GetMainLevelSequence() const { return MainLevelSequence; }
-	TArray< ULevelSequence* > GetSubSequences() const
+	ULevelSequence* GetMainLevelSequence() const
+	{
+		return MainLevelSequence;
+	}
+	TArray<ULevelSequence*> GetSubSequences() const
 	{
-		TArray< ULevelSequence* > SubSequences;
+		TArray<ULevelSequence*> SubSequences;
 		ObjectPtrDecay(LevelSequencesByIdentifier).GenerateValueArray(SubSequences);
 		SubSequences.Remove(MainLevelSequence);
 
@@ -579,10 +608,11 @@ private:
 	TMap<FString, TObjectPtr<ULevelSequence>> LevelSequencesByIdentifier;
 	TMap<ULevelSequence*, FString> IdentifierByLevelSequence;
 
-	TSet< FName > LocalLayersSequences; // List of sequences associated with sublayers
+	TSet<FName> LocalLayersSequences;							 // List of sequences associated with sublayers
 
-	FMovieSceneSequenceHierarchy SequenceHierarchyCache; // Cache for the hierarchy of level sequences and subsections
-	TMap< ULevelSequence*, FMovieSceneSequenceID > SequencesID; // Tracks the FMovieSceneSequenceID for each Sequence in the hierarchy. We assume that each Sequence is only present once in the hierarchy.
+	FMovieSceneSequenceHierarchy SequenceHierarchyCache;		 // Cache for the hierarchy of level sequences and subsections
+	TMap<ULevelSequence*, FMovieSceneSequenceID> SequencesID;	 // Tracks the FMovieSceneSequenceID for each Sequence in the hierarchy. We assume
+																 // that each Sequence is only present once in the hierarchy.
 
 	// Sequence Name to Layer Identifier Map. Relationship: N Sequences to 1 Layer.
 	TMap<FName, FString> LayerIdentifierByLevelSequenceName;
@@ -609,7 +639,13 @@ private:
 	void AddGeometryCacheTracks(const UUsdPrimTwin& PrimTwin, const UE::FUsdPrim& Prim);
 
 	template<typename TrackType>
-	TrackType* AddTrack(const FName& PropertyPath, const UUsdPrimTwin& PrimTwin, USceneComponent& ComponentToBind, ULevelSequence& Sequence, bool bIsMuted = false);
+	TrackType* AddTrack(
+		const FName& PropertyPath,
+		const UUsdPrimTwin& PrimTwin,
+		USceneComponent& ComponentToBind,
+		ULevelSequence& Sequence,
+		bool bIsMuted = false
+	);
 
 	void RemovePossessable(const UUsdPrimTwin& PrimTwin);
 
@@ -634,10 +670,10 @@ private:
 
 		// For now we support one binding per component type (mostly so we can fit a binding to a scene component and
 		// camera component for a Camera prim twin)
-		TMap< const UClass*, FGuid > ObjectClassToBindingGuid;
+		TMap<const UClass*, FGuid> ObjectClassToBindingGuid;
 	};
 
-	TMap< TWeakObjectPtr< const UUsdPrimTwin >, FPrimTwinBindings > PrimTwinToBindings;
+	TMap<TWeakObjectPtr<const UUsdPrimTwin>, FPrimTwinBindings> PrimTwinToBindings;
 
 	// Time codes handling
 private:
@@ -658,17 +694,26 @@ private:
 
 	FGuid GetOrCreateComponentBinding(const UUsdPrimTwin& PrimTwin, USceneComponent& ComponentToBind, ULevelSequence& Sequence);
 
-	TMap<FString, FLayerTimeInfo> LayerTimeInfosByLayerIdentifier; // Maps a LayerTimeInfo to a given Layer through its identifier
+	TMap<FString, FLayerTimeInfo> LayerTimeInfosByLayerIdentifier;	  // Maps a LayerTimeInfo to a given Layer through its identifier
 
-// Changes handling
+																	  // Changes handling
 public:
-	void StartMonitoringChanges() { MonitoringChangesWhenZero.Decrement(); }
-	void StopMonitoringChanges() { MonitoringChangesWhenZero.Increment(); }
-	bool IsMonitoringChanges() const { return MonitoringChangesWhenZero.GetValue() == 0; }
+	void StartMonitoringChanges()
+	{
+		MonitoringChangesWhenZero.Decrement();
+	}
+	void StopMonitoringChanges()
+	{
+		MonitoringChangesWhenZero.Increment();
+	}
+	bool IsMonitoringChanges() const
+	{
+		return MonitoringChangesWhenZero.GetValue() == 0;
+	}
 
 	/**
-	 * Used as a fire-and-forget block that will prevent any levelsequence object (tracks, moviescene, sections, etc.) change from being written to the stage.
-	 * We unblock during HandleTransactionStateChanged.
+	 * Used as a fire-and-forget block that will prevent any levelsequence object (tracks, moviescene, sections, etc.) change from being written to
+	 * the stage. We unblock during HandleTransactionStateChanged.
 	 */
 	void BlockMonitoringChangesForThisTransaction();
 
@@ -693,27 +738,30 @@ private:
 	static const EObjectFlags DefaultObjFlags;
 	static const double DefaultFramerate;
 	static const TCHAR* TimeTrackName;
-	static const double EmptySubSectionRange; // How many frames should an empty subsection cover, only needed so that the subsection is visible and the user can edit it
+	static const double EmptySubSectionRange;	 // How many frames should an empty subsection cover, only needed so that the subsection is visible
+												 // and the user can edit it
 
 	FUsdLevelSequenceHelper::FOnSkelAnimationBaked OnSkelAnimationBaked;
 
 	TWeakObjectPtr<AUsdStageActor> StageActor = nullptr;
-	TSharedPtr<FUsdInfoCache> InfoCache = nullptr;  // We keep a pointer to this directly because we may be called via the USDStageImporter directly, when we don't have an available actor
-	TSharedPtr<UE::FUsdGeomBBoxCache> BBoxCache = nullptr; // Same as for the info cache
+	TSharedPtr<FUsdInfoCache> InfoCache = nullptr;			  // We keep a pointer to this directly because we may be called via the USDStageImporter
+															  // directly, when we don't have an available actor
+	TSharedPtr<UE::FUsdGeomBBoxCache> BBoxCache = nullptr;	  // Same as for the info cache
 	EUsdRootMotionHandling RootMotionHandling = EUsdRootMotionHandling::NoAdditionalRootMotion;
 	FGuid StageActorBinding;
 
 	// Only when this is zero we write LevelSequence object (tracks, moviescene, sections, etc.) transactions back to the USD stage
 	FThreadSafeCounter MonitoringChangesWhenZero;
 
-	// When we call BlockMonitoringChangesForThisTransaction, we record the FGuid of the current transaction. We'll early out of all OnObjectTransacted calls for that transaction
-	// We keep a set here in order to remember all the blocked transactions as we're going through them
+	// When we call BlockMonitoringChangesForThisTransaction, we record the FGuid of the current transaction. We'll early out of all
+	// OnObjectTransacted calls for that transaction We keep a set here in order to remember all the blocked transactions as we're going through them
 	TSet<FGuid> BlockedTransactionGuids;
 
 	UE::FUsdStage UsdStage;
 };
 
-const EObjectFlags FUsdLevelSequenceHelperImpl::DefaultObjFlags = EObjectFlags::RF_Transactional | EObjectFlags::RF_Transient | EObjectFlags::RF_Public;
+const EObjectFlags FUsdLevelSequenceHelperImpl::DefaultObjFlags = EObjectFlags::RF_Transactional | EObjectFlags::RF_Transient
+																  | EObjectFlags::RF_Public;
 const double FUsdLevelSequenceHelperImpl::DefaultFramerate = 24.0;
 const TCHAR* FUsdLevelSequenceHelperImpl::TimeTrackName = TEXT("Time");
 const double FUsdLevelSequenceHelperImpl::EmptySubSectionRange = 10.0;
@@ -731,7 +779,7 @@ FUsdLevelSequenceHelperImpl::FUsdLevelSequenceHelperImpl()
 			Transactor->OnTransactionStateChanged().AddRaw(this, &FUsdLevelSequenceHelperImpl::HandleTransactionStateChanged);
 		}
 	}
-#endif // WITH_EDITOR
+#endif	  // WITH_EDITOR
 }
 
 FUsdLevelSequenceHelperImpl::~FUsdLevelSequenceHelperImpl()
@@ -754,7 +802,7 @@ FUsdLevelSequenceHelperImpl::~FUsdLevelSequenceHelperImpl()
 			Transactor->OnTransactionStateChanged().RemoveAll(this);
 		}
 	}
-#endif // WITH_EDITOR
+#endif	  // WITH_EDITOR
 }
 
 ULevelSequence* FUsdLevelSequenceHelperImpl::Init(const UE::FUsdStage& InUsdStage)
@@ -865,8 +913,9 @@ void FUsdLevelSequenceHelperImpl::CreateLocalLayersSequences()
 
 	LocalLayersSequences.Add(MainLevelSequence->GetFName());
 
-	TFunction< void(const FLayerTimeInfo* LayerTimeInfo, ULevelSequence& ParentSequence) > RecursivelyCreateSequencesForLayer;
-	RecursivelyCreateSequencesForLayer = [&RecursivelyCreateSequencesForLayer, this](const FLayerTimeInfo* LayerTimeInfo, ULevelSequence& ParentSequence)
+	TFunction<void(const FLayerTimeInfo* LayerTimeInfo, ULevelSequence& ParentSequence)> RecursivelyCreateSequencesForLayer;
+	RecursivelyCreateSequencesForLayer =
+		[&RecursivelyCreateSequencesForLayer, this](const FLayerTimeInfo* LayerTimeInfo, ULevelSequence& ParentSequence)
 	{
 		if (!LayerTimeInfo)
 		{
@@ -881,7 +930,7 @@ void FUsdLevelSequenceHelperImpl::CreateLocalLayersSequences()
 				{
 					if (ULevelSequence* SubSequence = FindOrAddSequenceForLayer(SubLayer, SubLayer.GetIdentifier(), SubLayer.GetDisplayName()))
 					{
-						if (!LocalLayersSequences.Contains(SubSequence->GetFName())) // Make sure we don't parse an already parsed layer
+						if (!LocalLayersSequences.Contains(SubSequence->GetFName()))	// Make sure we don't parse an already parsed layer
 						{
 							LocalLayersSequences.Add(SubSequence->GetFName());
 
@@ -921,7 +970,7 @@ void FUsdLevelSequenceHelperImpl::BindToUsdStageActor(AUsdStageActor* InStageAct
 		StageActor->GetActorLabel(),
 #else
 		StageActor->GetName(),
-#endif // WITH_EDITOR
+#endif	  // WITH_EDITOR
 		StageActor->GetClass()
 	);
 	MainLevelSequence->BindPossessableObject(StageActorBinding, *StageActor, StageActor->GetWorld());
@@ -977,13 +1026,14 @@ void FUsdLevelSequenceHelperImpl::OnStageActorRenamed()
 		return;
 	}
 
-	FMovieScenePossessable NewPossessable{
-	#if WITH_EDITOR
+	FMovieScenePossessable NewPossessable
+	{
+#if WITH_EDITOR
 		StageActorPtr->GetActorLabel(),
 #else
 		StageActorPtr->GetName(),
-#endif // WITH_EDITOR
-		StageActorPtr->GetClass()
+#endif	  // WITH_EDITOR
+			StageActorPtr->GetClass()
 	};
 	FGuid NewId = NewPossessable.GetGuid();
 
@@ -1089,7 +1139,11 @@ ULevelSequence* FUsdLevelSequenceHelperImpl::FindOrAddSequenceForAttribute(const
 	return Sequence;
 }
 
-ULevelSequence* FUsdLevelSequenceHelperImpl::FindOrAddSequenceForLayer(const UE::FSdfLayer& Layer, const FString& SequenceIdentifier, const FString& SequenceDisplayName)
+ULevelSequence* FUsdLevelSequenceHelperImpl::FindOrAddSequenceForLayer(
+	const UE::FSdfLayer& Layer,
+	const FString& SequenceIdentifier,
+	const FString& SequenceDisplayName
+)
 {
 	if (!Layer)
 	{
@@ -1100,23 +1154,22 @@ ULevelSequence* FUsdLevelSequenceHelperImpl::FindOrAddSequenceForLayer(const UE:
 
 	if (!Sequence)
 	{
-		// This needs to be unique, or else when we reload the stage we will end up with a new ULevelSequence with the same class, outer and name as the
-		// previous one. Also note that the previous level sequence, even though unreferenced by the stage actor, is likely still alive and valid due to references
-		// from the transaction buffer, so we would basically end up creating a identical new object on top of an existing one (the new object has the same address as the existing one).
-		// When importing we don't actually want to do this though, because we want these assets name to conflict so that we can publish/replace old assets if desired. The stage
-		// importer will make these names unique later if needed.
-		// We only get an InfoCache when importing (from UUsdStageImporter::ImportFromFile) or when BindToUsdStageActor is called,
-		// which also gives us a stage actor. So if we don't have an actor but have a cache, we're importing
+		// This needs to be unique, or else when we reload the stage we will end up with a new ULevelSequence with the same class, outer and name as
+		// the previous one. Also note that the previous level sequence, even though unreferenced by the stage actor, is likely still alive and valid
+		// due to references from the transaction buffer, so we would basically end up creating a identical new object on top of an existing one (the
+		// new object has the same address as the existing one). When importing we don't actually want to do this though, because we want these assets
+		// name to conflict so that we can publish/replace old assets if desired. The stage importer will make these names unique later if needed. We
+		// only get an InfoCache when importing (from UUsdStageImporter::ImportFromFile) or when BindToUsdStageActor is called, which also gives us a
+		// stage actor. So if we don't have an actor but have a cache, we're importing
 		const bool bIsImporting = StageActor.IsExplicitlyNull() && InfoCache;
-		FName UniqueSequenceName = bIsImporting
-			? *IUsdClassesModule::SanitizeObjectName(FPaths::GetBaseFilename(SequenceDisplayName))
-			: MakeUniqueObjectName(
-				GetTransientPackage(),
-				ULevelSequence::StaticClass(),
-				*IUsdClassesModule::SanitizeObjectName(FPaths::GetBaseFilename(SequenceDisplayName))
-			);
-
-		Sequence = NewObject< ULevelSequence >(GetTransientPackage(), UniqueSequenceName, FUsdLevelSequenceHelperImpl::DefaultObjFlags);
+		FName UniqueSequenceName = bIsImporting ? *IUsdClassesModule::SanitizeObjectName(FPaths::GetBaseFilename(SequenceDisplayName))
+												: MakeUniqueObjectName(
+													GetTransientPackage(),
+													ULevelSequence::StaticClass(),
+													*IUsdClassesModule::SanitizeObjectName(FPaths::GetBaseFilename(SequenceDisplayName))
+												);
+
+		Sequence = NewObject<ULevelSequence>(GetTransientPackage(), UniqueSequenceName, FUsdLevelSequenceHelperImpl::DefaultObjFlags);
 		Sequence->Initialize();
 
 		UMovieScene* MovieScene = Sequence->MovieScene;
@@ -1153,29 +1206,31 @@ UMovieSceneSubSection* FUsdLevelSequenceHelperImpl::FindSubSequenceSection(ULeve
 		return nullptr;
 	}
 
-	UMovieSceneSubTrack* SubTrack = MovieScene->FindTrack< UMovieSceneSubTrack >();
+	UMovieSceneSubTrack* SubTrack = MovieScene->FindTrack<UMovieSceneSubTrack>();
 
 	if (!SubTrack)
 	{
 		return nullptr;
 	}
 
-	UMovieSceneSection* const* SubSection = Algo::FindByPredicate(SubTrack->GetAllSections(),
+	UMovieSceneSection* const* SubSection = Algo::FindByPredicate(
+		SubTrack->GetAllSections(),
 		[&SubSequence](UMovieSceneSection* Section) -> bool
 		{
-			if (UMovieSceneSubSection* SubSection = Cast< UMovieSceneSubSection >(Section))
+			if (UMovieSceneSubSection* SubSection = Cast<UMovieSceneSubSection>(Section))
 			{
-				return  (SubSection->GetSequence() == &SubSequence);
+				return (SubSection->GetSequence() == &SubSequence);
 			}
 			else
 			{
 				return false;
 			}
-		});
+		}
+	);
 
 	if (SubSection)
 	{
-		return  Cast< UMovieSceneSubSection >(*SubSection);
+		return Cast<UMovieSceneSubSection>(*SubSection);
 	}
 	else
 	{
@@ -1203,10 +1258,10 @@ void FUsdLevelSequenceHelperImpl::CreateSubSequenceSection(ULevelSequence& Seque
 
 	FFrameRate TickResolution = MovieScene->GetTickResolution();
 
-	UMovieSceneSubTrack* SubTrack = MovieScene->FindTrack< UMovieSceneSubTrack >();
+	UMovieSceneSubTrack* SubTrack = MovieScene->FindTrack<UMovieSceneSubTrack>();
 	if (!SubTrack)
 	{
-		SubTrack = MovieScene->AddTrack< UMovieSceneSubTrack >();
+		SubTrack = MovieScene->AddTrack<UMovieSceneSubTrack>();
 	}
 
 	const FString* LayerIdentifier = LayerIdentifierByLevelSequenceName.Find(Sequence.GetFName());
@@ -1230,7 +1285,7 @@ void FUsdLevelSequenceHelperImpl::CreateSubSequenceSection(ULevelSequence& Seque
 	UE::FSdfLayer Layer = UE::FSdfLayer::FindOrOpen(**LayerIdentifier);
 	UE::FSdfLayer SubLayer = UE::FSdfLayer::FindOrOpen(**SubLayerIdentifier);
 
-	TArray< FString > PrimPathsForSequence;
+	TArray<FString> PrimPathsForSequence;
 	PrimPathByLevelSequenceName.MultiFind(SubSequence.GetFName(), PrimPathsForSequence);
 
 	if (PrimPathsForSequence.Num() > 0)
@@ -1246,11 +1301,13 @@ void FUsdLevelSequenceHelperImpl::CreateSubSequenceSection(ULevelSequence& Seque
 	}
 	else if (UsdStage.HasLocalLayer(SubLayer))
 	{
-		const FLayerOffsetInfo* SubLayerOffsetPtr = Algo::FindByPredicate(LayerTimeInfo->SubLayersOffsets,
+		const FLayerOffsetInfo* SubLayerOffsetPtr = Algo::FindByPredicate(
+			LayerTimeInfo->SubLayersOffsets,
 			[&SubLayerIdentifier](const FLayerOffsetInfo& Other)
 			{
 				return (Other.LayerIdentifier == *SubLayerIdentifier);
-			});
+			}
+		);
 
 		if (SubLayerOffsetPtr)
 		{
@@ -1260,7 +1317,7 @@ void FUsdLevelSequenceHelperImpl::CreateSubSequenceSection(ULevelSequence& Seque
 
 	const double TimeCodesPerSecond = Layer.GetTimeCodesPerSecond();
 	const bool bIsAlembicSublayer = SubLayerIdentifier->EndsWith(TEXT(".abc"));
-	TRange< FFrameNumber > SubSectionRange;
+	TRange<FFrameNumber> SubSectionRange;
 	FFrameNumber StartFrame;
 
 	if (!bIsAlembicSublayer)
@@ -1284,7 +1341,7 @@ void FUsdLevelSequenceHelperImpl::CreateSubSequenceSection(ULevelSequence& Seque
 
 		// Max here because StartFrame can theoretically be larger than StageEndFrame, which would generate a range where the upper bound is smaller
 		// than the lower bound, which can trigger asserts
-		SubSectionRange = TRange< FFrameNumber >{StartFrame, FMath::Max(StageEndFrame, EndFrame)};
+		SubSectionRange = TRange<FFrameNumber>{StartFrame, FMath::Max(StageEndFrame, EndFrame)};
 	}
 	else
 	{
@@ -1310,9 +1367,12 @@ void FUsdLevelSequenceHelperImpl::CreateSubSequenceSection(ULevelSequence& Seque
 	}
 	else
 	{
-		SubSection = SubTrack->AddSequence(&SubSequence, SubSectionRange.GetLowerBoundValue(), SubSectionRange.Size< FFrameNumber >().Value);
+		SubSection = SubTrack->AddSequence(&SubSequence, SubSectionRange.GetLowerBoundValue(), SubSectionRange.Size<FFrameNumber>().Value);
 
-		UE_LOG(LogUsd, Verbose, TEXT("Adding subsection '%s' to sequence '%s'. StartFrame: '%d'"),
+		UE_LOG(
+			LogUsd,
+			Verbose,
+			TEXT("Adding subsection '%s' to sequence '%s'. StartFrame: '%d'"),
 			*SubSection->GetName(),
 			*Sequence.GetName(),
 			StartFrame.Value
@@ -1326,7 +1386,7 @@ void FUsdLevelSequenceHelperImpl::CreateSubSequenceSection(ULevelSequence& Seque
 	{
 		UMovieSceneCompiledDataManager::CompileHierarchy(MainLevelSequence, &SequenceHierarchyCache, EMovieSceneServerClientMask::All);
 
-		for (const TTuple< FMovieSceneSequenceID, FMovieSceneSubSequenceData >& Pair : SequenceHierarchyCache.AllSubSequenceData())
+		for (const TTuple<FMovieSceneSequenceID, FMovieSceneSubSequenceData>& Pair : SequenceHierarchyCache.AllSubSequenceData())
 		{
 			if (UMovieSceneSequence* CachedSubSequence = Pair.Value.GetSequence())
 			{
@@ -1342,7 +1402,7 @@ void FUsdLevelSequenceHelperImpl::CreateSubSequenceSection(ULevelSequence& Seque
 
 void FUsdLevelSequenceHelperImpl::RemoveSubSequenceSection(ULevelSequence& Sequence, ULevelSequence& SubSequence)
 {
-	if (UMovieSceneSubTrack* SubTrack = Sequence.GetMovieScene()->FindTrack< UMovieSceneSubTrack >())
+	if (UMovieSceneSubTrack* SubTrack = Sequence.GetMovieScene()->FindTrack<UMovieSceneSubTrack>())
 	{
 		if (UMovieSceneSection* SubSection = FindSubSequenceSection(Sequence, SubSequence))
 		{
@@ -1374,7 +1434,10 @@ void FUsdLevelSequenceHelperImpl::CreateTimeTrack(const FLayerTimeInfo& Info)
 		return;
 	}
 
-	UMovieSceneFloatTrack* TimeTrack = MovieScene->FindTrack<UMovieSceneFloatTrack>(StageActorBinding, FName(FUsdLevelSequenceHelperImpl::TimeTrackName));
+	UMovieSceneFloatTrack* TimeTrack = MovieScene->FindTrack<UMovieSceneFloatTrack>(
+		StageActorBinding,
+		FName(FUsdLevelSequenceHelperImpl::TimeTrackName)
+	);
 	if (TimeTrack)
 	{
 		TimeTrack->RemoveAllAnimationData();
@@ -1402,7 +1465,7 @@ void FUsdLevelSequenceHelperImpl::CreateTimeTrack(const FLayerTimeInfo& Info)
 		FFrameNumber StartFrame = UsdLevelSequenceHelperImpl::RoundAsFrameNumber(DestTickRate, StartTimeCode / TimeCodesPerSecond);
 		FFrameNumber EndFrame = UsdLevelSequenceHelperImpl::RoundAsFrameNumber(DestTickRate, EndTimeCode / TimeCodesPerSecond);
 
-		TRange< FFrameNumber > PlaybackRange(StartFrame, EndFrame);
+		TRange<FFrameNumber> PlaybackRange(StartFrame, EndFrame);
 
 		bool bSectionAdded = false;
 
@@ -1448,7 +1511,10 @@ void FUsdLevelSequenceHelperImpl::RemoveTimeTrack(const FLayerTimeInfo* LayerTim
 		return;
 	}
 
-	UMovieSceneFloatTrack* TimeTrack = MovieScene->FindTrack< UMovieSceneFloatTrack >(StageActorBinding, FName(FUsdLevelSequenceHelperImpl::TimeTrackName));
+	UMovieSceneFloatTrack* TimeTrack = MovieScene->FindTrack<UMovieSceneFloatTrack>(
+		StageActorBinding,
+		FName(FUsdLevelSequenceHelperImpl::TimeTrackName)
+	);
 	if (TimeTrack)
 	{
 		MovieScene->RemoveTrack(*TimeTrack);
@@ -1632,8 +1698,7 @@ void FUsdLevelSequenceHelperImpl::AddCommonTracks(const UUsdPrimTwin& PrimTwin,
 			// many parents and specs on multiple layers, but this is hopefully at least *a* reasonable answer.
 			UE::FUsdAttribute FirstAnimatedVisibilityParentAttr;
 
-			if ((VisibilityAttribute.GetTimeSamples(TotalVisibilityTimeSamples) && TotalVisibilityTimeSamples.Num() > 0)
-				|| bForceVisibilityTracks)
+			if ((VisibilityAttribute.GetTimeSamples(TotalVisibilityTimeSamples) && TotalVisibilityTimeSamples.Num() > 0) || bForceVisibilityTracks)
 			{
 				// TODO: Improve this, as this is extremely inefficient since we'll be parsing this tree for the root down and repeatedly
 				// redoing this one child at a time...
@@ -1642,7 +1707,10 @@ void FUsdLevelSequenceHelperImpl::AddCommonTracks(const UUsdPrimTwin& PrimTwin,
 				{
 					if (UsdUtils::HasAnimatedVisibility(ParentPrim))
 					{
-						TArray<UE::FUsdAttribute> ParentAttrs = UnrealToUsd::GetAttributesForProperty(ParentPrim, UnrealIdentifiers::HiddenInGamePropertyName);
+						TArray<UE::FUsdAttribute> ParentAttrs = UnrealToUsd::GetAttributesForProperty(
+							ParentPrim,
+							UnrealIdentifiers::HiddenInGamePropertyName
+						);
 						if (ParentAttrs.Num() > 0)
 						{
 							UE::FUsdAttribute ParentVisAttr = ParentAttrs[0];
@@ -1694,10 +1762,25 @@ void FUsdLevelSequenceHelperImpl::AddCommonTracks(const UUsdPrimTwin& PrimTwin,
 
 					if (UMovieScene* MovieScene = AttributeSequence->GetMovieScene())
 					{
-						if (UMovieSceneVisibilityTrack* VisibilityTrack = AddTrack<UMovieSceneVisibilityTrack>(UnrealIdentifiers::HiddenInGamePropertyName, PrimTwin, *ComponentToBind, *AttributeSequence, bIsMuted))
+						if (UMovieSceneVisibilityTrack* VisibilityTrack = AddTrack<UMovieSceneVisibilityTrack>(
+								UnrealIdentifiers::HiddenInGamePropertyName,
+								PrimTwin,
+								*ComponentToBind,
+								*AttributeSequence,
+								bIsMuted
+							))
 						{
-							UsdToUnreal::FPropertyTrackReader Reader = UsdToUnreal::CreatePropertyTrackReader(Prim, UnrealIdentifiers::HiddenInGamePropertyName);
-							UsdToUnreal::ConvertBoolTimeSamples(UsdStage, TotalVisibilityTimeSamples, Reader.BoolReader, *VisibilityTrack, SequenceTransform);
+							UsdToUnreal::FPropertyTrackReader Reader = UsdToUnreal::CreatePropertyTrackReader(
+								Prim,
+								UnrealIdentifiers::HiddenInGamePropertyName
+							);
+							UsdToUnreal::ConvertBoolTimeSamples(
+								UsdStage,
+								TotalVisibilityTimeSamples,
+								Reader.BoolReader,
+								*VisibilityTrack,
+								SequenceTransform
+							);
 						}
 
 						PrimPathByLevelSequenceName.AddUnique(AttributeSequence->GetFName(), Prim.GetPrimPath().GetString());
@@ -1847,8 +1930,7 @@ void FUsdLevelSequenceHelperImpl::AddCameraTracks(const UUsdPrimTwin& PrimTwin,
 		UnrealIdentifiers::ManualFocusDistancePropertyName,
 		UnrealIdentifiers::CurrentAperturePropertyName,
 		UnrealIdentifiers::SensorWidthPropertyName,
-		UnrealIdentifiers::SensorHeightPropertyName
-	};
+		UnrealIdentifiers::SensorHeightPropertyName};
 
 	UE::FSdfLayer PrimLayer = UsdUtils::FindLayerForPrim(Prim);
 	ULevelSequence* PrimSequence = FindSequenceForIdentifier(PrimLayer.GetIdentifier());
@@ -1904,7 +1986,8 @@ void FUsdLevelSequenceHelperImpl::AddCameraTracks(const UUsdPrimTwin& PrimTwin,
 				continue;
 			}
 
-			if (UMovieSceneFloatTrack* FloatTrack = AddTrack<UMovieSceneFloatTrack>(PropertyName, PrimTwin, *ComponentToBind, *AttributeSequence, bIsMuted))
+			if (UMovieSceneFloatTrack* FloatTrack = AddTrack<
+					UMovieSceneFloatTrack>(PropertyName, PrimTwin, *ComponentToBind, *AttributeSequence, bIsMuted))
 			{
 				UsdToUnreal::FPropertyTrackReader Reader = UsdToUnreal::CreatePropertyTrackReader(Prim, PropertyName);
 				UsdToUnreal::ConvertFloatTimeSamples(UsdStage, TimeSamples, Reader.FloatReader, *FloatTrack, SequenceTransform);
@@ -1965,7 +2048,7 @@ void FUsdLevelSequenceHelperImpl::AddLightTracks(const UUsdPrimTwin& PrimTwin, c
 	// If we were told to specifically read only some property paths, ignore the other ones
 	if (PropertyPathsToRead.Num() > 0)
 	{
-		for (TMap< FName, ETrackType >::TIterator Iter = PropertyPathToTrackType.CreateIterator(); Iter; ++Iter)
+		for (TMap<FName, ETrackType>::TIterator Iter = PropertyPathToTrackType.CreateIterator(); Iter; ++Iter)
 		{
 			const FName& PropertyPath = Iter->Key;
 			if (!PropertyPathsToRead.Contains(PropertyPath))
@@ -1982,7 +2065,7 @@ void FUsdLevelSequenceHelperImpl::AddLightTracks(const UUsdPrimTwin& PrimTwin, c
 		return;
 	}
 
-	for (const TPair< FName, ETrackType >& Pair : PropertyPathToTrackType)
+	for (const TPair<FName, ETrackType>& Pair : PropertyPathToTrackType)
 	{
 		const FName& PropertyPath = Pair.Key;
 		ETrackType TrackType = Pair.Value;
@@ -2034,33 +2117,36 @@ void FUsdLevelSequenceHelperImpl::AddLightTracks(const UUsdPrimTwin& PrimTwin, c
 
 		switch (TrackType)
 		{
-		case ETrackType::Bool:
-		{
-			if (UMovieSceneBoolTrack* BoolTrack = AddTrack<UMovieSceneBoolTrack>(PropertyPath, PrimTwin, *ComponentToBind, *PrimSequence, bIsMuted))
+			case ETrackType::Bool:
 			{
-				UsdToUnreal::ConvertBoolTimeSamples(UsdStage, UnionedTimeSamples, Reader.BoolReader, *BoolTrack, SequenceTransform);
+				if (UMovieSceneBoolTrack* BoolTrack = AddTrack<
+						UMovieSceneBoolTrack>(PropertyPath, PrimTwin, *ComponentToBind, *PrimSequence, bIsMuted))
+				{
+					UsdToUnreal::ConvertBoolTimeSamples(UsdStage, UnionedTimeSamples, Reader.BoolReader, *BoolTrack, SequenceTransform);
+				}
+				break;
 			}
-			break;
-		}
-		case ETrackType::Float:
-		{
-			if (UMovieSceneFloatTrack* FloatTrack = AddTrack<UMovieSceneFloatTrack>(PropertyPath, PrimTwin, *ComponentToBind, *PrimSequence, bIsMuted))
+			case ETrackType::Float:
 			{
-				UsdToUnreal::ConvertFloatTimeSamples(UsdStage, UnionedTimeSamples, Reader.FloatReader, *FloatTrack, SequenceTransform);
+				if (UMovieSceneFloatTrack* FloatTrack = AddTrack<
+						UMovieSceneFloatTrack>(PropertyPath, PrimTwin, *ComponentToBind, *PrimSequence, bIsMuted))
+				{
+					UsdToUnreal::ConvertFloatTimeSamples(UsdStage, UnionedTimeSamples, Reader.FloatReader, *FloatTrack, SequenceTransform);
+				}
+				break;
 			}
-			break;
-		}
-		case ETrackType::Color:
-		{
-			if (UMovieSceneColorTrack* ColorTrack = AddTrack<UMovieSceneColorTrack>(PropertyPath, PrimTwin, *ComponentToBind, *PrimSequence, bIsMuted))
+			case ETrackType::Color:
 			{
-				UsdToUnreal::ConvertColorTimeSamples(UsdStage, UnionedTimeSamples, Reader.ColorReader, *ColorTrack, SequenceTransform);
+				if (UMovieSceneColorTrack* ColorTrack = AddTrack<
+						UMovieSceneColorTrack>(PropertyPath, PrimTwin, *ComponentToBind, *PrimSequence, bIsMuted))
+				{
+					UsdToUnreal::ConvertColorTimeSamples(UsdStage, UnionedTimeSamples, Reader.ColorReader, *ColorTrack, SequenceTransform);
+				}
+				break;
 			}
-			break;
-		}
-		default:
-			continue;
-			break;
+			default:
+				continue;
+				break;
 		}
 
 		PrimPathByLevelSequenceName.AddUnique(PrimSequence->GetFName(), Prim.GetPrimPath().GetString());
@@ -2094,9 +2180,7 @@ void FUsdLevelSequenceHelperImpl::AddSkeletalTracks(const UUsdPrimTwin& PrimTwin
 	// but we may belong to a FUsdStageImportContext, and so there's no AUsdStageActor at all to use.
 	// At this point it doesn't matter much though, because we shouldn't need to uncollapse a SkelAnimation prim path anyway
 	const UE::FSdfPath PrimPath = Prim.GetPrimPath();
-	UAnimSequence* Sequence = InfoCache->GetSingleAssetForPrim<UAnimSequence>(
-		PrimPath
-	);
+	UAnimSequence* Sequence = InfoCache->GetSingleAssetForPrim<UAnimSequence>(PrimPath);
 	if (!Sequence)
 	{
 		return;
@@ -2108,13 +2192,21 @@ void FUsdLevelSequenceHelperImpl::AddSkeletalTracks(const UUsdPrimTwin& PrimTwin
 	UE::FUsdAttribute BlendShapeWeightsAttr = SkelAnimationPrim.GetAttribute(TEXT("blendShapeWeights"));
 
 	const bool bIncludeSessionLayers = false;
-	UE::FSdfLayer SkelAnimationLayer = UsdUtils::FindLayerForAttributes({TranslationsAttr, RotationsAttr, ScalesAttr, BlendShapeWeightsAttr}, 0.0, bIncludeSessionLayers);
+	UE::FSdfLayer SkelAnimationLayer = UsdUtils::FindLayerForAttributes(
+		{TranslationsAttr, RotationsAttr, ScalesAttr, BlendShapeWeightsAttr},
+		0.0,
+		bIncludeSessionLayers
+	);
 	if (!SkelAnimationLayer)
 	{
 		return;
 	}
 
-	ULevelSequence* SkelAnimationSequence = FindOrAddSequenceForLayer(SkelAnimationLayer, SkelAnimationLayer.GetIdentifier(), SkelAnimationLayer.GetDisplayName());
+	ULevelSequence* SkelAnimationSequence = FindOrAddSequenceForLayer(
+		SkelAnimationLayer,
+		SkelAnimationLayer.GetIdentifier(),
+		SkelAnimationLayer.GetDisplayName()
+	);
 	if (!SkelAnimationSequence)
 	{
 		return;
@@ -2129,13 +2221,13 @@ void FUsdLevelSequenceHelperImpl::AddSkeletalTracks(const UUsdPrimTwin& PrimTwin
 	// We will mute all SkelAnimation attributes if we mute, so here let's only consider something muted
 	// if it has all attributes muted as well.
 	// We know at least one of these attributes ones is valid and animated because we have an UAnimSequence
-	const bool bIsMuted =
-		(!TranslationsAttr || UsdUtils::IsAttributeMuted(TranslationsAttr, UsdStage)) &&
-		(!RotationsAttr || UsdUtils::IsAttributeMuted(RotationsAttr, UsdStage)) &&
-		(!ScalesAttr || UsdUtils::IsAttributeMuted(ScalesAttr, UsdStage)) &&
-		(!BlendShapeWeightsAttr || UsdUtils::IsAttributeMuted(BlendShapeWeightsAttr, UsdStage));
+	const bool bIsMuted = (!TranslationsAttr || UsdUtils::IsAttributeMuted(TranslationsAttr, UsdStage))
+						  && (!RotationsAttr || UsdUtils::IsAttributeMuted(RotationsAttr, UsdStage))
+						  && (!ScalesAttr || UsdUtils::IsAttributeMuted(ScalesAttr, UsdStage))
+						  && (!BlendShapeWeightsAttr || UsdUtils::IsAttributeMuted(BlendShapeWeightsAttr, UsdStage));
 
-	if (UMovieSceneSkeletalAnimationTrack* SkeletalTrack = AddTrack<UMovieSceneSkeletalAnimationTrack>(SkelAnimationPrim.GetName(), PrimTwin, *ComponentToBind, *SkelAnimationSequence, bIsMuted))
+	if (UMovieSceneSkeletalAnimationTrack* SkeletalTrack = AddTrack<
+			UMovieSceneSkeletalAnimationTrack>(SkelAnimationPrim.GetName(), PrimTwin, *ComponentToBind, *SkelAnimationSequence, bIsMuted))
 	{
 		double LayerStartOffsetSeconds = 0.0f;
 		if (UUsdAnimSequenceAssetUserData* UserData = Sequence->GetAssetUserData<UUsdAnimSequenceAssetUserData>())
@@ -2147,7 +2239,9 @@ void FUsdLevelSequenceHelperImpl::AddSkeletalTracks(const UUsdPrimTwin& PrimTwin
 
 		SkeletalTrack->RemoveAllAnimationData();
 
-		UMovieSceneSkeletalAnimationSection* NewSection = Cast< UMovieSceneSkeletalAnimationSection >(SkeletalTrack->AddNewAnimation(StartOffsetTick, Sequence));
+		UMovieSceneSkeletalAnimationSection* NewSection = Cast<UMovieSceneSkeletalAnimationSection>(
+			SkeletalTrack->AddNewAnimation(StartOffsetTick, Sequence)
+		);
 		NewSection->EvalOptions.CompletionMode = EMovieSceneCompletionMode::KeepState;
 	}
 
@@ -2156,7 +2250,7 @@ void FUsdLevelSequenceHelperImpl::AddSkeletalTracks(const UUsdPrimTwin& PrimTwin
 
 void FUsdLevelSequenceHelperImpl::AddGeometryCacheTracks(const UUsdPrimTwin& PrimTwin, const UE::FUsdPrim& Prim)
 {
-	UGeometryCacheComponent* ComponentToBind = Cast< UGeometryCacheComponent >(PrimTwin.GetSceneComponent());
+	UGeometryCacheComponent* ComponentToBind = Cast<UGeometryCacheComponent>(PrimTwin.GetSceneComponent());
 	if (!ComponentToBind || !InfoCache)
 	{
 		return;
@@ -2164,9 +2258,7 @@ void FUsdLevelSequenceHelperImpl::AddGeometryCacheTracks(const UUsdPrimTwin& Pri
 
 	// Fetch the geometry cache asset from the asset cache. If there's none, don't actually need to create track
 	const UE::FSdfPath PrimPath = Prim.GetPrimPath();
-	UGeometryCache* GeometryCache = InfoCache->GetSingleAssetForPrim<UGeometryCache>(
-		PrimPath
-	);
+	UGeometryCache* GeometryCache = InfoCache->GetSingleAssetForPrim<UGeometryCache>(PrimPath);
 	if (!GeometryCache)
 	{
 		return;
@@ -2183,7 +2275,11 @@ void FUsdLevelSequenceHelperImpl::AddGeometryCacheTracks(const UUsdPrimTwin& Pri
 		return;
 	}
 
-	ULevelSequence* GeometryCacheSequence = FindOrAddSequenceForLayer(GeometryCacheLayer, GeometryCacheLayer.GetIdentifier(), GeometryCacheLayer.GetDisplayName());
+	ULevelSequence* GeometryCacheSequence = FindOrAddSequenceForLayer(
+		GeometryCacheLayer,
+		GeometryCacheLayer.GetIdentifier(),
+		GeometryCacheLayer.GetDisplayName()
+	);
 	if (!GeometryCacheSequence)
 	{
 		return;
@@ -2196,22 +2292,27 @@ void FUsdLevelSequenceHelperImpl::AddGeometryCacheTracks(const UUsdPrimTwin& Pri
 	}
 
 	const bool bIsMuted = false;
-	if (UMovieSceneGeometryCacheTrack* GeometryCacheTrack = AddTrack< UMovieSceneGeometryCacheTrack >(Prim.GetName(), PrimTwin, *ComponentToBind, *GeometryCacheSequence, bIsMuted))
+	if (UMovieSceneGeometryCacheTrack* GeometryCacheTrack = AddTrack<
+			UMovieSceneGeometryCacheTrack>(Prim.GetName(), PrimTwin, *ComponentToBind, *GeometryCacheSequence, bIsMuted))
 	{
 		GeometryCacheTrack->RemoveAllAnimationData();
 
 		double LayerStartOffsetSeconds = 0;
 		if (GeometryCacheSequence == MainLevelSequence)
 		{
-			// The LayerStartOffset needs to be applied only for the track in the main sequence. For subsequences, it's the subsequence section that is offset.
+			// The LayerStartOffset needs to be applied only for the track in the main sequence. For subsequences, it's the subsequence section that
+			// is offset.
 			if (UUsdGeometryCacheAssetUserData* UserData = GeometryCache->GetAssetUserData<UUsdGeometryCacheAssetUserData>())
 			{
 				LayerStartOffsetSeconds = UserData->LayerStartOffsetSeconds;
 			}
 		}
 
-		const FFrameNumber StartOffsetTick = FFrameTime::FromDecimal(LayerStartOffsetSeconds * MovieScene->GetTickResolution().AsDecimal()).RoundToFrame();
-		UMovieSceneGeometryCacheSection* NewSection = Cast< UMovieSceneGeometryCacheSection >(GeometryCacheTrack->AddNewAnimation(StartOffsetTick, ComponentToBind));
+		const FFrameNumber StartOffsetTick = FFrameTime::FromDecimal(LayerStartOffsetSeconds * MovieScene->GetTickResolution().AsDecimal())
+												 .RoundToFrame();
+		UMovieSceneGeometryCacheSection* NewSection = Cast<UMovieSceneGeometryCacheSection>(
+			GeometryCacheTrack->AddNewAnimation(StartOffsetTick, ComponentToBind)
+		);
 		NewSection->EvalOptions.CompletionMode = EMovieSceneCompletionMode::KeepState;
 	}
 
@@ -2220,7 +2321,7 @@ void FUsdLevelSequenceHelperImpl::AddGeometryCacheTracks(const UUsdPrimTwin& Pri
 
 void FUsdLevelSequenceHelperImpl::AddGroomTracks(const UUsdPrimTwin& PrimTwin, const UE::FUsdPrim& Prim)
 {
-	UGroomComponent* ComponentToBind = Cast< UGroomComponent >(PrimTwin.GetSceneComponent());
+	UGroomComponent* ComponentToBind = Cast<UGroomComponent>(PrimTwin.GetSceneComponent());
 	if (!ComponentToBind)
 	{
 		return;
@@ -2234,9 +2335,7 @@ void FUsdLevelSequenceHelperImpl::AddGroomTracks(const UUsdPrimTwin& PrimTwin, c
 	// Fetch the groom cache asset from the asset cache. If there's none, don't actually need to create track
 	const FString PrimPath = Prim.GetPrimPath().GetString();
 	const FString GroomCachePath = FString::Printf(TEXT("%s_strands_cache"), *PrimPath);
-	UGroomCache* GroomCache = InfoCache->GetSingleAssetForPrim<UGroomCache>(
-		UE::FSdfPath{*GroomCachePath}
-	);
+	UGroomCache* GroomCache = InfoCache->GetSingleAssetForPrim<UGroomCache>(UE::FSdfPath{*GroomCachePath});
 	if (!GroomCache)
 	{
 		return;
@@ -2266,12 +2365,13 @@ void FUsdLevelSequenceHelperImpl::AddGroomTracks(const UUsdPrimTwin& PrimTwin, c
 	}
 
 	const bool bIsMuted = false;
-	if (UMovieSceneGroomCacheTrack* GroomCacheTrack = AddTrack< UMovieSceneGroomCacheTrack >(Prim.GetName(), PrimTwin, *ComponentToBind, *GroomAnimationSequence, bIsMuted))
+	if (UMovieSceneGroomCacheTrack* GroomCacheTrack = AddTrack<
+			UMovieSceneGroomCacheTrack>(Prim.GetName(), PrimTwin, *ComponentToBind, *GroomAnimationSequence, bIsMuted))
 	{
 		GroomCacheTrack->RemoveAllAnimationData();
 
 		const FFrameNumber StartOffset;
-		UMovieSceneGroomCacheSection* NewSection = Cast< UMovieSceneGroomCacheSection >(GroomCacheTrack->AddNewAnimation(StartOffset, ComponentToBind));
+		UMovieSceneGroomCacheSection* NewSection = Cast<UMovieSceneGroomCacheSection>(GroomCacheTrack->AddNewAnimation(StartOffset, ComponentToBind));
 		NewSection->EvalOptions.CompletionMode = EMovieSceneCompletionMode::KeepState;
 	}
 
@@ -2297,7 +2397,7 @@ void FUsdLevelSequenceHelperImpl::AddPrim(UUsdPrimTwin& PrimTwin, bool bForceVis
 		PrimSequence = MainLevelSequence;
 	}
 
-	TArray< UE::FUsdAttribute > PrimAttributes = UsdPrim.GetAttributes();
+	TArray<UE::FUsdAttribute> PrimAttributes = UsdPrim.GetAttributes();
 
 	for (const UE::FUsdAttribute& PrimAttribute : PrimAttributes)
 	{
@@ -2351,7 +2451,13 @@ void FUsdLevelSequenceHelperImpl::AddPrim(UUsdPrimTwin& PrimTwin, bool bForceVis
 }
 
 template<typename TrackType>
-TrackType* FUsdLevelSequenceHelperImpl::AddTrack(const FName& TrackName, const UUsdPrimTwin& PrimTwin, USceneComponent& ComponentToBind, ULevelSequence& Sequence, bool bIsMuted)
+TrackType* FUsdLevelSequenceHelperImpl::AddTrack(
+	const FName& TrackName,
+	const UUsdPrimTwin& PrimTwin,
+	USceneComponent& ComponentToBind,
+	ULevelSequence& Sequence,
+	bool bIsMuted
+)
 {
 	if (!UsdStage)
 	{
@@ -2366,14 +2472,14 @@ TrackType* FUsdLevelSequenceHelperImpl::AddTrack(const FName& TrackName, const U
 
 	const FGuid ComponentBinding = GetOrCreateComponentBinding(PrimTwin, ComponentToBind, Sequence);
 
-	TrackType* Track = MovieScene->FindTrack< TrackType >(ComponentBinding, TrackName);
+	TrackType* Track = MovieScene->FindTrack<TrackType>(ComponentBinding, TrackName);
 	if (Track)
 	{
 		Track->RemoveAllAnimationData();
 	}
 	else
 	{
-		Track = MovieScene->AddTrack< TrackType >(ComponentBinding);
+		Track = MovieScene->AddTrack<TrackType>(ComponentBinding);
 		if (!Track)
 		{
 			return nullptr;
@@ -2388,7 +2494,7 @@ TrackType* FUsdLevelSequenceHelperImpl::AddTrack(const FName& TrackName, const U
 		{
 			Track->SetDisplayName(FText::FromName(TrackName));
 		}
-#endif // WITH_EDITOR
+#endif	  // WITH_EDITOR
 	}
 
 	UsdLevelSequenceHelperImpl::MuteTrack(Track, MovieScene, ComponentBinding.ToString(), Track->GetName(), bIsMuted);
@@ -2405,9 +2511,9 @@ void FUsdLevelSequenceHelperImpl::RemovePrim(const UUsdPrimTwin& PrimTwin)
 
 	// We can't assume that the UsdPrim still exists in the stage, it might have been removed already so work from the PrimTwin PrimPath.
 
-	TSet< FName > PrimSequences;
+	TSet<FName> PrimSequences;
 
-	for (TPair< FName, FString >& PrimPathByLevelSequenceNamePair : PrimPathByLevelSequenceName)
+	for (TPair<FName, FString>& PrimPathByLevelSequenceNamePair : PrimPathByLevelSequenceName)
 	{
 		if (PrimPathByLevelSequenceNamePair.Value == PrimTwin.PrimPath)
 		{
@@ -2415,7 +2521,7 @@ void FUsdLevelSequenceHelperImpl::RemovePrim(const UUsdPrimTwin& PrimTwin)
 		}
 	}
 
-	TSet< ULevelSequence* > SequencesToRemoveForPrim;
+	TSet<ULevelSequence*> SequencesToRemoveForPrim;
 
 	for (const FName& PrimSequenceName : PrimSequences)
 	{
@@ -2491,10 +2597,7 @@ void FUsdLevelSequenceHelperImpl::UpdateControlRigTracks(UUsdPrimTwin& PrimTwin)
 		SkelAnimationLayer = UsdStage.GetEditTarget();
 
 		FString UniqueChildName = UsdUtils::GetValidChildName(TEXT("Animation"), SkelRootPrim);
-		SkelAnimationPrim = UsdStage.DefinePrim(
-			SkelRootPrim.GetPrimPath().AppendChild(*UniqueChildName),
-			TEXT("SkelAnimation")
-		);
+		SkelAnimationPrim = UsdStage.DefinePrim(SkelRootPrim.GetPrimPath().AppendChild(*UniqueChildName), TEXT("SkelAnimation"));
 		if (!SkelAnimationPrim)
 		{
 			return;
@@ -2512,21 +2615,24 @@ void FUsdLevelSequenceHelperImpl::UpdateControlRigTracks(UUsdPrimTwin& PrimTwin)
 	// Fetch the UAnimSequence asset from the asset cache. Ideally we'd call AUsdStageActor::GetGeneratedAssets,
 	// but we may belong to a FUsdStageImportContext, and so there's no AUsdStageActor at all to use.
 	// At this point it doesn't matter much though, because we shouldn't need to uncollapse a SkelAnimation prim path anyway
-	UAnimSequence* AnimSequence = InfoCache->GetSingleAssetForPrim<UAnimSequence>(
-		PrimPath
-	);
+	UAnimSequence* AnimSequence = InfoCache->GetSingleAssetForPrim<UAnimSequence>(PrimPath);
 
 	UE::FUsdEditContext EditContext{UsdStage, SkelAnimationLayer};
 	FString Identifier = SkelAnimationLayer.GetIdentifier();
 
-	// Force-create these because these are mandatory anyway (https://graphics.pixar.com/usd/release/api/_usd_skel__schemas.html#UsdSkel_SkelAnimation)
+	// Force-create these because these are mandatory anyway
+	// (https://graphics.pixar.com/usd/release/api/_usd_skel__schemas.html#UsdSkel_SkelAnimation)
 	UE::FUsdAttribute JointsAttr = SkelAnimationPrim.CreateAttribute(TEXT("joints"), TEXT("token[]"));
 	UE::FUsdAttribute TranslationsAttr = SkelAnimationPrim.CreateAttribute(TEXT("translations"), TEXT("float3[]"));
 	UE::FUsdAttribute RotationsAttr = SkelAnimationPrim.CreateAttribute(TEXT("rotations"), TEXT("quatf[]"));
 	UE::FUsdAttribute ScalesAttr = SkelAnimationPrim.CreateAttribute(TEXT("scales"), TEXT("half3[]"));
 	UE::FUsdAttribute BlendShapeWeightsAttr = SkelAnimationPrim.GetAttribute(TEXT("blendShapeWeights"));
 
-	ULevelSequence* SkelAnimationSequence = FindOrAddSequenceForLayer(SkelAnimationLayer, SkelAnimationLayer.GetIdentifier(), SkelAnimationLayer.GetDisplayName());
+	ULevelSequence* SkelAnimationSequence = FindOrAddSequenceForLayer(
+		SkelAnimationLayer,
+		SkelAnimationLayer.GetIdentifier(),
+		SkelAnimationLayer.GetDisplayName()
+	);
 	if (!SkelAnimationSequence)
 	{
 		return;
@@ -2540,7 +2646,7 @@ void FUsdLevelSequenceHelperImpl::UpdateControlRigTracks(UUsdPrimTwin& PrimTwin)
 
 	const FGuid ComponentBinding = GetOrCreateComponentBinding(PrimTwin, *ComponentToBind, *SkelAnimationSequence);
 
-	UMovieSceneControlRigParameterTrack* ControlRigTrack = MovieScene->FindTrack< UMovieSceneControlRigParameterTrack >(ComponentBinding);
+	UMovieSceneControlRigParameterTrack* ControlRigTrack = MovieScene->FindTrack<UMovieSceneControlRigParameterTrack>(ComponentBinding);
 
 	// We should be in control rig track mode but don't have any tracks yet --> Setup for Control Rig
 	if (!ControlRigTrack)
@@ -2609,7 +2715,7 @@ void FUsdLevelSequenceHelperImpl::UpdateControlRigTracks(UUsdPrimTwin& PrimTwin)
 
 		if (ControlRigClass)
 		{
-			UAnimSeqExportOption* NewOptions = NewObject< UAnimSeqExportOption >();
+			UAnimSeqExportOption* NewOptions = NewObject<UAnimSeqExportOption>();
 
 			UsdLevelSequenceHelperImpl::BakeToControlRig(
 				ComponentToBind->GetWorld(),
@@ -2628,12 +2734,12 @@ void FUsdLevelSequenceHelperImpl::UpdateControlRigTracks(UUsdPrimTwin& PrimTwin)
 	}
 
 	PrimPathByLevelSequenceName.AddUnique(SkelAnimationSequence->GetFName(), PrimPath.GetString());
-#endif // WITH_EDITOR
+#endif	  // WITH_EDITOR
 }
 
 void FUsdLevelSequenceHelperImpl::RemoveSequenceForPrim(ULevelSequence& Sequence, const UUsdPrimTwin& PrimTwin)
 {
-	TArray< FString > PrimPathsForSequence;
+	TArray<FString> PrimPathsForSequence;
 	PrimPathByLevelSequenceName.MultiFind(Sequence.GetFName(), PrimPathsForSequence);
 
 	if (PrimPathsForSequence.Find(PrimTwin.PrimPath) != INDEX_NONE)
@@ -2652,7 +2758,7 @@ void FUsdLevelSequenceHelperImpl::RemoveSequenceForPrim(ULevelSequence& Sequence
 
 				if (FMovieSceneSubSequenceData* ParentSubSequenceData = SequenceHierarchyCache.FindSubData(ParentSequenceID))
 				{
-					ParentSequence = Cast< ULevelSequence >(ParentSubSequenceData->GetSequence());
+					ParentSequence = Cast<ULevelSequence>(ParentSubSequenceData->GetSequence());
 				}
 			}
 
@@ -2691,7 +2797,7 @@ void FUsdLevelSequenceHelperImpl::RemovePossessable(const UUsdPrimTwin& PrimTwin
 	// ones don't modify the Sequence and change properties, so we must modify them here
 	Bindings->Sequence->Modify();
 
-	for (const TPair< const UClass*, FGuid >& Pair : Bindings->ObjectClassToBindingGuid)
+	for (const TPair<const UClass*, FGuid>& Pair : Bindings->ObjectClassToBindingGuid)
 	{
 		const FGuid& ComponentPossessableGuid = Pair.Value;
 
@@ -2739,7 +2845,7 @@ void FUsdLevelSequenceHelperImpl::RefreshSequencer()
 		return;
 	}
 
-	if (TSharedPtr< ISequencer > Sequencer = UsdLevelSequenceHelperImpl::GetOpenedSequencerForLevelSequence(MainLevelSequence))
+	if (TSharedPtr<ISequencer> Sequencer = UsdLevelSequenceHelperImpl::GetOpenedSequencerForLevelSequence(MainLevelSequence))
 	{
 		// Don't try refreshing the sequencer if its displaying a stale sequence (e.g. during busy transitions like import) as it
 		// can crash
@@ -2748,7 +2854,7 @@ void FUsdLevelSequenceHelperImpl::RefreshSequencer()
 			Sequencer->NotifyMovieSceneDataChanged(EMovieSceneDataChangeType::RefreshTree);
 		}
 	}
-#endif // WITH_EDITOR
+#endif	  // WITH_EDITOR
 }
 
 void FUsdLevelSequenceHelperImpl::UpdateUsdLayerOffsetFromSection(const UMovieSceneSequence* Sequence, const UMovieSceneSubSection* Section)
@@ -2809,7 +2915,7 @@ void FUsdLevelSequenceHelperImpl::UpdateUsdLayerOffsetFromSection(const UMovieSc
 	}
 
 	// Prevent twins from being rebuilt when we update the layer offsets
-	TOptional< FScopedBlockNoticeListening > BlockNotices;
+	TOptional<FScopedBlockNoticeListening> BlockNotices;
 
 	if (StageActor.IsValid())
 	{
@@ -2826,14 +2932,16 @@ void FUsdLevelSequenceHelperImpl::UpdateUsdLayerOffsetFromSection(const UMovieSc
 		}
 
 		int32 SubLayerIndex = INDEX_NONE;
-		FLayerOffsetInfo* SubLayerOffset = Algo::FindByPredicate(LayerTimeInfo->SubLayersOffsets,
+		FLayerOffsetInfo* SubLayerOffset = Algo::FindByPredicate(
+			LayerTimeInfo->SubLayersOffsets,
 			[&SubLayerIndex, &SubLayerIdentifier = SubLayerTimeInfo->Identifier](const FLayerOffsetInfo& Other)
-		{
-			bool bFound = (Other.LayerIdentifier == SubLayerIdentifier);
-			++SubLayerIndex;
+			{
+				bool bFound = (Other.LayerIdentifier == SubLayerIdentifier);
+				++SubLayerIndex;
 
-			return bFound;
-		});
+				return bFound;
+			}
+		);
 
 		if (SubLayerIndex != INDEX_NONE)
 		{
@@ -2843,7 +2951,7 @@ void FUsdLevelSequenceHelperImpl::UpdateUsdLayerOffsetFromSection(const UMovieSc
 	}
 	else
 	{
-		TArray< FString > PrimPathsForSequence;
+		TArray<FString> PrimPathsForSequence;
 		PrimPathByLevelSequenceName.MultiFind(Section->GetSequence()->GetFName(), PrimPathsForSequence);
 
 		for (const FString& PrimPath : PrimPathsForSequence)
@@ -2885,7 +2993,7 @@ void FUsdLevelSequenceHelperImpl::UpdateMovieSceneTimeRanges(UMovieScene& MovieS
 		const FFrameRate TickResolution = MovieScene.GetTickResolution();
 		const FFrameNumber StartFrame = UsdLevelSequenceHelperImpl::RoundAsFrameNumber(TickResolution, StartTimeCode / TimeCodesPerSecond);
 		const FFrameNumber EndFrame = UsdLevelSequenceHelperImpl::RoundAsFrameNumber(TickResolution, EndTimeCode / TimeCodesPerSecond);
-		TRange< FFrameNumber > TimeRange = TRange<FFrameNumber>::Inclusive(StartFrame, EndFrame);
+		TRange<FFrameNumber> TimeRange = TRange<FFrameNumber>::Inclusive(StartFrame, EndFrame);
 
 		MovieScene.SetPlaybackRange(TimeRange);
 		MovieScene.SetViewRange(StartTimeCode / TimeCodesPerSecond - 1.0f, 1.0f + EndTimeCode / TimeCodesPerSecond);
@@ -2933,7 +3041,8 @@ void FUsdLevelSequenceHelperImpl::OnObjectTransacted(UObject* Object, const clas
 	//
 	// This is due to this mechanism on the Sequencer code where calls to MarkAsChanged (which is a member function of tracks, sections,
 	// MovieScene, etc. and is used to let the UI know it needs to refresh something) can be deferred.
-	// The thing that determines where a call is deferred or not is a global, private variable (check FScopedSignedObjectModifyDefer's implementation).
+	// The thing that determines where a call is deferred or not is a global, private variable (check FScopedSignedObjectModifyDefer's
+	// implementation).
 	//
 	// I think something is causing this mechanism to be stuck deferring everything, or maybe it's some interaction with our code in some way
 	// (not sure at this point). But what I do know is that Sequencer.cpp also has this class FDeferredSignedObjectChangeHandler that listens
@@ -2948,8 +3057,8 @@ void FUsdLevelSequenceHelperImpl::OnObjectTransacted(UObject* Object, const clas
 	// Given that all this does is essentially refresh the Sequencer UI (and only if it had stuck deferred calls!), it's probably not the worst thing
 	// in the world to check it every undo/redo anyway.
 	//
-	// Annoyingly we also need to do this on the next tick though, because we need to make sure this runs this after FDeferredSignedObjectChangeHandler
-	// itself
+	// Annoyingly we also need to do this on the next tick though, because we need to make sure this runs this after
+	// FDeferredSignedObjectChangeHandler itself
 	if (Event.GetEventType() == ETransactionObjectEventType::UndoRedo)
 	{
 		FTSTicker::GetCoreTicker().AddTicker(FTickerDelegate::CreateLambda(
@@ -2974,7 +3083,7 @@ void FUsdLevelSequenceHelperImpl::OnObjectTransacted(UObject* Object, const clas
 		return;
 	}
 
-	if (UMovieScene* MovieScene = Cast< UMovieScene >(Object))
+	if (UMovieScene* MovieScene = Cast<UMovieScene>(Object))
 	{
 		HandleMovieSceneChange(*MovieScene);
 	}
@@ -3009,18 +3118,23 @@ void FUsdLevelSequenceHelperImpl::OnObjectTransacted(UObject* Object, const clas
 					// a brand new scoped transaction, which asserts inside UTransBuffer::CheckState when if finds
 					// out that the previous transaction wasn't fully complete (OnObjectTransacted gets called before
 					// the current transaction is fully done).
-					GEditor->GetTimerManager()->SetTimerForNextTick([this, CRSection]()
+					GEditor->GetTimerManager()->SetTimerForNextTick(
+						[this, CRSection]()
 						{
 							HandleControlRigSectionChange(*CRSection);
-						});
+						}
+					);
 				}
 			}
 		}
-#endif // WITH_EDITOR
+#endif	  // WITH_EDITOR
 	}
 }
 
-void FUsdLevelSequenceHelperImpl::OnUsdObjectsChanged(const UsdUtils::FObjectChangesByPath& InfoChanges, const UsdUtils::FObjectChangesByPath& ResyncChanges)
+void FUsdLevelSequenceHelperImpl::OnUsdObjectsChanged(
+	const UsdUtils::FObjectChangesByPath& InfoChanges,
+	const UsdUtils::FObjectChangesByPath& ResyncChanges
+)
 {
 	AUsdStageActor* StageActorPtr = StageActor.Get();
 	if (!StageActorPtr || !StageActorPtr->IsListeningToUsdNotices())
@@ -3063,9 +3177,13 @@ void FUsdLevelSequenceHelperImpl::OnUsdObjectsChanged(const UsdUtils::FObjectCha
 	}
 }
 
-void FUsdLevelSequenceHelperImpl::HandleTransactionStateChanged(const FTransactionContext& InTransactionContext, const ETransactionStateEventType InTransactionState)
+void FUsdLevelSequenceHelperImpl::HandleTransactionStateChanged(
+	const FTransactionContext& InTransactionContext,
+	const ETransactionStateEventType InTransactionState
+)
 {
-	if (InTransactionState == ETransactionStateEventType::TransactionFinalized && BlockedTransactionGuids.Contains(InTransactionContext.TransactionId))
+	if (InTransactionState == ETransactionStateEventType::TransactionFinalized
+		&& BlockedTransactionGuids.Contains(InTransactionContext.TransactionId))
 	{
 		StartMonitoringChanges();
 	}
@@ -3093,7 +3211,11 @@ double FUsdLevelSequenceHelperImpl::GetTimeCodesPerSecond() const
 	return FMath::IsNearlyZero(StageTimeCodesPerSecond) ? DefaultFramerate : StageTimeCodesPerSecond;
 }
 
-FGuid FUsdLevelSequenceHelperImpl::GetOrCreateComponentBinding(const UUsdPrimTwin& PrimTwin, USceneComponent& ComponentToBind, ULevelSequence& Sequence)
+FGuid FUsdLevelSequenceHelperImpl::GetOrCreateComponentBinding(
+	const UUsdPrimTwin& PrimTwin,
+	USceneComponent& ComponentToBind,
+	ULevelSequence& Sequence
+)
 {
 	UMovieScene* MovieScene = Sequence.GetMovieScene();
 	if (!MovieScene)
@@ -3132,7 +3254,7 @@ FGuid FUsdLevelSequenceHelperImpl::GetOrCreateComponentBinding(const UUsdPrimTwi
 				Actor->GetActorLabel(),
 #else
 				Actor->GetName(),
-#endif // WITH_EDITOR
+#endif	  // WITH_EDITOR
 				Actor->GetClass()
 			);
 			Sequence.BindPossessableObject(ActorBinding, *Actor, Actor->GetWorld());
@@ -3168,7 +3290,7 @@ void FUsdLevelSequenceHelperImpl::HandleMovieSceneChange(UMovieScene& MovieScene
 		return;
 	}
 
-	ULevelSequence* Sequence = MovieScene.GetTypedOuter< ULevelSequence >();
+	ULevelSequence* Sequence = MovieScene.GetTypedOuter<ULevelSequence>();
 	if (!Sequence)
 	{
 		return;
@@ -3188,11 +3310,19 @@ void FUsdLevelSequenceHelperImpl::HandleMovieSceneChange(UMovieScene& MovieScene
 	}
 
 	const double StageTimeCodesPerSecond = GetTimeCodesPerSecond();
-	const TRange< FFrameNumber > PlaybackRange = MovieScene.GetPlaybackRange();
+	const TRange<FFrameNumber> PlaybackRange = MovieScene.GetPlaybackRange();
 	const FFrameRate DisplayRate = MovieScene.GetDisplayRate();
 	const FFrameRate LayerTimeCodesPerSecond(Layer.GetTimeCodesPerSecond(), 1);
-	const FFrameTime StartTime = FFrameRate::TransformTime(UE::MovieScene::DiscreteInclusiveLower(PlaybackRange).Value, MovieScene.GetTickResolution(), LayerTimeCodesPerSecond);
-	const FFrameTime EndTime = FFrameRate::TransformTime(UE::MovieScene::DiscreteExclusiveUpper(PlaybackRange).Value, MovieScene.GetTickResolution(), LayerTimeCodesPerSecond);
+	const FFrameTime StartTime = FFrameRate::TransformTime(
+		UE::MovieScene::DiscreteInclusiveLower(PlaybackRange).Value,
+		MovieScene.GetTickResolution(),
+		LayerTimeCodesPerSecond
+	);
+	const FFrameTime EndTime = FFrameRate::TransformTime(
+		UE::MovieScene::DiscreteExclusiveUpper(PlaybackRange).Value,
+		MovieScene.GetTickResolution(),
+		LayerTimeCodesPerSecond
+	);
 
 	FScopedBlockNoticeListening BlockNotices(StageActor.Get());
 	UE::FSdfChangeBlock ChangeBlock;
@@ -3200,12 +3330,13 @@ void FUsdLevelSequenceHelperImpl::HandleMovieSceneChange(UMovieScene& MovieScene
 	{
 		UsdStage.SetFramesPerSecond(DisplayRate.AsDecimal());
 
-		// For whatever reason setting a stage FramesPerSecond also automatically sets its TimeCodesPerSecond to the same value, so we need to undo it.
-		// This because all the sequencer does is change display rate, which is the analogue to USD's frames per second (i.e. we are only changing how many
-		// frames we'll display between any two timecodes, not how many timecodes we'll display per second)
+		// For whatever reason setting a stage FramesPerSecond also automatically sets its TimeCodesPerSecond to the same value, so we need to undo
+		// it. This because all the sequencer does is change display rate, which is the analogue to USD's frames per second (i.e. we are only changing
+		// how many frames we'll display between any two timecodes, not how many timecodes we'll display per second)
 		UsdStage.SetTimeCodesPerSecond(StageTimeCodesPerSecond);
 
-		// Propagate to all movie scenes, as USD only uses the stage FramesPerSecond so the sequences should have a unified DisplayRate to reflect that
+		// Propagate to all movie scenes, as USD only uses the stage FramesPerSecond so the sequences should have a unified DisplayRate to reflect
+		// that
 		for (auto& SequenceByIdentifier : LevelSequencesByIdentifier)
 		{
 			if (ULevelSequence* OtherSequence = SequenceByIdentifier.Value)
@@ -3245,7 +3376,8 @@ void FUsdLevelSequenceHelperImpl::HandleMovieSceneChange(UMovieScene& MovieScene
 		}
 	};
 
-	auto RemoveTimeSamplesForPropertyIfNeeded = [&MovieScene, &RemoveTimeSamplesForAttr](const UE::FUsdPrim& Prim, const FGuid& Guid, const FName& PropertyPath)
+	auto RemoveTimeSamplesForPropertyIfNeeded =
+		[&MovieScene, &RemoveTimeSamplesForAttr](const UE::FUsdPrim& Prim, const FGuid& Guid, const FName& PropertyPath)
 	{
 		if (!UsdLevelSequenceHelperImpl::FindTrackTypeOrDerived<UMovieScenePropertyTrack>(&MovieScene, Guid, PropertyPath))
 		{
@@ -3257,9 +3389,8 @@ void FUsdLevelSequenceHelperImpl::HandleMovieSceneChange(UMovieScene& MovieScene
 	};
 
 	// Check if we deleted things
-	for (TMap< TWeakObjectPtr< const UUsdPrimTwin >, FPrimTwinBindings >::TIterator PrimTwinIt = PrimTwinToBindings.CreateIterator();
-		PrimTwinIt;
-		++PrimTwinIt)
+	for (TMap<TWeakObjectPtr<const UUsdPrimTwin>, FPrimTwinBindings>::TIterator PrimTwinIt = PrimTwinToBindings.CreateIterator(); PrimTwinIt;
+		 ++PrimTwinIt)
 	{
 		const UUsdPrimTwin* UsdPrimTwin = PrimTwinIt->Key.Get();
 		FPrimTwinBindings& Bindings = PrimTwinIt->Value;
@@ -3269,9 +3400,7 @@ void FUsdLevelSequenceHelperImpl::HandleMovieSceneChange(UMovieScene& MovieScene
 			continue;
 		}
 
-		for (TMap< const UClass*, FGuid >::TIterator BindingIt = Bindings.ObjectClassToBindingGuid.CreateIterator();
-			BindingIt;
-			++BindingIt)
+		for (TMap<const UClass*, FGuid>::TIterator BindingIt = Bindings.ObjectClassToBindingGuid.CreateIterator(); BindingIt; ++BindingIt)
 		{
 			const FGuid& Guid = BindingIt->Value;
 
@@ -3282,7 +3411,8 @@ void FUsdLevelSequenceHelperImpl::HandleMovieSceneChange(UMovieScene& MovieScene
 			}
 
 			// Check if we have an animated attribute and no track for it --> We may have deleted the track, so clear that attribute
-			// We could keep track of these when adding in some kind of map, but while slower this is likely more robust due to the need to support undo/redo
+			// We could keep track of these when adding in some kind of map, but while slower this is likely more robust due to the need to support
+			// undo/redo
 			if (UsdPrimTwin)
 			{
 				USceneComponent* BoundComponent = UsdPrimTwin->GetSceneComponent();
@@ -3402,9 +3532,8 @@ void FUsdLevelSequenceHelperImpl::HandleControlRigSectionChange(UMovieSceneContr
 		return;
 	}
 
-	USkeletalMeshComponent* BoundComponent = Cast< USkeletalMeshComponent >(
-		UsdLevelSequenceHelperImpl::LocateBoundObject(*LevelSequence, *Possessable)
-		);
+	USkeletalMeshComponent* BoundComponent = Cast<USkeletalMeshComponent>(UsdLevelSequenceHelperImpl::LocateBoundObject(*LevelSequence, *Possessable)
+	);
 	if (!BoundComponent)
 	{
 		return;
@@ -3430,7 +3559,7 @@ void FUsdLevelSequenceHelperImpl::HandleControlRigSectionChange(UMovieSceneContr
 	}
 
 	UE::FUsdPrim SkelRootPrim = UsdUtils::GetClosestParentSkelRoot(SkeletonPrim);
-	if(!SkelRootPrim)
+	if (!SkelRootPrim)
 	{
 		return;
 	}
@@ -3493,13 +3622,12 @@ void FUsdLevelSequenceHelperImpl::HandleControlRigSectionChange(UMovieSceneContr
 			// Evaluate at the beginning of the subscene time to ensure that spawnables are created before export
 			// Note that we never actually generate spawnables on our LevelSequence, but its a common pattern to
 			// do this and the user may have added them manually
-			FFrameTime StartTime = FFrameRate::TransformTime(UE::MovieScene::DiscreteInclusiveLower(MovieScene->GetPlaybackRange()).Value, MovieScene->GetTickResolution(), MovieScene->GetDisplayRate());
-			LevelPlayer->SetPlaybackPosition(
-				FMovieSceneSequencePlaybackParams(
-					StartTime,
-					EUpdatePositionMethod::Play
-				)
+			FFrameTime StartTime = FFrameRate::TransformTime(
+				UE::MovieScene::DiscreteInclusiveLower(MovieScene->GetPlaybackRange()).Value,
+				MovieScene->GetTickResolution(),
+				MovieScene->GetDisplayRate()
 			);
+			LevelPlayer->SetPlaybackPosition(FMovieSceneSequencePlaybackParams(StartTime, EUpdatePositionMethod::Play));
 		}
 	}
 
@@ -3548,7 +3676,7 @@ void FUsdLevelSequenceHelperImpl::HandleControlRigSectionChange(UMovieSceneContr
 		// USD.RegenerateSkeletalAssetsOnControlRigBake to false.
 		GetOnSkelAnimationBaked().Broadcast(PrimTwin->PrimPath);
 	}
-#endif // WITH_EDITOR
+#endif	  // WITH_EDITOR
 }
 
 void FUsdLevelSequenceHelperImpl::HandleTrackChange(const UMovieSceneTrack& Track, bool bIsMuteChange)
@@ -3558,7 +3686,7 @@ void FUsdLevelSequenceHelperImpl::HandleTrackChange(const UMovieSceneTrack& Trac
 		return;
 	}
 
-	ULevelSequence* Sequence = Track.GetTypedOuter< ULevelSequence >();
+	ULevelSequence* Sequence = Track.GetTypedOuter<ULevelSequence>();
 	if (!Sequence)
 	{
 		return;
@@ -3590,17 +3718,17 @@ void FUsdLevelSequenceHelperImpl::HandleTrackChange(const UMovieSceneTrack& Trac
 	}
 
 	// Our tracked bindings are always directly to components
-	USceneComponent* BoundSceneComponent = Cast< USceneComponent >(BoundObject);
+	USceneComponent* BoundSceneComponent = Cast<USceneComponent>(BoundObject);
 	if (!BoundSceneComponent)
 	{
 		// The sequencer allows binding actor Transform directly, which gets deferred to its root component
 		// transform. Let's also allow that here
-		if ( AActor* BoundActor = Cast< AActor >( BoundObject ) )
+		if (AActor* BoundActor = Cast<AActor>(BoundObject))
 		{
 			BoundSceneComponent = BoundActor->GetRootComponent();
 		}
 	}
-	if ( !BoundSceneComponent )
+	if (!BoundSceneComponent)
 	{
 		return;
 	}
@@ -3621,23 +3749,21 @@ void FUsdLevelSequenceHelperImpl::HandleTrackChange(const UMovieSceneTrack& Trac
 			// In the scenario where we're trying to make non-decomposed Camera prims work, we only ever want to write out
 			// actual camera properties from the CameraComponent to the Camera prim. We won't write its USceneComponent
 			// properties, as we will use the ones from the ACineCameraActor's parent USceneComponent instead
-			if (PropertyPath == UnrealIdentifiers::CurrentFocalLengthPropertyName ||
-				PropertyPath == UnrealIdentifiers::ManualFocusDistancePropertyName ||
-				PropertyPath == UnrealIdentifiers::ManualFocusDistancePropertyName ||
-				PropertyPath == UnrealIdentifiers::CurrentAperturePropertyName ||
-				PropertyPath == UnrealIdentifiers::SensorWidthPropertyName ||
-				PropertyPath == UnrealIdentifiers::SensorHeightPropertyName
-				)
+			if (PropertyPath == UnrealIdentifiers::CurrentFocalLengthPropertyName
+				|| PropertyPath == UnrealIdentifiers::ManualFocusDistancePropertyName
+				|| PropertyPath == UnrealIdentifiers::ManualFocusDistancePropertyName
+				|| PropertyPath == UnrealIdentifiers::CurrentAperturePropertyName || PropertyPath == UnrealIdentifiers::SensorWidthPropertyName
+				|| PropertyPath == UnrealIdentifiers::SensorHeightPropertyName)
 			{
 				PrimTwin = StageActor->RootUsdTwin->Find(BoundSceneComponent->GetAttachParent());
 			}
 #if WITH_EDITOR
-			else if ( PropertyPath == UnrealIdentifiers::TransformPropertyName )
+			else if (PropertyPath == UnrealIdentifiers::TransformPropertyName)
 			{
 				// Let the user know that we currently don't support transform tracks directly on camera components
-				UsdLevelSequenceHelperImpl::ShowTransformTrackOnCameraComponentWarning( BoundSceneComponent );
+				UsdLevelSequenceHelperImpl::ShowTransformTrackOnCameraComponentWarning(BoundSceneComponent);
 			}
-#endif // WITH_EDITOR
+#endif	  // WITH_EDITOR
 		}
 	}
 
@@ -3681,7 +3807,7 @@ void FUsdLevelSequenceHelperImpl::HandleTrackChange(const UMovieSceneTrack& Trac
 					UE::FUsdAttribute& Attr = Attrs[0];
 
 					bool bAllSectionsMuted = true;
-					for (const UMovieSceneSection* Section : Track.GetAllSections()) // There's no const version of "FindSection"
+					for (const UMovieSceneSection* Section : Track.GetAllSections())	// There's no const version of "FindSection"
 					{
 						bAllSectionsMuted &= !Section->IsActive();
 					}
@@ -3702,7 +3828,7 @@ void FUsdLevelSequenceHelperImpl::HandleTrackChange(const UMovieSceneTrack& Trac
 			else if (const UMovieSceneSkeletalAnimationTrack* SkeletalTrack = Cast<const UMovieSceneSkeletalAnimationTrack>(&Track))
 			{
 				bool bAllSectionsMuted = true;
-				for (const UMovieSceneSection* Section : SkeletalTrack->GetAllSections()) // There's no const version of "FindSection"
+				for (const UMovieSceneSection* Section : SkeletalTrack->GetAllSections())	 // There's no const version of "FindSection"
 				{
 					bAllSectionsMuted &= !Section->IsActive();
 				}
@@ -3753,7 +3879,8 @@ void FUsdLevelSequenceHelperImpl::HandleTrackChange(const UMovieSceneTrack& Trac
 			if (const UMovieScenePropertyTrack* PropertyTrack = Cast<const UMovieScenePropertyTrack>(&Track))
 			{
 				TSet<FName> PropertyPathsToRefresh;
-				UnrealToUsd::FPropertyTrackWriter Writer = UnrealToUsd::CreatePropertyTrackWriter(*BoundSceneComponent, *PropertyTrack, UsdPrim, PropertyPathsToRefresh);
+				UnrealToUsd::FPropertyTrackWriter
+					Writer = UnrealToUsd::CreatePropertyTrackWriter(*BoundSceneComponent, *PropertyTrack, UsdPrim, PropertyPathsToRefresh);
 
 				if (const UMovieSceneFloatTrack* FloatTrack = Cast<const UMovieSceneFloatTrack>(&Track))
 				{
@@ -3836,10 +3963,14 @@ FUsdLevelSequenceHelperImpl::FLayerTimeInfo& FUsdLevelSequenceHelperImpl::FindOr
 	FLayerTimeInfo LayerTimeInfo;
 	UpdateLayerTimeInfoFromLayer(LayerTimeInfo, Layer);
 
-	UE_LOG(LogUsd, Verbose, TEXT("Creating layer time info for layer '%s'. Original timecodes: ['%s', '%s']"),
+	UE_LOG(
+		LogUsd,
+		Verbose,
+		TEXT("Creating layer time info for layer '%s'. Original timecodes: ['%s', '%s']"),
 		*LayerTimeInfo.Identifier,
 		LayerTimeInfo.StartTimeCode.IsSet() ? *LexToString(LayerTimeInfo.StartTimeCode.GetValue()) : TEXT("null"),
-		LayerTimeInfo.EndTimeCode.IsSet() ? *LexToString(LayerTimeInfo.EndTimeCode.GetValue()) : TEXT("null"));
+		LayerTimeInfo.EndTimeCode.IsSet() ? *LexToString(LayerTimeInfo.EndTimeCode.GetValue()) : TEXT("null")
+	);
 
 	return LayerTimeInfosByLayerIdentifier.Add(Layer.GetIdentifier(), LayerTimeInfo);
 }
@@ -3862,9 +3993,13 @@ void FUsdLevelSequenceHelperImpl::UpdateLayerTimeInfoFromLayer(FLayerTimeInfo& L
 	LayerTimeInfo.StartTimeCode = Layer.HasStartTimeCode() ? Layer.GetStartTimeCode() : TOptional<double>();
 	LayerTimeInfo.EndTimeCode = Layer.HasEndTimeCode() ? Layer.GetEndTimeCode() : TOptional<double>();
 
-	if (LayerTimeInfo.StartTimeCode.IsSet() && LayerTimeInfo.EndTimeCode.IsSet() && LayerTimeInfo.EndTimeCode.GetValue() < LayerTimeInfo.StartTimeCode.GetValue())
+	if (LayerTimeInfo.StartTimeCode.IsSet() && LayerTimeInfo.EndTimeCode.IsSet()
+		&& LayerTimeInfo.EndTimeCode.GetValue() < LayerTimeInfo.StartTimeCode.GetValue())
 	{
-		UE_LOG(LogUsd, Warning, TEXT("Sublayer '%s' has end time code (%f) before start time code (%f)! These values will be automatically swapped"),
+		UE_LOG(
+			LogUsd,
+			Warning,
+			TEXT("Sublayer '%s' has end time code (%f) before start time code (%f)! These values will be automatically swapped"),
 			*Layer.GetIdentifier(),
 			LayerTimeInfo.EndTimeCode.GetValue(),
 			LayerTimeInfo.StartTimeCode.GetValue()
@@ -3875,7 +4010,7 @@ void FUsdLevelSequenceHelperImpl::UpdateLayerTimeInfoFromLayer(FLayerTimeInfo& L
 		LayerTimeInfo.EndTimeCode = Temp;
 	}
 
-	const TArray< FString >& SubLayerPaths = Layer.GetSubLayerPaths();
+	const TArray<FString>& SubLayerPaths = Layer.GetSubLayerPaths();
 	LayerTimeInfo.SubLayersOffsets.Empty(SubLayerPaths.Num());
 
 	int32 SubLayerIndex = 0;
@@ -3911,36 +4046,73 @@ ULevelSequence* FUsdLevelSequenceHelperImpl::FindSequenceForIdentifier(const FSt
 class FUsdLevelSequenceHelperImpl
 {
 public:
-	FUsdLevelSequenceHelperImpl() {}
-	~FUsdLevelSequenceHelperImpl() {}
+	FUsdLevelSequenceHelperImpl()
+	{
+	}
+	~FUsdLevelSequenceHelperImpl()
+	{
+	}
 
-	ULevelSequence* Init(const UE::FUsdStage& InUsdStage) { return nullptr; }
-	void SetInfoCache(TSharedPtr<FUsdInfoCache> InfoCache) {};
-	void SetBBoxCache(TSharedPtr<UE::FUsdGeomBBoxCache> InBBoxCache) {};
-	bool HasData() const { return false; };
-	void Clear() {};
+	ULevelSequence* Init(const UE::FUsdStage& InUsdStage)
+	{
+		return nullptr;
+	}
+	void SetInfoCache(TSharedPtr<FUsdInfoCache> InfoCache){};
+	void SetBBoxCache(TSharedPtr<UE::FUsdGeomBBoxCache> InBBoxCache){};
+	bool HasData() const
+	{
+		return false;
+	};
+	void Clear(){};
 
-	void CreateLocalLayersSequences() {}
+	void CreateLocalLayersSequences()
+	{
+	}
 
-	void BindToUsdStageActor(AUsdStageActor* InStageActor) {}
-	void UnbindFromUsdStageActor() {}
-	EUsdRootMotionHandling GetRootMotionHandling() const {return EUsdRootMotionHandling::NoAdditionalRootMotion;}
+	void BindToUsdStageActor(AUsdStageActor* InStageActor)
+	{
+	}
+	void UnbindFromUsdStageActor()
+	{
+	}
+	EUsdRootMotionHandling GetRootMotionHandling() const
+	{
+		return EUsdRootMotionHandling::NoAdditionalRootMotion;
+	}
 	void SetRootMotionHandling(EUsdRootMotionHandling NewValue){};
-	void OnStageActorRenamed() {};
+	void OnStageActorRenamed(){};
 
-	void AddPrim(UUsdPrimTwin& PrimTwin, bool bForceVisibilityTracks, TOptional<bool> HasAnimatedBounds) {}
-	void RemovePrim(const UUsdPrimTwin& PrimTwin) {}
+	void AddPrim(UUsdPrimTwin& PrimTwin, bool bForceVisibilityTracks, TOptional<bool> HasAnimatedBounds)
+	{
+	}
+	void RemovePrim(const UUsdPrimTwin& PrimTwin)
+	{
+	}
 
-	void UpdateControlRigTracks(UUsdPrimTwin& PrimTwin) {}
+	void UpdateControlRigTracks(UUsdPrimTwin& PrimTwin)
+	{
+	}
 
-	void StartMonitoringChanges() {}
-	void StopMonitoringChanges() {}
-	void BlockMonitoringChangesForThisTransaction() {}
+	void StartMonitoringChanges()
+	{
+	}
+	void StopMonitoringChanges()
+	{
+	}
+	void BlockMonitoringChangesForThisTransaction()
+	{
+	}
 
-	ULevelSequence* GetMainLevelSequence() const { return nullptr; }
-	TArray< ULevelSequence* > GetSubSequences() const { return {}; }
+	ULevelSequence* GetMainLevelSequence() const
+	{
+		return nullptr;
+	}
+	TArray<ULevelSequence*> GetSubSequences() const
+	{
+		return {};
+	}
 };
-#endif // USE_USD_SDK
+#endif	  // USE_USD_SDK
 
 FUsdLevelSequenceHelper::FUsdLevelSequenceHelper()
 {
@@ -3949,22 +4121,22 @@ FUsdLevelSequenceHelper::FUsdLevelSequenceHelper()
 
 FUsdLevelSequenceHelper::~FUsdLevelSequenceHelper() = default;
 
-FUsdLevelSequenceHelper::FUsdLevelSequenceHelper(const FUsdLevelSequenceHelper & Other)
+FUsdLevelSequenceHelper::FUsdLevelSequenceHelper(const FUsdLevelSequenceHelper& Other)
 	: FUsdLevelSequenceHelper()
 {
 }
 
-FUsdLevelSequenceHelper& FUsdLevelSequenceHelper::operator=(const FUsdLevelSequenceHelper & Other)
+FUsdLevelSequenceHelper& FUsdLevelSequenceHelper::operator=(const FUsdLevelSequenceHelper& Other)
 {
 	// No copying, start fresh
 	UsdSequencerImpl = MakeUnique<FUsdLevelSequenceHelperImpl>();
 	return *this;
 }
 
-FUsdLevelSequenceHelper::FUsdLevelSequenceHelper(FUsdLevelSequenceHelper && Other) = default;
-FUsdLevelSequenceHelper& FUsdLevelSequenceHelper::operator=(FUsdLevelSequenceHelper && Other) = default;
+FUsdLevelSequenceHelper::FUsdLevelSequenceHelper(FUsdLevelSequenceHelper&& Other) = default;
+FUsdLevelSequenceHelper& FUsdLevelSequenceHelper::operator=(FUsdLevelSequenceHelper&& Other) = default;
 
-ULevelSequence* FUsdLevelSequenceHelper::Init(const UE::FUsdStage & UsdStage)
+ULevelSequence* FUsdLevelSequenceHelper::Init(const UE::FUsdStage& UsdStage)
 {
 	if (UsdSequencerImpl.IsValid())
 	{
@@ -4018,7 +4190,7 @@ void FUsdLevelSequenceHelper::Clear()
 	}
 }
 
-void FUsdLevelSequenceHelper::BindToUsdStageActor(AUsdStageActor * StageActor)
+void FUsdLevelSequenceHelper::BindToUsdStageActor(AUsdStageActor* StageActor)
 {
 	if (UsdSequencerImpl.IsValid())
 	{
@@ -4052,7 +4224,7 @@ void FUsdLevelSequenceHelper::SetRootMotionHandling(EUsdRootMotionHandling NewVa
 	}
 }
 
-void FUsdLevelSequenceHelper::AddPrim(UUsdPrimTwin & PrimTwin, bool bForceVisibilityTracks, TOptional<bool> HasAnimatedBounds)
+void FUsdLevelSequenceHelper::AddPrim(UUsdPrimTwin& PrimTwin, bool bForceVisibilityTracks, TOptional<bool> HasAnimatedBounds)
 {
 	if (UsdSequencerImpl.IsValid())
 	{
@@ -4060,7 +4232,7 @@ void FUsdLevelSequenceHelper::AddPrim(UUsdPrimTwin & PrimTwin, bool bForceVisibi
 	}
 }
 
-void FUsdLevelSequenceHelper::RemovePrim(const UUsdPrimTwin & PrimTwin)
+void FUsdLevelSequenceHelper::RemovePrim(const UUsdPrimTwin& PrimTwin)
 {
 	if (UsdSequencerImpl.IsValid())
 	{
@@ -4068,7 +4240,7 @@ void FUsdLevelSequenceHelper::RemovePrim(const UUsdPrimTwin & PrimTwin)
 	}
 }
 
-void FUsdLevelSequenceHelper::UpdateControlRigTracks(UUsdPrimTwin & PrimTwin)
+void FUsdLevelSequenceHelper::UpdateControlRigTracks(UUsdPrimTwin& PrimTwin)
 {
 	if (UsdSequencerImpl.IsValid())
 	{
@@ -4112,7 +4284,7 @@ ULevelSequence* FUsdLevelSequenceHelper::GetMainLevelSequence() const
 	}
 }
 
-TArray< ULevelSequence* > FUsdLevelSequenceHelper::GetSubSequences() const
+TArray<ULevelSequence*> FUsdLevelSequenceHelper::GetSubSequences() const
 {
 	if (UsdSequencerImpl.IsValid())
 	{
@@ -4132,19 +4304,19 @@ FUsdLevelSequenceHelper::FOnSkelAnimationBaked& FUsdLevelSequenceHelper::GetOnSk
 		return UsdSequencerImpl->GetOnSkelAnimationBaked();
 	}
 	else
-#endif // USE_USD_SDK
+#endif	  // USE_USD_SDK
 	{
 		static FOnSkelAnimationBaked DefaultHandler;
 		return DefaultHandler;
 	}
 }
 
-FScopedBlockMonitoringChangesForTransaction::FScopedBlockMonitoringChangesForTransaction(FUsdLevelSequenceHelper & InHelper)
+FScopedBlockMonitoringChangesForTransaction::FScopedBlockMonitoringChangesForTransaction(FUsdLevelSequenceHelper& InHelper)
 	: FScopedBlockMonitoringChangesForTransaction(*InHelper.UsdSequencerImpl.Get())
 {
 }
 
-FScopedBlockMonitoringChangesForTransaction::FScopedBlockMonitoringChangesForTransaction(FUsdLevelSequenceHelperImpl & InHelperImpl)
+FScopedBlockMonitoringChangesForTransaction::FScopedBlockMonitoringChangesForTransaction(FUsdLevelSequenceHelperImpl& InHelperImpl)
 	: HelperImpl(InHelperImpl)
 {
 	// If we're transacting we can just call this and the helper will unblock itself once the transaction is finished, because
diff --git a/Engine/Plugins/Importers/USDImporter/Source/USDStage/Private/USDStageActor.cpp b/Engine/Plugins/Importers/USDImporter/Source/USDStage/Private/USDStageActor.cpp
--- a/Engine/Plugins/Importers/USDImporter/Source/USDStage/Private/USDStageActor.cpp
+++ b/Engine/Plugins/Importers/USDImporter/Source/USDStage/Private/USDStageActor.cpp
@@ -5,7 +5,6 @@
 #include "UnrealUSDWrapper.h"
 #include "UObject/Package.h"
 #include "USDAssetCache.h"
-#include "USDAssetImportData.h"
 #include "USDAssetUserData.h"
 #include "USDClassesModule.h"
 #include "USDConversionUtils.h"
@@ -13,7 +12,6 @@
 #include "USDDynamicBindingResolverLibrary.h"
 #include "USDErrorUtils.h"
 #include "USDGeomMeshConversion.h"
-#include "USDGeomMeshTranslator.h"
 #include "USDGeomXformableTranslator.h"
 #include "USDInfoCache.h"
 #include "USDIntegrationUtils.h"
@@ -31,14 +29,11 @@
 #include "USDTransactor.h"
 #include "USDTypesConversion.h"
 
-#include "UsdWrappers/SdfChangeBlock.h"
 #include "UsdWrappers/SdfLayer.h"
-#include "UsdWrappers/UsdAttribute.h"
 #include "UsdWrappers/UsdGeomBBoxCache.h"
 #include "UsdWrappers/UsdGeomXformable.h"
 #include "UsdWrappers/UsdStage.h"
 
-#include "Async/ParallelFor.h"
 #include "CineCameraActor.h"
 #include "CineCameraComponent.h"
 #include "Components/DirectionalLightComponent.h"
@@ -49,62 +44,48 @@
 #include "Components/SkeletalMeshComponent.h"
 #include "Components/SkyLightComponent.h"
 #include "Components/SpotLightComponent.h"
-#include "Components/StaticMeshComponent.h"
 #include "Engine/Blueprint.h"
 #include "Engine/BlueprintGeneratedClass.h"
 #include "Engine/Engine.h"
 #include "Engine/Level.h"
 #include "Engine/Light.h"
-#include "Engine/SkeletalMesh.h"
-#include "Engine/StaticMesh.h"
 #include "Engine/World.h"
 #include "EngineAnalytics.h"
 #include "Framework/Notifications/NotificationManager.h"
-#include "HAL/FileManager.h"
 #include "HAL/IConsoleManager.h"
-#include "InterchangeGenericMaterialPipeline.h"
 #include "LevelSequence.h"
 #include "LiveLinkComponentController.h"
-#include "Materials/Material.h"
-#include "Materials/MaterialInstanceConstant.h"
 #include "Materials/MaterialInterface.h"
-#include "MeshDescription.h"
 #include "Misc/Paths.h"
 #include "Misc/ScopedSlowTask.h"
 #include "Modules/ModuleManager.h"
 #include "MovieScene.h"
-#include "PhysicsEngine/BodySetup.h"
-#include "Rendering/SkeletalMeshLODImporterData.h"
 #include "Roles/LiveLinkTransformRole.h"
 #include "Sections/MovieSceneSubSection.h"
-#include "StaticMeshAttributes.h"
-#include "StaticMeshOperations.h"
 #include "Tracks/MovieScene3DTransformTrack.h"
 #include "Widgets/Notifications/SNotificationList.h"
 
 #if WITH_EDITOR
-#include "USDClassesEditorModule.h"
-
 #include "BlueprintActionMenuItem.h"
 #include "BlueprintFunctionNodeSpawner.h"
 #include "EdGraphSchema_K2.h"
 #include "Editor.h"
 #include "Editor/TransBuffer.h"
 #include "Editor/UnrealEdEngine.h"
-#include "Engine/Selection.h"
 #include "ILevelSequenceEditorToolkit.h"
 #include "ISequencerModule.h"
 #include "K2Node_FunctionEntry.h"
 #include "K2Node_FunctionResult.h"
 #include "Kismet2/KismetEditorUtilities.h"
-#include "LevelEditor.h"
 #include "MovieSceneDirectorBlueprintUtils.h"
 #include "MovieSceneDynamicBindingUtils.h"
-#include "PropertyEditorModule.h"
 #include "ScopedTransaction.h"
+#include "Selection.h"
 #include "Subsystems/AssetEditorSubsystem.h"
 #include "UnrealEdGlobals.h"
-#endif // WITH_EDITOR
+#include "USDClassesEditorModule.h"
+
+#endif	  // WITH_EDITOR
 
 #define LOCTEXT_NAMESPACE "USDStageActor"
 
@@ -112,13 +93,17 @@ static bool GRegenerateSkeletalAssetsOnControlRigBake = true;
 static FAutoConsoleVariableRef CVarRegenerateSkeletalAssetsOnControlRigBake(
 	TEXT("USD.RegenerateSkeletalAssetsOnControlRigBake"),
 	GRegenerateSkeletalAssetsOnControlRigBake,
-	TEXT("Whether to regenerate the assets associated with a SkelRoot (mesh, skeleton, anim sequence, etc.) whenever we modify Control Rig tracks. The USD Stage itself is always updated however."));
+	TEXT("Whether to regenerate the assets associated with a SkelRoot (mesh, skeleton, anim sequence, etc.) whenever we modify Control Rig tracks. "
+		 "The USD Stage itself is always updated however.")
+);
 
 static bool GTranslateOnlyUsedMaterialsWhenOpeningStage = true;
 static FAutoConsoleVariableRef CVarTranslateOnlyUsedMaterialsWhenOpeningStage(
 	TEXT("USD.TranslateOnlyUsedMaterialsWhenOpeningStage"),
 	GTranslateOnlyUsedMaterialsWhenOpeningStage,
-	TEXT("If enabled, only Material prims bound by at least one Mesh are translated into Unreal material assets. If disabled, all Material prims are translated into Unreal material assets."));
+	TEXT("If enabled, only Material prims bound by at least one Mesh are translated into Unreal material assets. If disabled, all Material prims are "
+		 "translated into Unreal material assets.")
+);
 
 static const EObjectFlags DefaultObjFlag = EObjectFlags::RF_Transactional | EObjectFlags::RF_Transient;
 
@@ -127,9 +112,9 @@ AUsdStageActor::FOnOpenStageEditorClicked AUsdStageActor::OnOpenStageEditorClick
 
 struct FUsdStageActorImpl
 {
-	static TSharedRef< FUsdSchemaTranslationContext > CreateUsdSchemaTranslationContext(AUsdStageActor* StageActor, const FString& PrimPath)
+	static TSharedRef<FUsdSchemaTranslationContext> CreateUsdSchemaTranslationContext(AUsdStageActor* StageActor, const FString& PrimPath)
 	{
-		TSharedRef< FUsdSchemaTranslationContext > TranslationContext = MakeShared< FUsdSchemaTranslationContext >(
+		TSharedRef<FUsdSchemaTranslationContext> TranslationContext = MakeShared<FUsdSchemaTranslationContext>(
 			StageActor->GetOrOpenUsdStage(),
 			*StageActor->UsdAssetCache
 		);
@@ -198,19 +183,22 @@ struct FUsdStageActorImpl
 			TArray<UObject*> ComponentsToDeselect;
 
 			const bool bRecursive = true;
-			StageActor->GetRootPrimTwin()->Iterate([&ActorsToDeselect, &ComponentsToDeselect](UUsdPrimTwin& PrimTwin)
-			{
-				if (USceneComponent* ReferencedComponent = PrimTwin.SceneComponent.Get())
+			StageActor->GetRootPrimTwin()->Iterate(
+				[&ActorsToDeselect, &ComponentsToDeselect](UUsdPrimTwin& PrimTwin)
 				{
-					ComponentsToDeselect.Add(ReferencedComponent);
-
-					AActor* Owner = ReferencedComponent->GetOwner();
-					if (Owner && Owner->GetRootComponent() == ReferencedComponent)
+					if (USceneComponent* ReferencedComponent = PrimTwin.SceneComponent.Get())
 					{
-						ActorsToDeselect.Add(Owner);
+						ComponentsToDeselect.Add(ReferencedComponent);
+
+						AActor* Owner = ReferencedComponent->GetOwner();
+						if (Owner && Owner->GetRootComponent() == ReferencedComponent)
+						{
+							ActorsToDeselect.Add(Owner);
+						}
 					}
-				}
-			}, bRecursive);
+				},
+				bRecursive
+			);
 
 			if (USelection* SelectedComponents = GEditor->GetSelectedComponents())
 			{
@@ -236,16 +224,16 @@ struct FUsdStageActorImpl
 				}
 			}
 
-			if (bDeselected && GIsEditor) // Make sure we're not in standalone either
+			if (bDeselected && GIsEditor)	 // Make sure we're not in standalone either
 			{
 				GEditor->NoteSelectionChange();
 			}
 		}
-#endif // WITH_EDITOR
+#endif	  // WITH_EDITOR
 	}
 
 	template<typename ObjectPtr>
-	static void CloseEditorsForAssets(const TMap< FString, ObjectPtr >& AssetsCache)
+	static void CloseEditorsForAssets(const TMap<FString, ObjectPtr>& AssetsCache)
 	{
 #if WITH_EDITOR
 		if (GIsEditor && GEditor)
@@ -261,7 +249,7 @@ struct FUsdStageActorImpl
 				}
 			}
 		}
-#endif // WITH_EDITOR
+#endif	  // WITH_EDITOR
 	}
 
 	static void DiscardStage(const UE::FUsdStage& Stage, AUsdStageActor* DiscardingActor)
@@ -430,7 +418,7 @@ struct FUsdStageActorImpl
 		{
 			return FPaths::IsSamePath(UsdUtils::MakePathRelativeToProjectDir(Layer.GetRealPath()), FilePath);
 		}
-#endif // USE_USD_SDK
+#endif	  // USE_USD_SDK
 
 		return false;
 	}
@@ -603,7 +591,7 @@ struct FUsdStageActorImpl
 				}
 			}
 		}
-#endif // WITH_EDITOR
+#endif	  // WITH_EDITOR
 	}
 
 	static void SetupDynamicBinding(
@@ -652,19 +640,15 @@ struct FUsdStageActorImpl
 
 			// Store a path to this very actor on the binding, so that it can find us later and ask how to resolve a particular
 			// prim path
-			FMovieSceneDynamicBindingPayloadVariable& ActorPathVariable = Possessable.DynamicBinding.PayloadVariables.FindOrAdd(
-				TEXT("StageActorIDNameFilter")
-			);
+			FMovieSceneDynamicBindingPayloadVariable& ActorPathVariable = Possessable.DynamicBinding.PayloadVariables.FindOrAdd(TEXT("StageActorIDNam"
+																																	 "eFilter"));
 			ActorPathVariable.Value = DefaultActorFilter;
 
-			FMovieSceneDynamicBindingPayloadVariable& RootLayerVariable = Possessable.DynamicBinding.PayloadVariables.FindOrAdd(
-				TEXT("RootLayerFilter")
-			);
+			FMovieSceneDynamicBindingPayloadVariable& RootLayerVariable = Possessable.DynamicBinding.PayloadVariables.FindOrAdd(TEXT("RootLayerFilter"
+			));
 			RootLayerVariable.Value = FString{};	// No root layer filter by default for more flexibility
 
-			FMovieSceneDynamicBindingPayloadVariable& PrimPathVariable = Possessable.DynamicBinding.PayloadVariables.FindOrAdd(
-				TEXT("PrimPath")
-			);
+			FMovieSceneDynamicBindingPayloadVariable& PrimPathVariable = Possessable.DynamicBinding.PayloadVariables.FindOrAdd(TEXT("PrimPath"));
 			PrimPathVariable.Value = PrimPath;
 		}
 
@@ -1018,12 +1002,10 @@ struct FRecompilationTracker
 
 	static void OnCompilationStarted(UBlueprint* BP)
 	{
-		// We don't care if a BP is compiling on first load: It only matters to use if we're compiling one that already has loaded instances on the level
-		if (!BP ||
-			BP->bIsRegeneratingOnLoad ||
-			!BP->GeneratedClass ||
-			!BP->GeneratedClass->IsChildOf(AUsdStageActor::StaticClass()) ||
-			RecompilingBlueprints.Contains(BP))
+		// We don't care if a BP is compiling on first load: It only matters to use if we're compiling one that already has loaded instances on the
+		// level
+		if (!BP || BP->bIsRegeneratingOnLoad || !BP->GeneratedClass || !BP->GeneratedClass->IsChildOf(AUsdStageActor::StaticClass())
+			|| RecompilingBlueprints.Contains(BP))
 		{
 			return;
 		}
@@ -1052,7 +1034,7 @@ private:
 };
 bool FRecompilationTracker::bEventIsSetup = false;
 TMap<UBlueprint*, FDelegateHandle> FRecompilationTracker::RecompilingBlueprints;
-#endif // WITH_EDITOR
+#endif	  // WITH_EDITOR
 
 AUsdStageActor::AUsdStageActor()
 	: StageState(EUsdStageState::OpenedAndLoaded)
@@ -1067,18 +1049,18 @@ AUsdStageActor::AUsdStageActor()
 	, RootMotionHandling(EUsdRootMotionHandling::NoAdditionalRootMotion)
 	, SubdivisionLevel(0)
 	, MetadataOptions(FUsdMetadataImportOptions{
-		false, 	/* bCollectMetadata */
-		false, 	/* bCollectFromEntireSubtrees */
-		false, 	/* bCollectOnComponents */
-		{},		/* BlockedPrefixFilters */
-		false	/* bInvertFilters */
-	})
+		  false, /* bCollectMetadata */
+		  false, /* bCollectFromEntireSubtrees */
+		  false, /* bCollectOnComponents */
+		  {},	 /* BlockedPrefixFilters */
+		  false	 /* bInvertFilters */
+	  })
 	, Time(0.0f)
 	, bIsTransitioningIntoPIE(false)
 	, bIsModifyingAProperty(false)
 	, bIsUndoRedoing(false)
 {
-	SceneComponent = CreateDefaultSubobject< USceneComponent >(TEXT("SceneComponent0"));
+	SceneComponent = CreateDefaultSubobject<USceneComponent>(TEXT("SceneComponent0"));
 	SceneComponent->Mobility = EComponentMobility::Static;
 
 	RootComponent = SceneComponent;
@@ -1090,7 +1072,7 @@ AUsdStageActor::AUsdStageActor()
 	// some transitions like reinstantiation.
 	// c.f. doc comment on FRecompilationTracker for more info.
 
-	IUsdSchemasModule& UsdSchemasModule = FModuleManager::Get().LoadModuleChecked< IUsdSchemasModule >(TEXT("USDSchemas"));
+	IUsdSchemasModule& UsdSchemasModule = FModuleManager::Get().LoadModuleChecked<IUsdSchemasModule>(TEXT("USDSchemas"));
 	RenderContext = UsdSchemasModule.GetRenderContextRegistry().GetUnrealRenderContext();
 
 	const FName UniqueName = MakeUniqueObjectName(this, UUsdTransactor::StaticClass(), TEXT("Transactor"));
@@ -1110,11 +1092,14 @@ AUsdStageActor::AUsdStageActor()
 			FProperty* Property = *PropertyIterator;
 			if (Property && Property->GetFName() == GET_MEMBER_NAME_CHECKED(AUsdStageActor, RootLayer))
 			{
-				TArray< FString > SupportedExtensions = UnrealUSDWrapper::GetAllSupportedFileFormats();
+				TArray<FString> SupportedExtensions = UnrealUSDWrapper::GetAllSupportedFileFormats();
 				if (SupportedExtensions.Num() > 0)
 				{
-					FString JoinedExtensions = FString::Join(SupportedExtensions, TEXT("; *.")); // Combine "usd" and "usda" into "usd; *.usda"
-					Property->SetMetaData(TEXT("FilePathFilter"), FString::Printf(TEXT("Universal Scene Description files|*.%s"), *JoinedExtensions, *JoinedExtensions));
+					FString JoinedExtensions = FString::Join(SupportedExtensions, TEXT("; *."));	// Combine "usd" and "usda" into "usd; *.usda"
+					Property->SetMetaData(
+						TEXT("FilePathFilter"),
+						FString::Printf(TEXT("Universal Scene Description files|*.%s"), *JoinedExtensions, *JoinedExtensions)
+					);
 				}
 				break;
 			}
@@ -1140,10 +1125,9 @@ AUsdStageActor::AUsdStageActor()
 				[this](const FTransactionContext& TransactionContext, bool bSucceeded)
 				{
 					// This text should match the one in ConcertClientTransactionBridge.cpp
-					if (this &&
-						HasAuthorityOverStage() &&
-						TransactionContext.Title.EqualTo(LOCTEXT("ConcertTransactionEvent", "Concert Transaction Event")) &&
-						!RootLayer.FilePath.IsEmpty())
+					if (this && HasAuthorityOverStage()
+						&& TransactionContext.Title.EqualTo(LOCTEXT("ConcertTransactionEvent", "Concert Transaction Event"))
+						&& !RootLayer.FilePath.IsEmpty())
 					{
 						// Other user deleted us
 						if (!IsValid(this))
@@ -1152,7 +1136,8 @@ AUsdStageActor::AUsdStageActor()
 						}
 						// We have a valid filepath but no objects/assets spawned, so it's likely we were just spawned on the
 						// other client, and were replicated here with our RootLayer path already filled out, meaning we should just load that stage
-						// Note that now our UUsdTransactor may have already caused the stage itself to be loaded, but we may still need to call LoadUsdStage on our end.
+						// Note that now our UUsdTransactor may have already caused the stage itself to be loaded, but we may still need to call
+						// LoadUsdStage on our end.
 						else if (ObjectsToWatch.Num() == 0 && (!UsdAssetCache || UsdAssetCache->GetNumAssets() == 0))
 						{
 							this->LoadUsdStage();
@@ -1179,8 +1164,7 @@ AUsdStageActor::AUsdStageActor()
 			}
 		}
 		ISequencerModule& SequencerModule = FModuleManager::Get().LoadModuleChecked<ISequencerModule>(TEXT("Sequencer"));
-		OnSequencerCreatedHandle = SequencerModule.RegisterOnSequencerCreated(
-			FOnSequencerCreated::FDelegate::CreateLambda(SubscribeToSequencerEvents)
+		OnSequencerCreatedHandle = SequencerModule.RegisterOnSequencerCreated(FOnSequencerCreated::FDelegate::CreateLambda(SubscribeToSequencerEvents)
 		);
 
 		FCoreUObjectDelegates::OnObjectPropertyChanged.AddUObject(this, &AUsdStageActor::OnObjectPropertyChanged);
@@ -1197,7 +1181,7 @@ AUsdStageActor::AUsdStageActor()
 
 		LevelSequenceHelper.GetOnSkelAnimationBaked().AddUObject(this, &AUsdStageActor::OnSkelAnimationBaked);
 
-#endif // WITH_EDITOR
+#endif	  // WITH_EDITOR
 
 		OnTimeChanged.AddUObject(this, &AUsdStageActor::AnimatePrims);
 
@@ -1291,7 +1275,7 @@ void AUsdStageActor::NewStage()
 	{
 		return;
 	}
-	FString	StagePath = FString(UnrealIdentifiers::IdentifierPrefix) + Layer.GetIdentifier();
+	FString StagePath = FString(UnrealIdentifiers::IdentifierPrefix) + Layer.GetIdentifier();
 
 	UE::FUsdPrim RootPrim = NewStage.DefinePrim(UE::FSdfPath{TEXT("/Root")}, TEXT("Xform"));
 	ensure(UsdUtils::SetDefaultKind(RootPrim, EUsdDefaultKind::Assembly));
@@ -1317,7 +1301,7 @@ void AUsdStageActor::NewStage()
 	}
 
 	SetRootLayer(StagePath);
-#endif // USE_USD_SDK
+#endif	  // USE_USD_SDK
 }
 
 void AUsdStageActor::SetIsolatedRootLayer(const FString& IsolatedStageRootLayer)
@@ -1345,7 +1329,7 @@ FString AUsdStageActor::GetIsolatedRootLayer() const
 	return IsolatedStage ? IsolatedStage.GetRootLayer().GetIdentifier() : FString{};
 }
 
-void AUsdStageActor::IsolateLayer( const UE::FSdfLayer& Layer )
+void AUsdStageActor::IsolateLayer(const UE::FSdfLayer& Layer)
 {
 	if (IsolatedStage && IsolatedStage.GetRootLayer() == Layer)
 	{
@@ -1382,7 +1366,10 @@ void AUsdStageActor::IsolateLayer( const UE::FSdfLayer& Layer )
 		TArray<UE::FSdfLayer> CurrentLayerStack = FreshCurrentStage.GetLayerStack();
 		if (!CurrentLayerStack.Contains(Layer))
 		{
-			UE_LOG(LogUsd, Warning, TEXT("Failed to isolate layer '%s' as it is not part of the currently opened USD Stage's local layer stack"),
+			UE_LOG(
+				LogUsd,
+				Warning,
+				TEXT("Failed to isolate layer '%s' as it is not part of the currently opened USD Stage's local layer stack"),
 				*Layer.GetIdentifier(),
 				*UsdStage.GetRootLayer().GetIdentifier()
 			);
@@ -1449,9 +1436,10 @@ void AUsdStageActor::OnUsdObjectsChanged(const UsdUtils::FObjectChangesByPath& I
 	}
 
 	// Only update the transactor if we're listening to USD notices. Within OnObjectPropertyChanged we will stop listening when writing stage changes
-	// from our component changes, and this will also make sure we're not duplicating the events we store and replicate via multi-user: If a modification
-	// can be described purely via UObject changes, then those changes will be responsible for the whole modification and we won't record the corresponding
-	// stage changes. The intent is that when undo/redo/replicating that UObject change, it will automatically generate the corresponding stage changes
+	// from our component changes, and this will also make sure we're not duplicating the events we store and replicate via multi-user: If a
+	// modification can be described purely via UObject changes, then those changes will be responsible for the whole modification and we won't record
+	// the corresponding stage changes. The intent is that when undo/redo/replicating that UObject change, it will automatically generate the
+	// corresponding stage changes
 	if (Transactor)
 	{
 		Transactor->Update(InfoChanges, ResyncChanges);
@@ -1464,8 +1452,9 @@ void AUsdStageActor::OnUsdObjectsChanged(const UsdUtils::FObjectChangesByPath& I
 	}
 
 	// If the stage was closed in a big transaction (e.g. undo open) a random UObject may be transacting before us and triggering USD changes,
-	// and the UE::FUsdStage will still be opened and valid (even though we intend on closing/changing it when we transact). It could be problematic/wasteful if we
-	// responded to those notices, so just early out here. We can do this check because our RootLayer property will already have the new value
+	// and the UE::FUsdStage will still be opened and valid (even though we intend on closing/changing it when we transact). It could be
+	// problematic/wasteful if we responded to those notices, so just early out here. We can do this check because our RootLayer property will already
+	// have the new value
 	{
 		const UE::FUsdStage& BaseStage = GetBaseUsdStage();
 		const UE::FSdfLayer& StageRoot = BaseStage.GetRootLayer();
@@ -1489,7 +1478,7 @@ void AUsdStageActor::OnUsdObjectsChanged(const UsdUtils::FObjectChangesByPath& I
 			LogUsd,
 			Warning,
 			TEXT("Closing the stage '%s' as no asset cache was provided, but the AUsdStageActor '%s' was set to open the stage and load assets. "
-					"Either provide an asset cache or switch the stage actor to the 'Opened' state"),
+				 "Either provide an asset cache or switch the stage actor to the 'Opened' state"),
 			*RootLayer.FilePath,
 			*GetPathName()
 		);
@@ -1539,7 +1528,7 @@ void AUsdStageActor::OnUsdObjectsChanged(const UsdUtils::FObjectChangesByPath& I
 				// TODO: Better way of handling material override updates. Ideally we wouldn't be resyncing
 				// the Mesh prim just to update material overrides...
 				UE::FUsdPrim ChangedPrim = Stage.GetPrimAtPath(PrimPath);
-				if(ChangedPrim && ChangedPrim.IsA(TEXT("GeomSubset")))
+				if (ChangedPrim && ChangedPrim.IsA(TEXT("GeomSubset")))
 				{
 					PrimPath = PrimPath.GetParentPath();
 				}
@@ -1549,10 +1538,7 @@ void AUsdStageActor::OnUsdObjectsChanged(const UsdUtils::FObjectChangesByPath& I
 
 			for (const UsdUtils::FAttributeChange& AttributeChange : ObjectChange.AttributeChanges)
 			{
-				static const TSet<FString> StageResyncProperties = {
-					TEXT("metersPerUnit"),
-					TEXT("upAxis")
-				};
+				static const TSet<FString> StageResyncProperties = {TEXT("metersPerUnit"), TEXT("upAxis")};
 
 				// Upgrade these to resync so that the prim twins are regenerated, which clears all the existing
 				// animation tracks and adds new ones, automatically re-baking to control rig
@@ -1570,8 +1556,7 @@ void AUsdStageActor::OnUsdObjectsChanged(const UsdUtils::FObjectChangesByPath& I
 					// When we change these UsdGeomModelAPI attributes we may need to create a new component type
 					// for the prim (as it may now need/stop needing an alternate draw mode component)
 					UnrealIdentifiers::ModelDrawMode,
-					UnrealIdentifiers::ModelApplyDrawMode
-				};
+					UnrealIdentifiers::ModelApplyDrawMode};
 
 				// Some stage info should trigger some resyncs because they should trigger reparsing of geometry
 				if ((PrimPath.IsAbsoluteRootPath() && StageResyncProperties.Contains(AttributeChange.PropertyName))
@@ -1593,8 +1578,8 @@ void AUsdStageActor::OnUsdObjectsChanged(const UsdUtils::FObjectChangesByPath& I
 		}
 
 		// We may need the full spec path with variant selections later, but for traversal and retrieving prims from the stage we always need
-		// the prim path without any variant selections in it (i.e. GetPrimAtPath("/Root{Varset=Var}Child") doesn't work, we need GetPrimAtPath("/Root/Child")),
-		// and USD sometimes emits changes with the variant selection path (like during renames).
+		// the prim path without any variant selections in it (i.e. GetPrimAtPath("/Root{Varset=Var}Child") doesn't work, we need
+		// GetPrimAtPath("/Root/Child")), and USD sometimes emits changes with the variant selection path (like during renames).
 		SortedPrimsChangedList.Add(PrimPath, bIsResync);
 	}
 	// Do Resyncs after so that they overwrite pure info changes if we have any
@@ -1695,12 +1680,11 @@ void AUsdStageActor::OnUsdObjectsChanged(const UsdUtils::FObjectChangesByPath& I
 	// Traverses the info caches to find out which prims we need to update
 	TFunction<void(const UE::FSdfPath&, bool, TMap<UE::FSdfPath, bool>&, TSet<UE::FSdfPath>&)> RecursiveCollectPrimsToUpdate;
 	RecursiveCollectPrimsToUpdate =
-		[this, &OldInfoCache, &RecursiveCollectPrimsToUpdate, &KnownOldPrims, &Stage](
-			const UE::FSdfPath& PrimPath,
-			bool bIsResync,
-			TMap<UE::FSdfPath, bool>& OutPrimsToUpdate,
-			TSet<UE::FSdfPath>& InOutVisitedPaths
-		)
+		[this,
+		 &OldInfoCache,
+		 &RecursiveCollectPrimsToUpdate,
+		 &KnownOldPrims,
+		 &Stage](const UE::FSdfPath& PrimPath, bool bIsResync, TMap<UE::FSdfPath, bool>& OutPrimsToUpdate, TSet<UE::FSdfPath>& InOutVisitedPaths)
 	{
 		// Allow revisiting a prim if we're resyncing: This is a cheap way of making sure we mark the
 		// resynced prims as resynced, or else we may recurse into them while traversing another prim and preemptively
@@ -1762,7 +1746,7 @@ void AUsdStageActor::OnUsdObjectsChanged(const UsdUtils::FObjectChangesByPath& I
 				KnownOldPrims = OldInfoCache->GetKnownPrims();
 			}
 
-			if(UE::FUsdPrim Prim = Stage.GetPrimAtPath(PrimPath))
+			if (UE::FUsdPrim Prim = Stage.GetPrimAtPath(PrimPath))
 			{
 				for (const UE::FUsdPrim& ChildPrim : Prim.GetChildren())
 				{
@@ -1899,11 +1883,11 @@ void AUsdStageActor::OnUsdObjectsChanged(const UsdUtils::FObjectChangesByPath& I
 		}
 	}
 
-	if(MaterialUserPrims.Num())
+	if (MaterialUserPrims.Num())
 	{
 		TSet<UE::FSdfPath> VisitedPaths;
 		PrimsToUpdate.Reserve(PrimsToUpdate.Num() + MaterialUserPrims.Num());
-		for(const UE::FSdfPath& MaterialUserPrim : MaterialUserPrims)
+		for (const UE::FSdfPath& MaterialUserPrim : MaterialUserPrims)
 		{
 			// This also needs to be done recursivly because our "material user prim" may be just something
 			// like a UsdGeomSubset. If we're going to resync its component, we actually need to resync whatever
@@ -1989,7 +1973,7 @@ void AUsdStageActor::OnUsdObjectsChanged(const UsdUtils::FObjectChangesByPath& I
 	}
 
 #if WITH_EDITOR
-	if (GIsEditor && GEditor && !IsGarbageCollecting()) // Make sure we're not in standalone either
+	if (GIsEditor && GEditor && !IsGarbageCollecting())	   // Make sure we're not in standalone either
 	{
 		if (bHasResync)
 		{
@@ -1998,9 +1982,9 @@ void AUsdStageActor::OnUsdObjectsChanged(const UsdUtils::FObjectChangesByPath& I
 
 		GEditor->RedrawLevelEditingViewports();
 	}
-#endif // WITH_EDITOR
+#endif	  // WITH_EDITOR
 
-#endif // USE_USD_SDK
+#endif	  // USE_USD_SDK
 }
 
 USDSTAGE_API void AUsdStageActor::Reset()
@@ -2099,9 +2083,10 @@ UUsdPrimTwin* AUsdStageActor::ExpandPrim(const UE::FUsdPrim& Prim, bool bResync,
 
 	bool bExpandChildren = true;
 
-	IUsdSchemasModule& UsdSchemasModule = FModuleManager::Get().LoadModuleChecked< IUsdSchemasModule >(TEXT("USDSchemas"));
+	IUsdSchemasModule& UsdSchemasModule = FModuleManager::Get().LoadModuleChecked<IUsdSchemasModule>(TEXT("USDSchemas"));
 
-	if (TSharedPtr< FUsdSchemaTranslator > SchemaTranslator = UsdSchemasModule.GetTranslatorRegistry().CreateTranslatorForSchema(TranslationContext.AsShared(), UE::FUsdTyped(Prim)))
+	if (TSharedPtr<FUsdSchemaTranslator> SchemaTranslator = UsdSchemasModule.GetTranslatorRegistry()
+																.CreateTranslatorForSchema(TranslationContext.AsShared(), UE::FUsdTyped(Prim)))
 	{
 		if (bResync && !UsdPrimTwin->SceneComponent.IsValid())
 		{
@@ -2132,10 +2117,10 @@ UUsdPrimTwin* AUsdStageActor::ExpandPrim(const UE::FUsdPrim& Prim, bool bResync,
 			ContextParentComponent = UsdPrimTwin->SceneComponent.Get();
 		}
 
-		TGuardValue< USceneComponent* > ParentComponentGuard(TranslationContext.ParentComponent, ContextParentComponent);
+		TGuardValue<USceneComponent*> ParentComponentGuard(TranslationContext.ParentComponent, ContextParentComponent);
 
 		const bool bTraverseInstanceProxies = true;
-		const TArray< UE::FUsdPrim > PrimChildren = Prim.GetFilteredChildren(bTraverseInstanceProxies);
+		const TArray<UE::FUsdPrim> PrimChildren = Prim.GetFilteredChildren(bTraverseInstanceProxies);
 
 		for (const UE::FUsdPrim& ChildPrim : PrimChildren)
 		{
@@ -2148,7 +2133,7 @@ UUsdPrimTwin* AUsdStageActor::ExpandPrim(const UE::FUsdPrim& Prim, bool bResync,
 	{
 #if WITH_EDITOR
 		TwinSceneComponent->PostEditChange();
-#endif // WITH_EDITOR
+#endif	  // WITH_EDITOR
 
 		if (!TwinSceneComponent->IsRegistered())
 		{
@@ -2291,7 +2276,7 @@ UUsdPrimTwin* AUsdStageActor::ExpandPrim(const UE::FUsdPrim& Prim, bool bResync,
 		}
 	}
 
-#endif // USE_USD_SDK
+#endif	  // USE_USD_SDK
 	return UsdPrimTwin;
 }
 
@@ -2368,8 +2353,7 @@ void AUsdStageActor::SetUsdStage(const UE::FUsdStage& NewStage)
 		// Try loading a UE-state session layer if we can find one
 		const bool bCreateIfNeeded = false;
 		UsdUtils::GetUEPersistentStateSublayer(UsdStage, bCreateIfNeeded);
-#endif // #if USE_USD_SDK
-
+#endif	  // #if USE_USD_SDK
 	}
 
 	LoadUsdStage();
@@ -2398,7 +2382,7 @@ void AUsdStageActor::SetRootLayer(const FString& RootFilePath)
 	{
 		RelativeFilePath = UsdUtils::MakePathRelativeToProjectDir(RootFilePath);
 	}
-#endif // USE_USD_SDK
+#endif	  // USE_USD_SDK
 
 	// See if we're talking about the stage that is already loaded
 	if (UsdStage)
@@ -2485,7 +2469,11 @@ void AUsdStageActor::SetAssetCache(UUsdAssetCache2* NewCache)
 	if (!UsdAssetCache && UsdStage)
 	{
 		FNotificationInfo Toast(LOCTEXT("MustHaveCache", "Must have an Asset Cache"));
-		Toast.SubText = LOCTEXT("MustHaveCache_Subtext", "The Stage Actor must always have an Asset Cache while a stage is loaded, so a temporary cache will be created.\n\nClose the stage before clearing the cache if you wish to clear this property.");
+		Toast.SubText = LOCTEXT(
+			"MustHaveCache_Subtext",
+			"The Stage Actor must always have an Asset Cache while a stage is loaded, so a temporary cache will be created.\n\nClose the stage "
+			"before clearing the cache if you wish to clear this property."
+		);
 		Toast.Image = FCoreStyle::Get().GetBrush(TEXT("MessageLog.Warning"));
 		Toast.bUseLargeFont = false;
 		Toast.bFireAndForget = true;
@@ -2495,14 +2483,14 @@ void AUsdStageActor::SetAssetCache(UUsdAssetCache2* NewCache)
 		Toast.bUseSuccessFailIcons = false;
 		FSlateNotificationManager::Get().AddNotification(Toast);
 
-		UsdAssetCache = NewObject< UUsdAssetCache2 >(GetTransientPackage(), NAME_None, GetMaskedFlags(RF_PropagateToSubObjects));
+		UsdAssetCache = NewObject<UUsdAssetCache2>(GetTransientPackage(), NAME_None, GetMaskedFlags(RF_PropagateToSubObjects));
 	}
 
 	// Here we pretend we just received a root resync so that we re-fetch assets from the cache
 	// and update its components
 	UsdUtils::FObjectChangesByPath InfoChanges;
 	UsdUtils::FObjectChangesByPath ResyncChanges;
-	ResyncChanges.Add({ TEXT("/"), {} });
+	ResyncChanges.Add({TEXT("/"), {}});
 	OnUsdObjectsChanged(InfoChanges, ResyncChanges);
 }
 
@@ -2782,7 +2770,7 @@ ULevelSequence* AUsdStageActor::GetLevelSequence()
 	return LevelSequence;
 }
 
-USceneComponent* AUsdStageActor::GetGeneratedComponent( const FString& PrimPath )
+USceneComponent* AUsdStageActor::GetGeneratedComponent(const FString& PrimPath)
 {
 	const UE::FUsdStage& CurrentStage = static_cast<const AUsdStageActor*>(this)->GetUsdStage();
 	if (!CurrentStage)
@@ -2884,7 +2872,8 @@ void AUsdStageActor::OpenUsdStage()
 	FString AbsPath;
 	if (!RootLayer.FilePath.StartsWith(UnrealIdentifiers::IdentifierPrefix) && FPaths::IsRelative(RootLayer.FilePath))
 	{
-		// The RootLayer property is marked as RelativeToGameDir, and FUsdStageViewModel::OpenStage will also give us paths relative to the project's directory
+		// The RootLayer property is marked as RelativeToGameDir, and FUsdStageViewModel::OpenStage will also give us paths relative to the project's
+		// directory
 		FString ProjectDir = FPaths::ConvertRelativePathToFull(FPaths::ProjectDir());
 		AbsPath = FPaths::ConvertRelativePathToFull(FPaths::Combine(ProjectDir, RootLayer.FilePath));
 	}
@@ -2918,12 +2907,15 @@ void AUsdStageActor::OpenUsdStage()
 		// Try loading a UE-state session layer if we can find one
 		const bool bCreateIfNeeded = false;
 		UsdUtils::GetUEPersistentStateSublayer(UsdStage, bCreateIfNeeded);
-#endif // #if USE_USD_SDK
+#endif	  // #if USE_USD_SDK
 	}
 
 	OnStageChanged.Broadcast();
 
-	UsdUtils::ShowErrorsAndStopMonitoring(FText::Format(LOCTEXT("USDOpenError", "Encountered some errors opening USD file at path '{0}!\nCheck the Output Log for details."), FText::FromString(RootLayer.FilePath)));
+	UsdUtils::ShowErrorsAndStopMonitoring(FText::Format(
+		LOCTEXT("USDOpenError", "Encountered some errors opening USD file at path '{0}!\nCheck the Output Log for details."),
+		FText::FromString(RootLayer.FilePath)
+	));
 }
 
 void AUsdStageActor::CloseUsdStage(bool bUnloadIfNeeded)
@@ -3093,7 +3085,7 @@ void AUsdStageActor::OnLevelActorDeleted(AActor* DeletedActor)
 	}
 }
 
-#endif // WITH_EDITOR
+#endif	  // WITH_EDITOR
 
 void AUsdStageActor::LoadUsdStage(bool bOpenIfNeeded)
 {
@@ -3181,10 +3173,10 @@ void AUsdStageActor::LoadUsdStage(bool bOpenIfNeeded)
 	// the tracks
 	PrimsToAnimate.Reset();
 
-	TSharedRef< FUsdSchemaTranslationContext > TranslationContext = FUsdStageActorImpl::CreateUsdSchemaTranslationContext(this, RootTwin->PrimPath);
+	TSharedRef<FUsdSchemaTranslationContext> TranslationContext = FUsdStageActorImpl::CreateUsdSchemaTranslationContext(this, RootTwin->PrimPath);
 
 	SlowTask.EnterProgressFrame(0.1f);
-	InfoCache->RemoveAllAssetPrimLinks();  // These are usually not reset when rebuilding the cache so we must call manually
+	InfoCache->RemoveAllAssetPrimLinks();	 // These are usually not reset when rebuilding the cache so we must call manually
 	InfoCache->RebuildCacheForSubtree(StageToLoad.GetPseudoRoot(), TranslationContext.Get());
 
 	SlowTask.EnterProgressFrame(0.7f);
@@ -3211,7 +3203,8 @@ void AUsdStageActor::LoadUsdStage(bool bOpenIfNeeded)
 			UWorld* World = GetWorld();
 			if (World && World->WorldType == EWorldType::EditorPreview)
 			{
-				// Note: CDO is an instance of a BlueprintGeneratedClass here and this is just a base class pointer. We're not changing the actual AUsdStageActor's CDO
+				// Note: CDO is an instance of a BlueprintGeneratedClass here and this is just a base class pointer. We're not changing the actual
+				// AUsdStageActor's CDO
 				if (AUsdStageActor* CDO = Cast<AUsdStageActor>(GetClass()->GetDefaultObject()))
 				{
 					CDO->SetTime(GetTime());
@@ -3228,12 +3221,12 @@ void AUsdStageActor::LoadUsdStage(bool bOpenIfNeeded)
 	FUsdStageActorImpl::RepairExternalSequencerBindings();
 
 #if WITH_EDITOR
-	if (GIsEditor && GEditor && !IsGarbageCollecting()) // Make sure we're not in standalone either
+	if (GIsEditor && GEditor && !IsGarbageCollecting())	   // Make sure we're not in standalone either
 	{
 		GEditor->BroadcastLevelActorListChanged();
 		GEditor->RedrawLevelEditingViewports();
 	}
-#endif // WITH_EDITOR
+#endif	  // WITH_EDITOR
 
 	OnStageLoaded.Broadcast();
 
@@ -3243,16 +3236,19 @@ void AUsdStageActor::LoadUsdStage(bool bOpenIfNeeded)
 	int ElapsedMin = int(ElapsedSeconds / 60.0);
 	ElapsedSeconds -= 60.0 * (double)ElapsedMin;
 
-	UE_LOG(LogUsd, Log, TEXT("%s %s in [%d min %.3f s]"), TEXT("Stage loaded"), *FPaths::GetBaseFilename(RootLayer.FilePath), ElapsedMin, ElapsedSeconds);
+	UE_LOG(
+		LogUsd,
+		Log,
+		TEXT("%s %s in [%d min %.3f s]"),
+		TEXT("Stage loaded"),
+		*FPaths::GetBaseFilename(RootLayer.FilePath),
+		ElapsedMin,
+		ElapsedSeconds
+	);
 
 #if USE_USD_SDK
-	FUsdStageActorImpl::SendAnalytics(
-		this,
-		ElapsedSeconds,
-		UsdUtils::GetUsdStageNumFrames(StageToLoad),
-		FPaths::GetExtension(RootLayer.FilePath)
-	);
-#endif // #if USE_USD_SDK
+	FUsdStageActorImpl::SendAnalytics(this, ElapsedSeconds, UsdUtils::GetUsdStageNumFrames(StageToLoad), FPaths::GetExtension(RootLayer.FilePath));
+#endif	  // #if USE_USD_SDK
 }
 
 void AUsdStageActor::UnloadUsdStage()
@@ -3283,7 +3279,7 @@ void AUsdStageActor::UnloadUsdStage()
 		{
 			GEditor->GetEditorSubsystem<UAssetEditorSubsystem>()->CloseAllEditorsForAsset(LevelSequence);
 		}
-#endif // WITH_EDITOR
+#endif	  // WITH_EDITOR
 		LevelSequence = nullptr;
 	}
 	LevelSequenceHelper.Clear();
@@ -3306,7 +3302,7 @@ void AUsdStageActor::UnloadUsdStage()
 		GEditor->BroadcastLevelActorListChanged();
 		GEditor->RedrawLevelEditingViewports();
 	}
-#endif // WITH_EDITOR
+#endif	  // WITH_EDITOR
 
 	if (UsdAssetCache)
 	{
@@ -3335,7 +3331,11 @@ void AUsdStageActor::SetupAssetCacheIfNeeded()
 		{
 			if (UUsdAssetCache2* DefaultCache = Cast<UUsdAssetCache2>(ProjectSettings->DefaultAssetCache.TryLoad()))
 			{
-				UE_LOG(LogUsd, Log, TEXT("USD Stage Actor '%s' had no previous USD Asset Cache, so it will use the default cache at '%s'. This can be configured on the project settings."),
+				UE_LOG(
+					LogUsd,
+					Log,
+					TEXT("USD Stage Actor '%s' had no previous USD Asset Cache, so it will use the default cache at '%s'. This can be configured on "
+						 "the project settings."),
 					*GetPathName(),
 					*DefaultCache->GetPathName()
 				);
@@ -3371,7 +3371,11 @@ void AUsdStageActor::SetupAssetCacheIfNeeded()
 
 					UsdAssetCache = NewCache;
 
-					UE_LOG(LogUsd, Log, TEXT("USD Stage Actor '%s' will use newly created, default USD Asset Cache at '%s'. This can be configured on the project settings."),
+					UE_LOG(
+						LogUsd,
+						Log,
+						TEXT("USD Stage Actor '%s' will use newly created, default USD Asset Cache at '%s'. This can be configured on the project "
+							 "settings."),
 						*GetPathName(),
 						*UsdAssetCache->GetPathName()
 					);
@@ -3389,11 +3393,16 @@ void AUsdStageActor::SetupAssetCacheIfNeeded()
 
 	if (!UsdAssetCache)
 	{
-		UE_LOG(LogUsd, Warning, TEXT("USD Stage Actor '%s' had no previous USD Asset Cache and no default cache is specified on the project settings, so a temporary cache will be generated. For better performance, create a persistent USD Asset Cache asset and point to it with this actor's UsdAssetCache property."),
+		UE_LOG(
+			LogUsd,
+			Warning,
+			TEXT("USD Stage Actor '%s' had no previous USD Asset Cache and no default cache is specified on the project settings, so a temporary "
+				 "cache will be generated. For better performance, create a persistent USD Asset Cache asset and point to it with this actor's "
+				 "UsdAssetCache property."),
 			*GetPathName()
 		);
 
-		UsdAssetCache = NewObject< UUsdAssetCache2 >(GetTransientPackage(), NAME_None, GetMaskedFlags(RF_PropagateToSubObjects));
+		UsdAssetCache = NewObject<UUsdAssetCache2>(GetTransientPackage(), NAME_None, GetMaskedFlags(RF_PropagateToSubObjects));
 	}
 }
 
@@ -3473,7 +3482,7 @@ void AUsdStageActor::ReloadAnimations()
 				}
 			}
 		}
-#endif // WITH_EDITOR
+#endif	  // WITH_EDITOR
 
 		// We need to guarantee we'll record our change of LevelSequence into the transaction, as Init() will create a new one
 		const bool bMarkDirty = false;
@@ -3490,7 +3499,7 @@ void AUsdStageActor::ReloadAnimations()
 				Sequencer->ResetToNewRootSequence(*LevelSequence);
 			}
 		}
-#endif // WITH_EDITOR
+#endif	  // WITH_EDITOR
 	}
 }
 
@@ -3566,9 +3575,9 @@ void AUsdStageActor::PostTransacted(const FTransactionObjectEvent& TransactionEv
 	{
 		// If we're in a sublevel that is hidden, we'll respond to the generated PostUnregisterAllComponent call
 		// and unload our spawned actors/assets, so let's close/open the stage too
-		if (ChangedProperties.Contains(GET_MEMBER_NAME_CHECKED(AActor, bHiddenEdLevel)) ||
-			ChangedProperties.Contains(GET_MEMBER_NAME_CHECKED(AActor, bHiddenEdLayer)) ||
-			ChangedProperties.Contains(GET_MEMBER_NAME_CHECKED(AActor, bHiddenEd)))
+		if (ChangedProperties.Contains(GET_MEMBER_NAME_CHECKED(AActor, bHiddenEdLevel))
+			|| ChangedProperties.Contains(GET_MEMBER_NAME_CHECKED(AActor, bHiddenEdLayer))
+			|| ChangedProperties.Contains(GET_MEMBER_NAME_CHECKED(AActor, bHiddenEd)))
 		{
 			if (IsHiddenEd())
 			{
@@ -3627,13 +3636,16 @@ void AUsdStageActor::PostTransacted(const FTransactionObjectEvent& TransactionEv
 			// Sometimes when we undo/redo changes that modify SkinnedMeshComponents, their render state is not correctly updated which can show some
 			// very garbled meshes. Here we workaround that by recreating all those render states manually
 			const bool bRecurive = true;
-			GetRootPrimTwin()->Iterate([](UUsdPrimTwin& PrimTwin)
+			GetRootPrimTwin()->Iterate(
+				[](UUsdPrimTwin& PrimTwin)
 				{
 					if (USkinnedMeshComponent* Component = Cast<USkinnedMeshComponent>(PrimTwin.GetSceneComponent()))
 					{
 						FRenderStateRecreator RecreateRenderState{Component};
 					}
-				}, bRecurive);
+				},
+				bRecurive
+			);
 		}
 	}
 
@@ -3671,12 +3683,15 @@ void AUsdStageActor::PreEditUndo()
 	Super::PreEditUndo();
 }
 
-void AUsdStageActor::HandleTransactionStateChanged(const FTransactionContext& InTransactionContext, const ETransactionStateEventType InTransactionState)
+void AUsdStageActor::HandleTransactionStateChanged(
+	const FTransactionContext& InTransactionContext,
+	const ETransactionStateEventType InTransactionState
+)
 {
 	// Hack for solving UE-127253
-	// When we Reload (or open a new stage), we call ReloadAnimations which will close the Sequencer (if opened), recreate our LevelSequence, and get the Sequencer
-	// to show that one instead. If we undo the Reload, that new LevelSequence will be deleted and the Sequencer will be left open trying to display it,
-	// which leads to crashes. Here we try detecting for that case and close/reopen the sequencer to show the correct one.
+	// When we Reload (or open a new stage), we call ReloadAnimations which will close the Sequencer (if opened), recreate our LevelSequence, and get
+	// the Sequencer to show that one instead. If we undo the Reload, that new LevelSequence will be deleted and the Sequencer will be left open
+	// trying to display it, which leads to crashes. Here we try detecting for that case and close/reopen the sequencer to show the correct one.
 #if WITH_EDITOR
 	if (GIsEditor && LevelSequence
 		&& (InTransactionState == ETransactionStateEventType::UndoRedoStarted || InTransactionState == ETransactionStateEventType::UndoRedoFinalized))
@@ -3724,22 +3739,26 @@ void AUsdStageActor::HandleTransactionStateChanged(const FTransactionContext& In
 										SequencersToReset.Add(PinnedSequencer);
 
 										// Hack for solving UE-171596
-										// In this transaction we will switch LevelSequences, and have a Sequencer opened displaying our current Sequence.
-										// - We cannot leave this Sequencer displaying our old LevelSequence, because it will go PendingKill, and as the
-										//   Sequencer fetches it through WeakPtrs it will not find a valid LevelSequence and crash (this was the reason for
-										//   the original UE-127253 hack above). This means on UndoRedoStarted we *must* do something;
-										// - We cannot set our new LevelSequence into it yet of course, because it hasn't been created yet (it will be spawned
+										// In this transaction we will switch LevelSequences, and have a Sequencer opened displaying our current
+										// Sequence.
+										// - We cannot leave this Sequencer displaying our old LevelSequence, because it will go PendingKill, and as
+										// the
+										//   Sequencer fetches it through WeakPtrs it will not find a valid LevelSequence and crash (this was the
+										//   reason for the original UE-127253 hack above). This means on UndoRedoStarted we *must* do something;
+										// - We cannot set our new LevelSequence into it yet of course, because it hasn't been created yet (it will be
+										// spawned
 										//   by the undo system after UndoRedoStarted);
-										// - We cannot close this Sequencer, because of this "DeferredModify" mechanism that pushes some updates to the end
-										//   of the transaction (to UndoRedoFinalized). If one of those updates executes after we close the Sequencer and before
-										//   we fix things up (which it can always do as the order of execution of the delegates is not deterministic), it will
-										//   crash (this is the issue at UE-171596);
+										// - We cannot close this Sequencer, because of this "DeferredModify" mechanism that pushes some updates to
+										// the end
+										//   of the transaction (to UndoRedoFinalized). If one of those updates executes after we close the Sequencer
+										//   and before we fix things up (which it can always do as the order of execution of the delegates is not
+										//   deterministic), it will crash (this is the issue at UE-171596);
 										//
-										// This means we're forced to give *some valid LevelSequence* to the Sequencer for the split second while we switch
-										// our actual generated LevelSequence.
+										// This means we're forced to give *some valid LevelSequence* to the Sequencer for the split second while we
+										// switch our actual generated LevelSequence.
 										static TStrongObjectPtr<ULevelSequence> DummySequencePtr = nullptr;
 										ULevelSequence* DummySequence = DummySequencePtr.Get();
-										if(!DummySequence)
+										if (!DummySequence)
 										{
 											DummySequence = NewObject<ULevelSequence>();
 											DummySequence->Initialize();
@@ -3775,17 +3794,16 @@ void AUsdStageActor::HandleTransactionStateChanged(const FTransactionContext& In
 			}
 		}
 	}
-#endif // WITH_EDITOR
+#endif	  // WITH_EDITOR
 
-	if (InTransactionState == ETransactionStateEventType::TransactionFinalized ||
-		InTransactionState == ETransactionStateEventType::UndoRedoFinalized ||
-		InTransactionState == ETransactionStateEventType::TransactionCanceled)
+	if (InTransactionState == ETransactionStateEventType::TransactionFinalized || InTransactionState == ETransactionStateEventType::UndoRedoFinalized
+		|| InTransactionState == ETransactionStateEventType::TransactionCanceled)
 	{
 		OldRootLayer = RootLayer;
 	}
 }
 
-#endif // WITH_EDITOR
+#endif	  // WITH_EDITOR
 
 void AUsdStageActor::PostDuplicate(bool bDuplicateForPIE)
 {
@@ -3800,7 +3818,7 @@ void AUsdStageActor::PostDuplicate(bool bDuplicateForPIE)
 		// Temporary asset caches aren't meant to be shared
 		if (UsdAssetCache && UsdAssetCache->GetOutermost() == GetTransientPackage())
 		{
-			UsdAssetCache = NewObject< UUsdAssetCache2 >(GetTransientPackage(), NAME_None, GetMaskedFlags(RF_PropagateToSubObjects));
+			UsdAssetCache = NewObject<UUsdAssetCache2>(GetTransientPackage(), NAME_None, GetMaskedFlags(RF_PropagateToSubObjects));
 		}
 
 		LoadUsdStage();
@@ -3923,7 +3941,7 @@ void AUsdStageActor::BeginDestroy()
 	{
 		RootUsdTwin->Clear();
 	}
-#endif // WITH_EDITOR
+#endif	  // WITH_EDITOR
 
 	Super::BeginDestroy();
 }
@@ -3935,7 +3953,8 @@ void AUsdStageActor::PostRegisterAllComponents()
 	// This may fail if our stage happened to not spawn any components, actors or assets, but by that
 	// point "being loaded" doesn't really mean anything anyway
 	const bool bStageIsLoaded = GetBaseUsdStage()
-		&& ((RootUsdTwin && RootUsdTwin->GetSceneComponent() != nullptr) || (UsdAssetCache && UsdAssetCache->GetNumAssets() > 0));
+								&& ((RootUsdTwin && RootUsdTwin->GetSceneComponent() != nullptr)
+									|| (UsdAssetCache && UsdAssetCache->GetNumAssets() > 0));
 
 	// Blocks loading stage when going into PIE, if we already have something loaded (we'll want to duplicate stuff instead).
 	// We need to allow loading when going into PIE when we have nothing loaded yet because the MovieRenderQueue (or other callers)
@@ -3955,7 +3974,8 @@ void AUsdStageActor::PostRegisterAllComponents()
 	}
 
 #if WITH_EDITOR
-	// Prevent loading on bHiddenEdLevel because PostRegisterAllComponents gets called in the process of hiding our level, if we're in the persistent level.
+	// Prevent loading on bHiddenEdLevel because PostRegisterAllComponents gets called in the process of hiding our level, if we're in the persistent
+	// level.
 	if (bIsEditorPreviewActor || bHiddenEdLevel)
 	{
 		return;
@@ -3978,11 +3998,11 @@ void AUsdStageActor::PostRegisterAllComponents()
 			return;
 		}
 	}
-#endif // WITH_EDITOR
+#endif	  // WITH_EDITOR
 
-	// When we add a sublevel the very first time (i.e. when it is associating) it may still be invisible, but we should load the stage anyway because by
-	// default it will become visible shortly after this call. On subsequent postregisters, if our level is invisible there is no point to loading our stage,
-	// as our spawned actors/components should be invisible too
+	// When we add a sublevel the very first time (i.e. when it is associating) it may still be invisible, but we should load the stage anyway because
+	// by default it will become visible shortly after this call. On subsequent postregisters, if our level is invisible there is no point to loading
+	// our stage, as our spawned actors/components should be invisible too
 	ULevel* Level = GetLevel();
 	const bool bIsLevelHidden = !Level || (!Level->bIsVisible && !Level->bIsAssociatingLevel);
 	if (bIsLevelHidden)
@@ -4025,10 +4045,11 @@ void AUsdStageActor::UnregisterAllComponents(bool bForReregister)
 			return;
 		}
 	}
-#endif // WITH_EDITOR
+#endif	  // WITH_EDITOR
 
 	const bool bStageIsLoaded = GetBaseUsdStage()
-		&& ((RootUsdTwin && RootUsdTwin->GetSceneComponent() != nullptr) || (UsdAssetCache && UsdAssetCache->GetNumAssets() > 0));
+								&& ((RootUsdTwin && RootUsdTwin->GetSceneComponent() != nullptr)
+									|| (UsdAssetCache && UsdAssetCache->GetNumAssets() > 0));
 
 	UWorld* World = GetWorld();
 	if (bIsTransitioningIntoPIE && bStageIsLoaded && (!World || World->WorldType == EWorldType::PIE))
@@ -4205,7 +4226,7 @@ void AUsdStageActor::OnMovieSceneDataChanged(EMovieSceneDataChangeType ChangeTyp
 	}
 
 	UUsdLevelSequenceAssetUserData* AssetUserData = Sequence->GetAssetUserData<UUsdLevelSequenceAssetUserData>();
-	if(!AssetUserData)
+	if (!AssetUserData)
 	{
 		AssetUserData = NewObject<UUsdLevelSequenceAssetUserData>(Sequence);
 		Sequence->AddAssetUserData(AssetUserData);
@@ -4226,9 +4247,7 @@ void AUsdStageActor::OnMovieSceneDataChanged(EMovieSceneDataChangeType ChangeTyp
 
 	// The UI action to add the actor to the sequencer should have spawned it's transaction already, but just in case it
 	// hasn't let's do that, as we'll end up modifying UObjects in here
-	FScopedTransaction Transaction(
-		LOCTEXT("ComponentAddedToSequencer", "Added USD component bindings to the Sequencer")
-	);
+	FScopedTransaction Transaction(LOCTEXT("ComponentAddedToSequencer", "Added USD component bindings to the Sequencer"));
 	Sequence->Modify();
 	MovieScene->Modify();
 
@@ -4257,7 +4276,7 @@ void AUsdStageActor::OnMovieSceneDataChanged(EMovieSceneDataChangeType ChangeTyp
 	};
 
 	int32 PossessableCount = MovieScene->GetPossessableCount();
-	for(int32 Index = 0; Index < PossessableCount; ++Index)
+	for (int32 Index = 0; Index < PossessableCount; ++Index)
 	{
 		FMovieScenePossessable& Possessable = MovieScene->GetPossessable(Index);
 
@@ -4304,7 +4323,7 @@ void AUsdStageActor::OnMovieSceneDataChanged(EMovieSceneDataChangeType ChangeTyp
 void AUsdStageActor::OnUsdPrimTwinDestroyed(const UUsdPrimTwin& UsdPrimTwin)
 {
 	PrimsToAnimate.Remove(UsdPrimTwin.PrimPath);
-	ObjectsToWatch.Remove( UsdPrimTwin.SceneComponent.Get() );
+	ObjectsToWatch.Remove(UsdPrimTwin.SceneComponent.Get());
 	LevelSequenceHelper.RemovePrim(UsdPrimTwin);
 }
 
@@ -4383,8 +4402,7 @@ void AUsdStageActor::OnObjectPropertyChanged(UObject* ObjectBeingModified, FProp
 	const static TSet<FName> TransformProperties = {
 		USceneComponent::GetRelativeLocationPropertyName(),
 		USceneComponent::GetRelativeRotationPropertyName(),
-		USceneComponent::GetRelativeScale3DPropertyName()
-	};
+		USceneComponent::GetRelativeScale3DPropertyName()};
 	const bool bIsTransformChange = TransformProperties.Contains(PropertyChangedEvent.GetPropertyName());
 
 	// When we change an actor property that is just a mirror of a component property (e.g. light intensity, or camera aperture)
@@ -4474,7 +4492,7 @@ void AUsdStageActor::OnObjectPropertyChanged(UObject* ObjectBeingModified, FProp
 			}
 		}
 	}
-#endif // WITH_EDITOR
+#endif	  // WITH_EDITOR
 
 	// We spawn Cine Camera Actors for Camera prims, but those have two components by default. Our convention is to place
 	// camera stuff on the camera component (not much choice there), but use the transform of the scene (root) component.
@@ -4537,7 +4555,7 @@ void AUsdStageActor::OnObjectPropertyChanged(UObject* ObjectBeingModified, FProp
 				const bool bTimeSampleValues = false;
 				UsdUtils::AuthorIdentityTransformGprimAttributes(UsdPrim, bDefaultValues, bTimeSampleValues);
 
-				if (UMeshComponent* MeshComponent = Cast< UMeshComponent >(PrimSceneComponent))
+				if (UMeshComponent* MeshComponent = Cast<UMeshComponent>(PrimSceneComponent))
 				{
 					UnrealToUsd::ConvertMeshComponent(CurrentStage, MeshComponent, UsdPrim);
 				}
@@ -4606,7 +4624,7 @@ void AUsdStageActor::OnObjectPropertyChanged(UObject* ObjectBeingModified, FProp
 					UnrealToUsd::ConvertLightComponent(*SkyLightComponent, UsdPrim, UsdUtils::GetDefaultTimeCode());
 					UnrealToUsd::ConvertSkyLightComponent(*SkyLightComponent, UsdPrim, UsdUtils::GetDefaultTimeCode());
 				}
-#endif // #if USE_USD_SDK
+#endif	  // #if USE_USD_SDK
 
 				// Update stage window in case any of our component changes trigger USD stage changes
 				this->OnPrimChanged.Broadcast(PrimPath, false);
@@ -4617,10 +4635,11 @@ void AUsdStageActor::OnObjectPropertyChanged(UObject* ObjectBeingModified, FProp
 
 void AUsdStageActor::HandlePropertyChangedEvent(FPropertyChangedEvent& PropertyChangedEvent)
 {
-	// Handle property changed events with this function (called from our OnObjectPropertyChanged delegate) instead of overriding PostEditChangeProperty because replicated
-	// multi-user transactions directly broadcast OnObjectPropertyChanged on the properties that were changed, instead of making PostEditChangeProperty events.
-	// Note that UObject::PostEditChangeProperty ends up broadcasting OnObjectPropertyChanged anyway, so this works just the same as before.
-	// see ConcertClientTransactionBridge.cpp, function ConcertClientTransactionBridgeUtil::ProcessTransactionEvent.
+	// Handle property changed events with this function (called from our OnObjectPropertyChanged delegate) instead of overriding
+	// PostEditChangeProperty because replicated multi-user transactions directly broadcast OnObjectPropertyChanged on the properties that were
+	// changed, instead of making PostEditChangeProperty events. Note that UObject::PostEditChangeProperty ends up broadcasting
+	// OnObjectPropertyChanged anyway, so this works just the same as before. see ConcertClientTransactionBridge.cpp, function
+	// ConcertClientTransactionBridgeUtil::ProcessTransactionEvent.
 
 	// Note that in here we'll delegate to these setter functions (like SetRootLayer) to actually set the new property values.
 	// We want our setter functions to be able to automatically refresh the stage (both for simplicity, since we have a single
@@ -4662,7 +4681,7 @@ void AUsdStageActor::HandlePropertyChangedEvent(FPropertyChangedEvent& PropertyC
 #if WITH_EDITOR
 		static IConsoleVariable* CVar = IConsoleManager::Get().FindConsoleVariable(TEXT("USD.DiscardUndoBufferOnStageOpenClose"));
 		const bool bDiscardUndo = CVar && CVar->GetBool();
-#endif // WITH_EDITOR
+#endif	  // WITH_EDITOR
 
 		// Technically we don't need this guard value for the root layer itself, since SetRootLayer can compare
 		// RootLayer with the path of the current stage's root layer, but let's just do this for consistency.
@@ -4679,12 +4698,12 @@ void AUsdStageActor::HandlePropertyChangedEvent(FPropertyChangedEvent& PropertyC
 				EditorTransactor->Reset(Reason);
 			}
 		}
-#endif // WITH_EDITOR
+#endif	  // WITH_EDITOR
 	}
 	else if (PropertyName == GET_MEMBER_NAME_CHECKED(AUsdStageActor, StageState))
 	{
 		const EUsdStageState CorrectState = StageState;
-		StageState = (EUsdStageState)!((uint8)StageState);
+		StageState = (EUsdStageState) !((uint8)StageState);
 		SetStageState(CorrectState);
 	}
 	else if (PropertyName == GET_MEMBER_NAME_CHECKED(AUsdStageActor, Time))
@@ -4802,12 +4821,12 @@ void AUsdStageActor::HandlePropertyChangedEvent(FPropertyChangedEvent& PropertyC
 bool AUsdStageActor::HasAuthorityOverStage() const
 {
 #if WITH_EDITOR
-	if (GIsEditor) // Don't check for world in Standalone: The game world is the only one there, so it's OK if we have authority while in it
+	if (GIsEditor)	  // Don't check for world in Standalone: The game world is the only one there, so it's OK if we have authority while in it
 	{
 		// In the editor we have to prevent actors in PIE worlds from having authority
 		return !IsTemplate() && (!GetWorld() || !GetWorld()->IsGameWorld());
 	}
-#endif // WITH_EDITOR
+#endif	  // WITH_EDITOR
 
 	return !IsTemplate();
 }
@@ -4845,12 +4864,13 @@ void AUsdStageActor::OnSkelAnimationBaked(const FString& SkeletonPrimPath)
 		return;
 	}
 
-	TSharedRef< FUsdSchemaTranslationContext > TranslationContext = FUsdStageActorImpl::CreateUsdSchemaTranslationContext(this, SkeletonPrimPath);
+	TSharedRef<FUsdSchemaTranslationContext> TranslationContext = FUsdStageActorImpl::CreateUsdSchemaTranslationContext(this, SkeletonPrimPath);
 	// The only way we could have baked a skel animation is via the sequencer, so we know its playing
 	TranslationContext->bSequencerIsAnimating = true;
 
-	IUsdSchemasModule& UsdSchemasModule = FModuleManager::Get().LoadModuleChecked< IUsdSchemasModule >(TEXT("USDSchemas"));
-	if (TSharedPtr< FUsdSchemaTranslator > SchemaTranslator = UsdSchemasModule.GetTranslatorRegistry().CreateTranslatorForSchema(TranslationContext, UE::FUsdTyped(SkeletonPrim)))
+	IUsdSchemasModule& UsdSchemasModule = FModuleManager::Get().LoadModuleChecked<IUsdSchemasModule>(TEXT("USDSchemas"));
+	if (TSharedPtr<FUsdSchemaTranslator> SchemaTranslator = UsdSchemasModule.GetTranslatorRegistry()
+																.CreateTranslatorForSchema(TranslationContext, UE::FUsdTyped(SkeletonPrim)))
 	{
 		if (TSharedPtr<FUsdSkelSkeletonTranslator> SkelRootTranslator = StaticCastSharedPtr<FUsdSkelSkeletonTranslator>(SchemaTranslator))
 		{
@@ -4864,7 +4884,7 @@ void AUsdStageActor::OnSkelAnimationBaked(const FString& SkeletonPrimPath)
 			SkelRootTranslator->UpdateComponents(SkeletalMeshComponent);
 		}
 	}
-#endif // #if USE_USD_SDK
+#endif	  // #if USE_USD_SDK
 }
 
 bool AUsdStageActor::UnloadAssets(const UE::FSdfPath& StartPrimPath, bool bForEntireSubtree)
@@ -4937,21 +4957,19 @@ bool AUsdStageActor::LoadAsset(FUsdSchemaTranslationContext& TranslationContext,
 	}
 
 	// Mark the assets as non transactional so that they don't get serialized in the transaction buffer
-	TGuardValue< EObjectFlags > ContextFlagsGuard(TranslationContext.ObjectFlags, TranslationContext.ObjectFlags & ~RF_Transactional);
+	TGuardValue<EObjectFlags> ContextFlagsGuard(TranslationContext.ObjectFlags, TranslationContext.ObjectFlags & ~RF_Transactional);
 
 	int32 StartNumAssets = UsdAssetCache->GetNumAssets();
 
-	IUsdSchemasModule& UsdSchemasModule = FModuleManager::Get().LoadModuleChecked< IUsdSchemasModule >(TEXT("USDSchemas"));
-	if (TSharedPtr< FUsdSchemaTranslator > SchemaTranslator = UsdSchemasModule.GetTranslatorRegistry().CreateTranslatorForSchema(
-		TranslationContext.AsShared(),
-		UE::FUsdTyped(Prim)
-	))
+	IUsdSchemasModule& UsdSchemasModule = FModuleManager::Get().LoadModuleChecked<IUsdSchemasModule>(TEXT("USDSchemas"));
+	if (TSharedPtr<FUsdSchemaTranslator> SchemaTranslator = UsdSchemasModule.GetTranslatorRegistry()
+																.CreateTranslatorForSchema(TranslationContext.AsShared(), UE::FUsdTyped(Prim)))
 	{
 		TRACE_CPUPROFILER_EVENT_SCOPE(AUsdStageActor::CreateAssetsForPrim);
 		SchemaTranslator->CreateAssets();
 	}
 
-	TranslationContext.CompleteTasks(); // Finish the asset tasks before moving on
+	TranslationContext.CompleteTasks();	   // Finish the asset tasks before moving on
 
 	return UsdAssetCache->GetNumAssets() != StartNumAssets;
 }
@@ -4966,13 +4984,13 @@ bool AUsdStageActor::LoadAssets(FUsdSchemaTranslationContext& TranslationContext
 	}
 
 	// Mark the assets as non transactional so that they don't get serialized in the transaction buffer
-	TGuardValue< EObjectFlags > ContextFlagsGuard(TranslationContext.ObjectFlags, TranslationContext.ObjectFlags & ~RF_Transactional);
+	TGuardValue<EObjectFlags> ContextFlagsGuard(TranslationContext.ObjectFlags, TranslationContext.ObjectFlags & ~RF_Transactional);
 
 	int32 StartNumAssets = UsdAssetCache->GetNumAssets();
 
-	IUsdSchemasModule& UsdSchemasModule = FModuleManager::Get().LoadModuleChecked< IUsdSchemasModule >(TEXT("USDSchemas"));
+	IUsdSchemasModule& UsdSchemasModule = FModuleManager::Get().LoadModuleChecked<IUsdSchemasModule>(TEXT("USDSchemas"));
 
-	auto CreateAssetsForPrims = [&UsdSchemasModule, &TranslationContext](const TArray< UE::FUsdPrim >& AllPrimAssets, FSlowTask& Progress)
+	auto CreateAssetsForPrims = [&UsdSchemasModule, &TranslationContext](const TArray<UE::FUsdPrim>& AllPrimAssets, FSlowTask& Progress)
 	{
 		TRACE_CPUPROFILER_EVENT_SCOPE(AUsdStageActor::CreateAssetsForPrims);
 
@@ -4980,25 +4998,23 @@ bool AUsdStageActor::LoadAssets(FUsdSchemaTranslationContext& TranslationContext
 		{
 			Progress.EnterProgressFrame(1.f);
 
-			if (TSharedPtr< FUsdSchemaTranslator > SchemaTranslator = UsdSchemasModule.GetTranslatorRegistry().CreateTranslatorForSchema(
-				TranslationContext.AsShared(),
-				UE::FUsdTyped(UsdPrim)
-			))
+			if (TSharedPtr<FUsdSchemaTranslator> SchemaTranslator = UsdSchemasModule.GetTranslatorRegistry().CreateTranslatorForSchema(
+					TranslationContext.AsShared(),
+					UE::FUsdTyped(UsdPrim)
+				))
 			{
 				TRACE_CPUPROFILER_EVENT_SCOPE(AUsdStageActor::CreateAssetsForPrim);
 				SchemaTranslator->CreateAssets();
 			}
 		}
 
-		TranslationContext.CompleteTasks(); // Finish the assets tasks before moving on
+		TranslationContext.CompleteTasks();	   // Finish the assets tasks before moving on
 	};
 
 	auto PruneChildren = [&UsdSchemasModule, &TranslationContext](const UE::FUsdPrim& UsdPrim) -> bool
 	{
-		if (TSharedPtr< FUsdSchemaTranslator > SchemaTranslator = UsdSchemasModule.GetTranslatorRegistry().CreateTranslatorForSchema(
-			TranslationContext.AsShared(),
-			UE::FUsdTyped(UsdPrim)
-		))
+		if (TSharedPtr<FUsdSchemaTranslator> SchemaTranslator = UsdSchemasModule.GetTranslatorRegistry()
+																	.CreateTranslatorForSchema(TranslationContext.AsShared(), UE::FUsdTyped(UsdPrim)))
 		{
 			return SchemaTranslator->CollapsesChildren(ECollapsingType::Assets);
 		}
@@ -5007,7 +5023,7 @@ bool AUsdStageActor::LoadAssets(FUsdSchemaTranslationContext& TranslationContext
 	};
 
 	// Load materials first since meshes are referencing them
-	TArray< UE::FUsdPrim > AllPrimAssets = UsdUtils::GetAllPrimsOfType(StartPrim, TEXT("UsdShadeMaterial"));
+	TArray<UE::FUsdPrim> AllPrimAssets = UsdUtils::GetAllPrimsOfType(StartPrim, TEXT("UsdShadeMaterial"));
 	{
 		FScopedSlowTask MaterialsProgress(AllPrimAssets.Num(), LOCTEXT("CreateMaterials", "Creating materials"));
 		CreateAssetsForPrims(AllPrimAssets, MaterialsProgress);
@@ -5034,7 +5050,10 @@ void AUsdStageActor::AnimatePrims()
 		return;
 	}
 
-	TSharedRef< FUsdSchemaTranslationContext > TranslationContext = FUsdStageActorImpl::CreateUsdSchemaTranslationContext(this, GetRootPrimTwin()->PrimPath);
+	TSharedRef<FUsdSchemaTranslationContext> TranslationContext = FUsdStageActorImpl::CreateUsdSchemaTranslationContext(
+		this,
+		GetRootPrimTwin()->PrimPath
+	);
 
 	// c.f. comment on bSequencerIsAnimating's declaration
 #if WITH_EDITOR
@@ -5047,14 +5066,17 @@ void AUsdStageActor::AnimatePrims()
 			TranslationContext->bSequencerIsAnimating = true;
 		}
 	}
-#endif // WITH_EDITOR
+#endif	  // WITH_EDITOR
 
 	for (const FString& PrimToAnimate : PrimsToAnimate)
 	{
 		UE::FSdfPath PrimPath(*PrimToAnimate);
 
-		IUsdSchemasModule& SchemasModule = FModuleManager::Get().LoadModuleChecked< IUsdSchemasModule >("USDSchemas");
-		if (TSharedPtr< FUsdSchemaTranslator > SchemaTranslator = SchemasModule.GetTranslatorRegistry().CreateTranslatorForSchema(TranslationContext, UE::FUsdTyped(CurrentStage.GetPrimAtPath(PrimPath))))
+		IUsdSchemasModule& SchemasModule = FModuleManager::Get().LoadModuleChecked<IUsdSchemasModule>("USDSchemas");
+		if (TSharedPtr<FUsdSchemaTranslator> SchemaTranslator = SchemasModule.GetTranslatorRegistry().CreateTranslatorForSchema(
+				TranslationContext,
+				UE::FUsdTyped(CurrentStage.GetPrimAtPath(PrimPath))
+			))
 		{
 			if (UUsdPrimTwin* UsdPrimTwin = GetRootPrimTwin()->Find(PrimToAnimate))
 			{
diff --git a/Engine/Plugins/Importers/USDImporter/Source/USDStage/Private/USDStageActorCustomization.cpp b/Engine/Plugins/Importers/USDImporter/Source/USDStage/Private/USDStageActorCustomization.cpp
--- a/Engine/Plugins/Importers/USDImporter/Source/USDStage/Private/USDStageActorCustomization.cpp
+++ b/Engine/Plugins/Importers/USDImporter/Source/USDStage/Private/USDStageActorCustomization.cpp
@@ -31,68 +31,69 @@ TSharedRef<IDetailCustomization> FUsdStageActorCustomization::MakeInstance()
 	return MakeShared<FUsdStageActorCustomization>();
 }
 
-void FUsdStageActorCustomization::CustomizeDetails( IDetailLayoutBuilder& DetailLayoutBuilder )
+void FUsdStageActorCustomization::CustomizeDetails(IDetailLayoutBuilder& DetailLayoutBuilder)
 {
 	TArray<TWeakObjectPtr<UObject>> SelectedObjects = DetailLayoutBuilder.GetSelectedObjects();
-	if ( SelectedObjects.Num() != 1 )
+	if (SelectedObjects.Num() != 1)
 	{
 		return;
 	}
 
-	TWeakObjectPtr<UObject> SelectedObject = SelectedObjects[ 0 ];
-	if ( !SelectedObject.IsValid() )
+	TWeakObjectPtr<UObject> SelectedObject = SelectedObjects[0];
+	if (!SelectedObject.IsValid())
 	{
 		return;
 	}
 
-	CurrentActor = Cast<AUsdStageActor>( SelectedObject.Get() );
-	if ( !CurrentActor )
+	CurrentActor = Cast<AUsdStageActor>(SelectedObject.Get());
+	if (!CurrentActor)
 	{
 		return;
 	}
 
-	CurrentActor->OnStageChanged.AddSP( this, &FUsdStageActorCustomization::ForceRefreshDetails );
+	CurrentActor->OnStageChanged.AddSP(this, &FUsdStageActorCustomization::ForceRefreshDetails);
 
-	IUsdSchemasModule& UsdSchemasModule = FModuleManager::Get().LoadModuleChecked< IUsdSchemasModule >( TEXT( "USDSchemas" ) );
+	IUsdSchemasModule& UsdSchemasModule = FModuleManager::Get().LoadModuleChecked<IUsdSchemasModule>(TEXT("USDSchemas"));
 
 	RenderContextComboBoxItems.Reset();
 	TSharedPtr<FString> InitiallySelectedContext;
-	for ( const FName& Context : UsdSchemasModule.GetRenderContextRegistry().GetRenderContexts() )
+	for (const FName& Context : UsdSchemasModule.GetRenderContextRegistry().GetRenderContexts())
 	{
 		TSharedPtr<FString> ContextStr;
-		if ( Context == NAME_None )
+		if (Context == NAME_None)
 		{
-			ContextStr = MakeShared<FString>( TEXT( "universal" ) );
+			ContextStr = MakeShared<FString>(TEXT("universal"));
 		}
 		else
 		{
-			ContextStr = MakeShared<FString>( Context.ToString() );
+			ContextStr = MakeShared<FString>(Context.ToString());
 		}
 
-		if ( Context == CurrentActor->RenderContext )
+		if (Context == CurrentActor->RenderContext)
 		{
 			InitiallySelectedContext = ContextStr;
 		}
 
-		RenderContextComboBoxItems.Add( ContextStr );
+		RenderContextComboBoxItems.Add(ContextStr);
 	}
 
-	IDetailCategoryBuilder& CatBuilder = DetailLayoutBuilder.EditCategory( TEXT( "USD" ) );
+	IDetailCategoryBuilder& CatBuilder = DetailLayoutBuilder.EditCategory(TEXT("USD"));
 
+	// clang-format off
 	TWeakObjectPtr<AUsdStageActor> WeakActor = CurrentActor;
 	CatBuilder.AddCustomRow(FText())
 	.WholeRowContent()
 	[
 		SNew(SHorizontalBox)
-		+ SHorizontalBox::Slot()
+		+SHorizontalBox::Slot()
 		[
 			SNew(SButton)
 			.Text(LOCTEXT("OpenInStageEditorText", "Open in USD Stage Editor"))
 			.OnClicked_Lambda([WeakActor]()
 			{
-				if( AUsdStageActor* StageActor = WeakActor.Get() )
+				if(AUsdStageActor* StageActor = WeakActor.Get())
 				{
-					AUsdStageActor::OnOpenStageEditorClicked.Broadcast( StageActor );
+					AUsdStageActor::OnOpenStageEditorClicked.Broadcast(StageActor);
 				}
 
 				return FReply::Handled();
@@ -100,70 +101,70 @@ void FUsdStageActorCustomization::CustomizeDetails( IDetailLayoutBuilder& Detail
 		]
 	];
 
-	if ( TSharedPtr<IPropertyHandle> RenderContextProperty = DetailLayoutBuilder.GetProperty( GET_MEMBER_NAME_CHECKED( AUsdStageActor, RenderContext ) ) )
+	if (TSharedPtr<IPropertyHandle> RenderContextProperty = DetailLayoutBuilder.GetProperty(GET_MEMBER_NAME_CHECKED(AUsdStageActor, RenderContext)))
 	{
-		DetailLayoutBuilder.HideProperty( RenderContextProperty );
+		DetailLayoutBuilder.HideProperty(RenderContextProperty);
 
-		CatBuilder.AddCustomRow( FText::FromString( TEXT( "RenderContextCustomization" ) ) )
+		CatBuilder.AddCustomRow(FText::FromString(TEXT("RenderContextCustomization")))
 		.NameContent()
 		[
-			SNew( STextBlock )
-			.Text( FText::FromString( TEXT( "Render context" ) ) )
-			.Font( FAppStyle::GetFontStyle( TEXT( "PropertyWindow.NormalFont" ) ) )
-			.ToolTipText( RenderContextProperty->GetToolTipText() )
+			SNew(STextBlock)
+			.Text(FText::FromString(TEXT("Render context")))
+			.Font(FAppStyle::GetFontStyle(TEXT("PropertyWindow.NormalFont")))
+			.ToolTipText(RenderContextProperty->GetToolTipText())
 		]
 		.ValueContent()
 		[
-			SAssignNew( RenderContextComboBox, SComboBox<TSharedPtr<FString>> )
-			.OptionsSource( &RenderContextComboBoxItems )
-			.InitiallySelectedItem( InitiallySelectedContext )
-			.OnSelectionChanged( this, &FUsdStageActorCustomization::OnComboBoxSelectionChanged )
-			.OnGenerateWidget_Lambda( []( TSharedPtr<FString> Item )
+			SAssignNew(RenderContextComboBox, SComboBox<TSharedPtr<FString>>)
+			.OptionsSource(&RenderContextComboBoxItems)
+			.InitiallySelectedItem(InitiallySelectedContext)
+			.OnSelectionChanged(this, &FUsdStageActorCustomization::OnComboBoxSelectionChanged)
+			.OnGenerateWidget_Lambda([](TSharedPtr<FString> Item)
 			{
-				return SNew( STextBlock )
-					.Text( Item.IsValid() ? FText::FromString( *Item ) : FText::GetEmpty() )
-					.Font(FAppStyle::GetFontStyle( TEXT( "PropertyWindow.NormalFont" ) ) );
-			} )
+				return SNew(STextBlock)
+					.Text(Item.IsValid() ? FText::FromString(*Item) : FText::GetEmpty())
+					.Font(FAppStyle::GetFontStyle(TEXT("PropertyWindow.NormalFont")));
+			})
 			.Content()
 			[
-				SNew( STextBlock )
-				.Text( this, &FUsdStageActorCustomization::GetComboBoxSelectedOptionText )
-				.Font(FAppStyle::GetFontStyle( TEXT( "PropertyWindow.NormalFont" ) ) )
+				SNew(STextBlock)
+				.Text(this, &FUsdStageActorCustomization::GetComboBoxSelectedOptionText)
+				.Font(FAppStyle::GetFontStyle(TEXT("PropertyWindow.NormalFont")))
 			]
 		];
 	}
 
-	if ( TSharedPtr<IPropertyHandle> MaterialPurposeProperty = DetailLayoutBuilder.GetProperty( GET_MEMBER_NAME_CHECKED( AUsdStageActor, MaterialPurpose ) ) )
+	if (TSharedPtr<IPropertyHandle> MaterialPurposeProperty = DetailLayoutBuilder.GetProperty(GET_MEMBER_NAME_CHECKED(AUsdStageActor, MaterialPurpose)))
 	{
-		DetailLayoutBuilder.HideProperty( MaterialPurposeProperty );
+		DetailLayoutBuilder.HideProperty(MaterialPurposeProperty);
 
-		CatBuilder.AddCustomRow( FText::FromString( TEXT( "MaterialPurposeCustomization" ) ) )
+		CatBuilder.AddCustomRow(FText::FromString(TEXT("MaterialPurposeCustomization")))
 		.NameContent()
 		[
-			SNew( STextBlock )
-			.Text( FText::FromString( TEXT( "Material purpose" ) ) )
-			.Font( FAppStyle::GetFontStyle( TEXT( "PropertyWindow.NormalFont" ) ) )
-			.ToolTipText( MaterialPurposeProperty->GetToolTipText() )
+			SNew(STextBlock)
+			.Text(FText::FromString(TEXT("Material purpose")))
+			.Font(FAppStyle::GetFontStyle(TEXT("PropertyWindow.NormalFont")))
+			.ToolTipText(MaterialPurposeProperty->GetToolTipText())
 		]
 		.ValueContent()
 		[
-			SNew( SBox )
-			.VAlign( VAlign_Center )
+			SNew(SBox)
+			.VAlign(VAlign_Center)
 			[
-				SNew( SComboBox< TSharedPtr<FString> > )
-				.OptionsSource( &MaterialPurposeComboBoxItems )
+				SNew(SComboBox< TSharedPtr<FString> >)
+				.OptionsSource(&MaterialPurposeComboBoxItems)
 				.OnComboBoxOpening_Lambda([this]()
 				{
 					MaterialPurposeComboBoxItems = {
-						MakeShared<FString>( UnrealIdentifiers::MaterialAllPurpose ),
-						MakeShared<FString>( UnrealIdentifiers::MaterialPreviewPurpose ),
-						MakeShared<FString>( UnrealIdentifiers::MaterialFullPurpose )
+						MakeShared<FString>(UnrealIdentifiers::MaterialAllPurpose),
+						MakeShared<FString>(UnrealIdentifiers::MaterialPreviewPurpose),
+						MakeShared<FString>(UnrealIdentifiers::MaterialFullPurpose)
 					};
 
 					// Add additional purposes from project settings
-					if ( const UUsdProjectSettings* ProjectSettings = GetDefault<UUsdProjectSettings>() )
+					if (const UUsdProjectSettings* ProjectSettings = GetDefault<UUsdProjectSettings>())
 					{
-						MaterialPurposeComboBoxItems.Reserve( MaterialPurposeComboBoxItems.Num() + ProjectSettings->AdditionalMaterialPurposes.Num() );
+						MaterialPurposeComboBoxItems.Reserve(MaterialPurposeComboBoxItems.Num() + ProjectSettings->AdditionalMaterialPurposes.Num());
 
 						TSet<FString> ExistingEntries = {
 							UnrealIdentifiers::MaterialAllPurpose,
@@ -171,47 +172,47 @@ void FUsdStageActorCustomization::CustomizeDetails( IDetailLayoutBuilder& Detail
 							UnrealIdentifiers::MaterialFullPurpose
 						};
 
-						for ( const FName& AdditionalPurpose : ProjectSettings->AdditionalMaterialPurposes )
+						for (const FName& AdditionalPurpose : ProjectSettings->AdditionalMaterialPurposes)
 						{
 							FString AdditionalPurposeStr = AdditionalPurpose.ToString();
 
-							if ( !ExistingEntries.Contains( AdditionalPurposeStr ) )
+							if (!ExistingEntries.Contains(AdditionalPurposeStr))
 							{
-								ExistingEntries.Add( AdditionalPurposeStr );
-								MaterialPurposeComboBoxItems.AddUnique( MakeShared<FString>( AdditionalPurposeStr ) );
+								ExistingEntries.Add(AdditionalPurposeStr);
+								MaterialPurposeComboBoxItems.AddUnique(MakeShared<FString>(AdditionalPurposeStr));
 							}
 						}
 					}
 				})
-				.OnGenerateWidget_Lambda( [ & ]( TSharedPtr<FString> Option )
+				.OnGenerateWidget_Lambda([&](TSharedPtr<FString> Option)
 				{
 					TSharedPtr<SWidget> Widget = SNullWidget::NullWidget;
-					if ( Option )
+					if (Option)
 					{
-						Widget = SNew( STextBlock )
-							.Text( FText::FromString( ( *Option ) == UnrealIdentifiers::MaterialAllPurpose
+						Widget = SNew(STextBlock)
+							.Text(FText::FromString((*Option) == UnrealIdentifiers::MaterialAllPurpose
 								? UnrealIdentifiers::MaterialAllPurposeText
 								: *Option
-							) )
-							.Font( FAppStyle::GetFontStyle( "PropertyWindow.NormalFont" ) );
+							))
+							.Font(FAppStyle::GetFontStyle("PropertyWindow.NormalFont"));
 					}
 
 					return Widget.ToSharedRef();
 				})
-				.OnSelectionChanged_Lambda([this]( TSharedPtr<FString> ChosenOption, ESelectInfo::Type SelectInfo )
+				.OnSelectionChanged_Lambda([this](TSharedPtr<FString> ChosenOption, ESelectInfo::Type SelectInfo)
 				{
-					if ( ChosenOption && CurrentActor )
+					if (ChosenOption && CurrentActor)
 					{
-						CurrentActor->SetMaterialPurpose( **ChosenOption );
+						CurrentActor->SetMaterialPurpose(**ChosenOption);
 					}
 				})
 				[
-					SNew( SEditableTextBox )
+					SNew(SEditableTextBox)
 					.Text_Lambda([this]() -> FText
 					{
-						if ( CurrentActor )
+						if (CurrentActor)
 						{
-							return FText::FromString( CurrentActor->MaterialPurpose == *UnrealIdentifiers::MaterialAllPurpose
+							return FText::FromString(CurrentActor->MaterialPurpose == *UnrealIdentifiers::MaterialAllPurpose
 								? UnrealIdentifiers::MaterialAllPurposeText
 								: CurrentActor->MaterialPurpose.ToString()
 							);
@@ -219,10 +220,10 @@ void FUsdStageActorCustomization::CustomizeDetails( IDetailLayoutBuilder& Detail
 
 						return FText::GetEmpty();
 					})
-					.Font( FAppStyle::GetFontStyle( "PropertyWindow.NormalFont" ) )
-					.OnTextCommitted_Lambda( [this]( const FText& NewText, ETextCommit::Type CommitType )
+					.Font(FAppStyle::GetFontStyle("PropertyWindow.NormalFont"))
+					.OnTextCommitted_Lambda([this](const FText& NewText, ETextCommit::Type CommitType)
 					{
-						if ( CommitType != ETextCommit::OnEnter )
+						if (CommitType != ETextCommit::OnEnter)
 						{
 							return;
 						}
@@ -231,75 +232,76 @@ void FUsdStageActorCustomization::CustomizeDetails( IDetailLayoutBuilder& Detail
 						FName NewPurpose = *NewPurposeString;
 
 						bool bIsNew = true;
-						for ( const TSharedPtr<FString>& Purpose : MaterialPurposeComboBoxItems )
+						for (const TSharedPtr<FString>& Purpose : MaterialPurposeComboBoxItems)
 						{
-							if ( Purpose && *Purpose == NewPurposeString )
+							if (Purpose && *Purpose == NewPurposeString)
 							{
 								bIsNew = false;
 								break;
 							}
 						}
 
-						if ( bIsNew )
+						if (bIsNew)
 						{
-							if ( UUsdProjectSettings* ProjectSettings = GetMutableDefault<UUsdProjectSettings>() )
+							if (UUsdProjectSettings* ProjectSettings = GetMutableDefault<UUsdProjectSettings>())
 							{
-								ProjectSettings->AdditionalMaterialPurposes.AddUnique( NewPurpose );
+								ProjectSettings->AdditionalMaterialPurposes.AddUnique(NewPurpose);
 								ProjectSettings->SaveConfig();
 							}
 						}
 
-						if ( CurrentActor )
+						if (CurrentActor)
 						{
-							CurrentActor->SetMaterialPurpose( NewPurpose );
+							CurrentActor->SetMaterialPurpose(NewPurpose);
 						}
 					})
 				]
 			]
 		];
 	}
+	// clang-format on
 
 	// Add/remove properties so that they retain their usual order
-	if ( TSharedPtr<IPropertyHandle> TimeProperty = DetailLayoutBuilder.GetProperty( TEXT( "Time" ) ) )
+	if (TSharedPtr<IPropertyHandle> TimeProperty = DetailLayoutBuilder.GetProperty(TEXT("Time")))
 	{
-		CatBuilder.AddProperty( TimeProperty );
+		CatBuilder.AddProperty(TimeProperty);
 	}
-	if ( TSharedPtr<IPropertyHandle> LevelSequenceProperty = DetailLayoutBuilder.GetProperty( TEXT( "LevelSequence" ) ) )
+	if (TSharedPtr<IPropertyHandle> LevelSequenceProperty = DetailLayoutBuilder.GetProperty(TEXT("LevelSequence")))
 	{
-		CatBuilder.AddProperty( LevelSequenceProperty );
+		CatBuilder.AddProperty(LevelSequenceProperty);
 	}
 }
 
-void FUsdStageActorCustomization::CustomizeDetails( const TSharedPtr<IDetailLayoutBuilder>& DetailBuilder )
+void FUsdStageActorCustomization::CustomizeDetails(const TSharedPtr<IDetailLayoutBuilder>& DetailBuilder)
 {
 	DetailBuilderWeakPtr = DetailBuilder;
-	CustomizeDetails( *DetailBuilder );
+	CustomizeDetails(*DetailBuilder);
 }
 
-void FUsdStageActorCustomization::OnComboBoxSelectionChanged( TSharedPtr<FString> NewContext, ESelectInfo::Type SelectType )
+void FUsdStageActorCustomization::OnComboBoxSelectionChanged(TSharedPtr<FString> NewContext, ESelectInfo::Type SelectType)
 {
-	if ( CurrentActor == nullptr || !NewContext.IsValid() )
+	if (CurrentActor == nullptr || !NewContext.IsValid())
 	{
 		return;
 	}
 
 	FScopedTransaction Transaction(FText::Format(
 		LOCTEXT("RenderContextChangedTransaction", "Changed the UsdStageActor {0}'s RenderContext to '{1}'"),
-		FText::FromString( CurrentActor->GetActorLabel() ),
-		FText::FromString( NewContext.IsValid() ? *NewContext : TEXT( "None" ) ) )
-	);
+		FText::FromString(CurrentActor->GetActorLabel()),
+		FText::FromString(NewContext.IsValid() ? *NewContext : TEXT("None"))
+	));
 
-	FName NewContextName = ( *NewContext ) == TEXT( "universal" ) ? NAME_None : FName( **NewContext );
+	FName NewContextName = (*NewContext) == TEXT("universal") ? NAME_None : FName(**NewContext);
 
-	CurrentActor->SetRenderContext( NewContextName );
+	CurrentActor->SetRenderContext(NewContextName);
 }
 
 FText FUsdStageActorCustomization::GetComboBoxSelectedOptionText() const
 {
 	TSharedPtr<FString> SelectedItem = RenderContextComboBox->GetSelectedItem();
-	if ( SelectedItem.IsValid() )
+	if (SelectedItem.IsValid())
 	{
-		return FText::FromString( *SelectedItem );
+		return FText::FromString(*SelectedItem);
 	}
 
 	return FText::GetEmpty();
@@ -309,7 +311,7 @@ void FUsdStageActorCustomization::ForceRefreshDetails()
 {
 	// Raw because we don't want to keep alive the details builder when calling the force refresh details
 	IDetailLayoutBuilder* DetailLayoutBuilder = DetailBuilderWeakPtr.Pin().Get();
-	if ( DetailLayoutBuilder )
+	if (DetailLayoutBuilder)
 	{
 		DetailLayoutBuilder->ForceRefreshDetails();
 	}
@@ -317,4 +319,4 @@ void FUsdStageActorCustomization::ForceRefreshDetails()
 
 #undef LOCTEXT_NAMESPACE
 
-#endif // WITH_EDITOR
\ No newline at end of file
+#endif	  // WITH_EDITOR
diff --git a/Engine/Plugins/Importers/USDImporter/Source/USDStage/Private/USDStageActorCustomization.h b/Engine/Plugins/Importers/USDImporter/Source/USDStage/Private/USDStageActorCustomization.h
--- a/Engine/Plugins/Importers/USDImporter/Source/USDStage/Private/USDStageActorCustomization.h
+++ b/Engine/Plugins/Importers/USDImporter/Source/USDStage/Private/USDStageActorCustomization.h
@@ -4,7 +4,6 @@
 
 #if WITH_EDITOR
 
-#include "CoreMinimal.h"
 #include "IDetailCustomization.h"
 #include "Widgets/Input/SComboBox.h"
 
@@ -18,12 +17,12 @@ public:
 	static TSharedRef<IDetailCustomization> MakeInstance();
 
 	// IDetailCustomization interface
-	virtual void CustomizeDetails( IDetailLayoutBuilder& DetailLayout ) override;
-	virtual void CustomizeDetails( const TSharedPtr<IDetailLayoutBuilder>& DetailBuilder ) override;
+	virtual void CustomizeDetails(IDetailLayoutBuilder& DetailLayout) override;
+	virtual void CustomizeDetails(const TSharedPtr<IDetailLayoutBuilder>& DetailBuilder) override;
 	// End of IDetailCustomization interface
 
 private:
-	void OnComboBoxSelectionChanged( TSharedPtr<FString> NewContext, ESelectInfo::Type SelectType );
+	void OnComboBoxSelectionChanged(TSharedPtr<FString> NewContext, ESelectInfo::Type SelectType);
 	FText GetComboBoxSelectedOptionText() const;
 	void ForceRefreshDetails();
 
@@ -36,4 +35,4 @@ private:
 	TArray<TSharedPtr<FString>> MaterialPurposeComboBoxItems;
 };
 
-#endif // WITH_EDITOR
\ No newline at end of file
+#endif	  // WITH_EDITOR
diff --git a/Engine/Plugins/Importers/USDImporter/Source/USDStage/Private/USDStageModule.cpp b/Engine/Plugins/Importers/USDImporter/Source/USDStage/Private/USDStageModule.cpp
--- a/Engine/Plugins/Importers/USDImporter/Source/USDStage/Private/USDStageModule.cpp
+++ b/Engine/Plugins/Importers/USDImporter/Source/USDStage/Private/USDStageModule.cpp
@@ -6,16 +6,15 @@
 #include "USDStageActor.h"
 #include "USDStageActorCustomization.h"
 
-#include "Modules/ModuleManager.h"
-#include "EngineUtils.h"
 #include "Engine/World.h"
+#include "EngineUtils.h"
 #include "Interfaces/IPluginManager.h"
+#include "Modules/ModuleManager.h"
 
 #if WITH_EDITOR
 #include "ISequencerModule.h"
 #include "PropertyEditorModule.h"
-#include "Settings/ProjectPackagingSettings.h"
-#endif // WITH_EDITOR
+#endif	  // WITH_EDITOR
 
 class FUsdStageModule : public IUsdStageModule
 {
@@ -25,8 +24,11 @@ public:
 #if WITH_EDITOR
 		LLM_SCOPE_BYTAG(Usd);
 
-		FPropertyEditorModule& PropertyModule = FModuleManager::LoadModuleChecked<FPropertyEditorModule>( TEXT( "PropertyEditor" ) );
-		PropertyModule.RegisterCustomClassLayout( TEXT( "UsdStageActor" ), FOnGetDetailCustomizationInstance::CreateStatic( &FUsdStageActorCustomization::MakeInstance ) );
+		FPropertyEditorModule& PropertyModule = FModuleManager::LoadModuleChecked<FPropertyEditorModule>(TEXT("PropertyEditor"));
+		PropertyModule.RegisterCustomClassLayout(
+			TEXT("UsdStageActor"),
+			FOnGetDetailCustomizationInstance::CreateStatic(&FUsdStageActorCustomization::MakeInstance)
+		);
 
 		Sequencers.Reset();
 		ISequencerModule& SequencerModule = FModuleManager::Get().LoadModuleChecked<ISequencerModule>("Sequencer");
@@ -48,21 +50,21 @@ public:
 				Sequencers.Add(NewSequencer.ToWeakPtr());
 			}
 		));
-#endif // WITH_EDITOR
+#endif	  // WITH_EDITOR
 	}
 
 	virtual void ShutdownModule() override
 	{
 #if WITH_EDITOR
-		FPropertyEditorModule& PropertyModule = FModuleManager::LoadModuleChecked< FPropertyEditorModule >( TEXT( "PropertyEditor" ) );
-		PropertyModule.UnregisterCustomClassLayout( TEXT( "UsdStageActor" ) );
+		FPropertyEditorModule& PropertyModule = FModuleManager::LoadModuleChecked<FPropertyEditorModule>(TEXT("PropertyEditor"));
+		PropertyModule.UnregisterCustomClassLayout(TEXT("UsdStageActor"));
 
-		if(ISequencerModule* SequencerModule = FModuleManager::Get().GetModulePtr<ISequencerModule>("Sequencer"))
+		if (ISequencerModule* SequencerModule = FModuleManager::Get().GetModulePtr<ISequencerModule>("Sequencer"))
 		{
 			SequencerModule->UnregisterOnSequencerCreated(OnSequencerCreatedHandle);
 			Sequencers.Reset();
 		}
-#endif // WITH_EDITOR
+#endif	  // WITH_EDITOR
 	}
 
 #if WITH_EDITOR
@@ -71,25 +73,25 @@ public:
 	{
 		return Sequencers;
 	}
-#endif // WITH_EDITOR
+#endif	  // WITH_EDITOR
 
-	virtual AUsdStageActor& GetUsdStageActor( UWorld* World ) override
+	virtual AUsdStageActor& GetUsdStageActor(UWorld* World) override
 	{
-		if ( AUsdStageActor* UsdStageActor = FindUsdStageActor( World ) )
+		if (AUsdStageActor* UsdStageActor = FindUsdStageActor(World))
 		{
 			return *UsdStageActor;
 		}
 		else
 		{
-			return *( World->SpawnActor< AUsdStageActor >() );
+			return *(World->SpawnActor<AUsdStageActor>());
 		}
 	}
 
-	virtual AUsdStageActor* FindUsdStageActor( UWorld* World ) override
+	virtual AUsdStageActor* FindUsdStageActor(UWorld* World) override
 	{
-		for ( FActorIterator ActorIterator( World ); ActorIterator; ++ActorIterator )
+		for (FActorIterator ActorIterator(World); ActorIterator; ++ActorIterator)
 		{
-			if ( AUsdStageActor* UsdStageActor = Cast< AUsdStageActor >( *ActorIterator ) )
+			if (AUsdStageActor* UsdStageActor = Cast<AUsdStageActor>(*ActorIterator))
 			{
 				return UsdStageActor;
 			}
@@ -102,7 +104,7 @@ private:
 #if WITH_EDITOR
 	TArray<TWeakPtr<ISequencer>> Sequencers;
 	FDelegateHandle OnSequencerCreatedHandle;
-#endif // WITH_EDITOR
+#endif	  // WITH_EDITOR
 };
 
-IMPLEMENT_MODULE_USD( FUsdStageModule, USDStage );
+IMPLEMENT_MODULE_USD(FUsdStageModule, USDStage);
diff --git a/Engine/Plugins/Importers/USDImporter/Source/USDStage/Private/USDTransactor.cpp b/Engine/Plugins/Importers/USDImporter/Source/USDStage/Private/USDTransactor.cpp
--- a/Engine/Plugins/Importers/USDImporter/Source/USDStage/Private/USDTransactor.cpp
+++ b/Engine/Plugins/Importers/USDImporter/Source/USDStage/Private/USDTransactor.cpp
@@ -3,7 +3,6 @@
 #include "USDTransactor.h"
 
 #include "UnrealUSDWrapper.h"
-#include "USDErrorUtils.h"
 #include "USDLayerUtils.h"
 #include "USDListener.h"
 #include "USDLog.h"
@@ -24,11 +23,11 @@
 #include "Editor/Transactor.h"
 #include "Editor/TransBuffer.h"
 #include "Misc/ITransaction.h"
-#endif // WITH_EDITOR
+#endif	  // WITH_EDITOR
 
 #define LOCTEXT_NAMESPACE "USDTransactor"
 
-const FName UE::UsdTransactor::ConcertSyncEnableTag = TEXT( "EnableConcertSync" );
+const FName UE::UsdTransactor::ConcertSyncEnableTag = TEXT("EnableConcertSync");
 
 namespace UsdUtils
 {
@@ -37,11 +36,13 @@ namespace UsdUtils
 	{
 		FString PropertyName;
 		FString Field;
-		FString AttributeTypeName;				// Full SdfValueTypeName of the attribute (e.g. normal3f, bool, texCoord3d, float2) so that we can undo/redo attribute creation
-		FConvertedVtValue OldValue;				// Default old value
-		FConvertedVtValue NewValue;				// Default new value
+		FString AttributeTypeName;	   // Full SdfValueTypeName of the attribute (e.g. normal3f, bool, texCoord3d, float2) so that we can undo/redo
+									   // attribute creation
+		FConvertedVtValue OldValue;	   // Default old value
+		FConvertedVtValue NewValue;	   // Default new value
 		TArray<double> TimeSamples;
-		TArray< FConvertedVtValue > TimeValues; // We can't fetch old/new values when time samples change, so we just have one of these and save the whole set every time
+		TArray<FConvertedVtValue> TimeValues;	 // We can't fetch old/new values when time samples change, so we just have one of these and save the
+												 // whole set every time
 	};
 
 	struct FTransactorObjectChange
@@ -62,19 +63,19 @@ namespace UsdUtils
 	struct FTransactorRecordedEdits
 	{
 		FString EditTargetIdentifier;
-		TArray< FTransactorRecordedEdit > Edits; // Kept in the order of recording
+		TArray<FTransactorRecordedEdit> Edits;	  // Kept in the order of recording
 	};
 
-	using FTransactorEditStorage = TArray< FTransactorRecordedEdits >;
+	using FTransactorEditStorage = TArray<FTransactorRecordedEdits>;
 
 	enum class EApplicationDirection
 	{
 		Reverse,
 		Forward
 	};
-}
+}	 // namespace UsdUtils
 
-FArchive& operator<<( FArchive& Ar, UsdUtils::FTransactorAttributeChange& Change )
+FArchive& operator<<(FArchive& Ar, UsdUtils::FTransactorAttributeChange& Change)
 {
 	Ar << Change.PropertyName;
 	Ar << Change.Field;
@@ -86,20 +87,20 @@ FArchive& operator<<( FArchive& Ar, UsdUtils::FTransactorAttributeChange& Change
 	return Ar;
 }
 
-FArchive& operator<<( FArchive& Ar, UsdUtils::FPrimChangeFlags& Flags )
+FArchive& operator<<(FArchive& Ar, UsdUtils::FPrimChangeFlags& Flags)
 {
 	TArray<uint8> Bytes;
-	if ( Ar.IsLoading() )
+	if (Ar.IsLoading())
 	{
 		Ar << Bytes;
 
-		check( Bytes.Num() == sizeof( Flags ) );
-		FMemory::Memcpy( reinterpret_cast< uint8* >( &Flags ), Bytes.GetData(), sizeof( Flags ) );
+		check(Bytes.Num() == sizeof(Flags));
+		FMemory::Memcpy(reinterpret_cast<uint8*>(&Flags), Bytes.GetData(), sizeof(Flags));
 	}
-	else if ( Ar.IsSaving() )
+	else if (Ar.IsSaving())
 	{
-		Bytes.SetNumZeroed( sizeof( Flags ) );
-		FMemory::Memcpy( Bytes.GetData(), reinterpret_cast< uint8* >( &Flags ), sizeof( Flags ) );
+		Bytes.SetNumZeroed(sizeof(Flags));
+		FMemory::Memcpy(Bytes.GetData(), reinterpret_cast<uint8*>(&Flags), sizeof(Flags));
 
 		Ar << Bytes;
 	}
@@ -107,7 +108,7 @@ FArchive& operator<<( FArchive& Ar, UsdUtils::FPrimChangeFlags& Flags )
 	return Ar;
 }
 
-FArchive& operator<<( FArchive& Ar, UsdUtils::FTransactorObjectChange& Change )
+FArchive& operator<<(FArchive& Ar, UsdUtils::FTransactorObjectChange& Change)
 {
 	Ar << Change.AttributeChanges;
 	Ar << Change.Flags;
@@ -117,14 +118,14 @@ FArchive& operator<<( FArchive& Ar, UsdUtils::FTransactorObjectChange& Change )
 	return Ar;
 }
 
-FArchive& operator<<( FArchive& Ar, UsdUtils::FTransactorRecordedEdit& Edit )
+FArchive& operator<<(FArchive& Ar, UsdUtils::FTransactorRecordedEdit& Edit)
 {
 	Ar << Edit.ObjectPath;
 	Ar << Edit.ObjectChanges;
 	return Ar;
 }
 
-FArchive& operator<<( FArchive& Ar, UsdUtils::FTransactorRecordedEdits& Edits )
+FArchive& operator<<(FArchive& Ar, UsdUtils::FTransactorRecordedEdits& Edits)
 {
 	Ar << Edits.EditTargetIdentifier;
 	Ar << Edits.Edits;
@@ -137,11 +138,11 @@ namespace UsdUtils
 	 * Converts the received VtValue map to an analogue using converted UE types that can be serialized with the UUsdTransactor.
 	 * Needs the stage because we need to manually fetch additional prim/attribute data, in order to support undo/redoing attribute creation
 	 */
-	bool ConvertFieldValueMap( const FObjectChangesByPath& InChanges, const UE::FUsdStage& InStage, FTransactorRecordedEdits& InOutEdits )
+	bool ConvertFieldValueMap(const FObjectChangesByPath& InChanges, const UE::FUsdStage& InStage, FTransactorRecordedEdits& InOutEdits)
 	{
-		InOutEdits.Edits.Reserve( InChanges.Num() + InOutEdits.Edits.Num() );
+		InOutEdits.Edits.Reserve(InChanges.Num() + InOutEdits.Edits.Num());
 
-		for ( const TPair<FString, TArray<FSdfChangeListEntry>>& ChangesByPrimPath : InChanges )
+		for (const TPair<FString, TArray<FSdfChangeListEntry>>& ChangesByPrimPath : InChanges)
 		{
 			const FString& PrimPath = ChangesByPrimPath.Key;
 			const TArray<FSdfChangeListEntry>& Changes = ChangesByPrimPath.Value;
@@ -149,52 +150,53 @@ namespace UsdUtils
 			FTransactorRecordedEdit& Edit = InOutEdits.Edits.Emplace_GetRef();
 			Edit.ObjectPath = PrimPath;
 
-			for ( const FSdfChangeListEntry& Change : Changes )
+			for (const FSdfChangeListEntry& Change : Changes)
 			{
 				FTransactorObjectChange& ConvertedChange = Edit.ObjectChanges.Emplace_GetRef();
 				ConvertedChange.Flags = Change.Flags;
 				ConvertedChange.OldPath = Change.OldPath;
 
-				UE::FUsdPrim Prim = InStage.GetPrimAtPath( UE::FSdfPath{ *PrimPath } );
+				UE::FUsdPrim Prim = InStage.GetPrimAtPath(UE::FSdfPath{*PrimPath});
 
-				if ( Prim &&
-					( ConvertedChange.Flags.bDidAddInertPrim ||
-					  ConvertedChange.Flags.bDidAddNonInertPrim ||
-					  ConvertedChange.Flags.bDidRemoveInertPrim ||
-					  ConvertedChange.Flags.bDidRemoveNonInertPrim ) )
+				if (Prim
+					&& (ConvertedChange.Flags.bDidAddInertPrim || ConvertedChange.Flags.bDidAddNonInertPrim
+						|| ConvertedChange.Flags.bDidRemoveInertPrim || ConvertedChange.Flags.bDidRemoveNonInertPrim))
 				{
 					ConvertedChange.PrimTypeName = Prim.GetTypeName().ToString();
 
 					TArray<FName> AppliedSchemaNames = Prim.GetAppliedSchemas();
-					ConvertedChange.PrimAppliedSchemas.Reset( AppliedSchemaNames.Num() );
-					for ( const FName SchemaName : AppliedSchemaNames )
+					ConvertedChange.PrimAppliedSchemas.Reset(AppliedSchemaNames.Num());
+					for (const FName SchemaName : AppliedSchemaNames)
 					{
-						ConvertedChange.PrimAppliedSchemas.Add( SchemaName.ToString() );
+						ConvertedChange.PrimAppliedSchemas.Add(SchemaName.ToString());
 					}
 
-					UE_LOG( LogUsd, Log, TEXT( "Recorded the %s of prim '%s' with TypeName '%s' and PrimAppliedSchemas [%s]" ),
+					UE_LOG(
+						LogUsd,
+						Log,
+						TEXT("Recorded the %s of prim '%s' with TypeName '%s' and PrimAppliedSchemas [%s]"),
 						(ConvertedChange.Flags.bDidAddInertPrim || ConvertedChange.Flags.bDidAddNonInertPrim) ? TEXT("addition") : TEXT("removal"),
 						*PrimPath,
 						*ConvertedChange.PrimTypeName,
-						*FString::Join( ConvertedChange.PrimAppliedSchemas, TEXT( ", " ) )
+						*FString::Join(ConvertedChange.PrimAppliedSchemas, TEXT(", "))
 					);
 				}
 
 				// Don't record any attribute changes if we can't find the prim anyway
-				if ( !Prim )
+				if (!Prim)
 				{
 					// We expect not to find the prim if the change says it was just removed though
-					if ( !ConvertedChange.Flags.bDidRemoveInertPrim && !ConvertedChange.Flags.bDidRemoveNonInertPrim )
+					if (!ConvertedChange.Flags.bDidRemoveInertPrim && !ConvertedChange.Flags.bDidRemoveNonInertPrim)
 					{
-						UE_LOG( LogUsd, Warning, TEXT( "Failed to find prim at path '%s' when serializing object changes in transactor" ), *PrimPath );
+						UE_LOG(LogUsd, Warning, TEXT("Failed to find prim at path '%s' when serializing object changes in transactor"), *PrimPath);
 					}
 					continue;
 				}
 
-				for ( const FAttributeChange& AttributeChange : Change.AttributeChanges )
+				for (const FAttributeChange& AttributeChange : Change.AttributeChanges)
 				{
 					// We won't be able to do anything with this, so just ignore it
-					if ( AttributeChange.PropertyName.IsEmpty() )
+					if (AttributeChange.PropertyName.IsEmpty())
 					{
 						continue;
 					}
@@ -204,22 +206,28 @@ namespace UsdUtils
 					ConvertedAttributeChange.Field = AttributeChange.Field;
 
 					FConvertedVtValue ConvertedOldValue;
-					if ( UsdToUnreal::ConvertValue( AttributeChange.OldValue, ConvertedOldValue ) )
+					if (UsdToUnreal::ConvertValue(AttributeChange.OldValue, ConvertedOldValue))
 					{
 						ConvertedAttributeChange.OldValue = ConvertedOldValue;
 					}
 
 					FConvertedVtValue ConvertedNewValue;
-					if ( UsdToUnreal::ConvertValue( AttributeChange.NewValue, ConvertedNewValue ) )
+					if (UsdToUnreal::ConvertValue(AttributeChange.NewValue, ConvertedNewValue))
 					{
 						ConvertedAttributeChange.NewValue = ConvertedNewValue;
 					}
 
-					// We likely won't be able to fetch the attribute if it was removed, so try deducing the typename from the value just so that we have *something*
-					if ( !AttributeChange.OldValue.IsEmpty() && (ConvertedChange.Flags.bDidRemoveProperty || ConvertedChange.Flags.bDidRemovePropertyWithOnlyRequiredFields) )
+					// We likely won't be able to fetch the attribute if it was removed, so try deducing the typename from the value just so that we
+					// have *something*
+					if (!AttributeChange.OldValue.IsEmpty()
+						&& (ConvertedChange.Flags.bDidRemoveProperty || ConvertedChange.Flags.bDidRemovePropertyWithOnlyRequiredFields))
 					{
-						ConvertedAttributeChange.AttributeTypeName = UsdUtils::GetImpliedTypeName( AttributeChange.OldValue );
-						UE_LOG( LogUsd, Log, TEXT( "Recording the removal of properties is not fully supported: Using underlying type '%s' for record of attribute '%s' of prim '%s', as we don't have access to the attribute's role" ),
+						ConvertedAttributeChange.AttributeTypeName = UsdUtils::GetImpliedTypeName(AttributeChange.OldValue);
+						UE_LOG(
+							LogUsd,
+							Log,
+							TEXT("Recording the removal of properties is not fully supported: Using underlying type '%s' for record of attribute "
+								 "'%s' of prim '%s', as we don't have access to the attribute's role"),
 							*ConvertedAttributeChange.AttributeTypeName,
 							*ConvertedAttributeChange.PropertyName,
 							*PrimPath
@@ -228,32 +236,37 @@ namespace UsdUtils
 
 					// If the prim is the root, this is a stage info change, and PropertyName is actually a metadata key, so
 					// it's not possible (or required) to know its SdfValueTypeName to e.g. undo/redo the creation of the property
-					if ( !Prim.IsPseudoRoot() && ConvertedAttributeChange.PropertyName != TEXT( "kind" ) && !ConvertedAttributeChange.PropertyName.IsEmpty() &&
-						 !ConvertedChange.Flags.bDidRemoveProperty && !ConvertedChange.Flags.bDidRemovePropertyWithOnlyRequiredFields )
+					if (!Prim.IsPseudoRoot() && ConvertedAttributeChange.PropertyName != TEXT("kind")
+						&& !ConvertedAttributeChange.PropertyName.IsEmpty() && !ConvertedChange.Flags.bDidRemoveProperty
+						&& !ConvertedChange.Flags.bDidRemovePropertyWithOnlyRequiredFields)
 					{
-						if ( UE::FUsdAttribute Attribute = Prim.GetAttribute( *ConvertedAttributeChange.PropertyName ) )
+						if (UE::FUsdAttribute Attribute = Prim.GetAttribute(*ConvertedAttributeChange.PropertyName))
 						{
 							ConvertedAttributeChange.AttributeTypeName = Attribute.GetTypeName().ToString();
 
-							// USD doesn't tell us what changed, what type of change it was, or old/new values... so just save the entire timeSamples of the attribute so we can mirror via multi-user
-							if ( ConvertedChange.Flags.bDidChangeAttributeTimeSamples )
+							// USD doesn't tell us what changed, what type of change it was, or old/new values... so just save the entire timeSamples
+							// of the attribute so we can mirror via multi-user
+							if (ConvertedChange.Flags.bDidChangeAttributeTimeSamples)
 							{
-								Attribute.GetTimeSamples( ConvertedAttributeChange.TimeSamples );
-								ConvertedAttributeChange.TimeValues.SetNum( ConvertedAttributeChange.TimeSamples.Num() );
+								Attribute.GetTimeSamples(ConvertedAttributeChange.TimeSamples);
+								ConvertedAttributeChange.TimeValues.SetNum(ConvertedAttributeChange.TimeSamples.Num());
 
-								for ( int32 Index = 0; Index < ConvertedAttributeChange.TimeSamples.Num(); Index++ )
+								for (int32 Index = 0; Index < ConvertedAttributeChange.TimeSamples.Num(); Index++)
 								{
 									UE::FVtValue UsdValue;
-									if ( Attribute.Get( UsdValue, ConvertedAttributeChange.TimeSamples[ Index ] ) )
+									if (Attribute.Get(UsdValue, ConvertedAttributeChange.TimeSamples[Index]))
 									{
-										UsdToUnreal::ConvertValue( UsdValue, ConvertedAttributeChange.TimeValues[ Index ] );
+										UsdToUnreal::ConvertValue(UsdValue, ConvertedAttributeChange.TimeValues[Index]);
 									}
 								}
 							}
 						}
 						else
 						{
-							UE_LOG( LogUsd, Warning, TEXT( "Failed to find attribute '%s' for prim at path '%s' when serializing object changes in transactor" ),
+							UE_LOG(
+								LogUsd,
+								Warning,
+								TEXT("Failed to find attribute '%s' for prim at path '%s' when serializing object changes in transactor"),
 								*ConvertedAttributeChange.PropertyName,
 								*PrimPath
 							);
@@ -266,43 +279,42 @@ namespace UsdUtils
 		return true;
 	}
 
-	bool ApplyPrimChange( const UE::FSdfPath& PrimPath, const FTransactorObjectChange& PrimChange, UE::FUsdStage& Stage, EApplicationDirection Direction )
+	bool ApplyPrimChange(
+		const UE::FSdfPath& PrimPath,
+		const FTransactorObjectChange& PrimChange,
+		UE::FUsdStage& Stage,
+		EApplicationDirection Direction
+	)
 	{
 		const bool bAdd = PrimChange.Flags.bDidAddInertPrim || PrimChange.Flags.bDidAddNonInertPrim;
 		const bool bRemove = PrimChange.Flags.bDidRemoveInertPrim || PrimChange.Flags.bDidRemoveNonInertPrim;
 		const bool bRename = PrimChange.Flags.bDidRename;
 
-		if ( ( bAdd && Direction == EApplicationDirection::Forward ) || ( bRemove && Direction == EApplicationDirection::Reverse ) )
+		if ((bAdd && Direction == EApplicationDirection::Forward) || (bRemove && Direction == EApplicationDirection::Reverse))
 		{
-			UE_LOG( LogUsd, Log, TEXT( "Creating prim '%s' with typename '%s'" ),
-				*PrimPath.GetString(),
-				*PrimChange.PrimTypeName
-			);
+			UE_LOG(LogUsd, Log, TEXT("Creating prim '%s' with typename '%s'"), *PrimPath.GetString(), *PrimChange.PrimTypeName);
 
-			UE::FUsdPrim Prim = Stage.DefinePrim( PrimPath, *PrimChange.PrimTypeName );
-			if ( !Prim )
+			UE::FUsdPrim Prim = Stage.DefinePrim(PrimPath, *PrimChange.PrimTypeName);
+			if (!Prim)
 			{
 				return false;
 			}
 
 			return true;
 		}
-		else if ( ( bAdd && Direction == EApplicationDirection::Reverse ) || ( bRemove && Direction == EApplicationDirection::Forward ) )
+		else if ((bAdd && Direction == EApplicationDirection::Reverse) || (bRemove && Direction == EApplicationDirection::Forward))
 		{
-			UE_LOG( LogUsd, Log, TEXT( "Removing prim '%s' with typename '%s'" ),
-				*PrimPath.GetString(),
-				*PrimChange.PrimTypeName
-			);
+			UE_LOG(LogUsd, Log, TEXT("Removing prim '%s' with typename '%s'"), *PrimPath.GetString(), *PrimChange.PrimTypeName);
 
-			UsdUtils::RemoveAllLocalPrimSpecs( Stage.GetPrimAtPath( PrimPath ), Stage.GetEditTarget() );
+			UsdUtils::RemoveAllLocalPrimSpecs(Stage.GetPrimAtPath(PrimPath), Stage.GetEditTarget());
 			return true;
 		}
-		else if ( bRename )
+		else if (bRename)
 		{
 			FString NewName;
 			FString CurrentPath;
 
-			if ( Direction == EApplicationDirection::Forward )
+			if (Direction == EApplicationDirection::Forward)
 			{
 				// It hasn't been renamed yet, so it's still at the old path
 				CurrentPath = PrimChange.OldPath;
@@ -311,28 +323,25 @@ namespace UsdUtils
 			else
 			{
 				CurrentPath = PrimPath.GetString();
-				NewName = UE::FSdfPath( *PrimChange.OldPath ).GetElementString();
+				NewName = UE::FSdfPath(*PrimChange.OldPath).GetElementString();
 			}
 
 			// When redoing, we'll be using the old path, and USD sends it with all the variant selections in there
 			// UsdUtils::RenamePrim can figure out the variant selections on its own, but we need to strip them here
 			// to be able to GetPrimAtPath with this path
-			UE::FSdfPath UsdCurrentPath = UE::FSdfPath( *CurrentPath ).StripAllVariantSelections();
+			UE::FSdfPath UsdCurrentPath = UE::FSdfPath(*CurrentPath).StripAllVariantSelections();
 
-			if ( UE::FUsdPrim Prim = Stage.GetPrimAtPath( UsdCurrentPath ) )
+			if (UE::FUsdPrim Prim = Stage.GetPrimAtPath(UsdCurrentPath))
 			{
-				UE_LOG( LogUsd, Log, TEXT( "Renaming prim '%s' to '%s'" ),
-					*Prim.GetPrimPath().GetString(),
-					*NewName
-				);
+				UE_LOG(LogUsd, Log, TEXT("Renaming prim '%s' to '%s'"), *Prim.GetPrimPath().GetString(), *NewName);
 
-				const bool bDidRename = UsdUtils::RenamePrim( Prim, *NewName );
-				if ( bDidRename )
+				const bool bDidRename = UsdUtils::RenamePrim(Prim, *NewName);
+				if (bDidRename)
 				{
 					return true;
 				}
 			}
-			else if ( UE::FUsdPrim TargetPrim = Stage.GetPrimAtPath( UsdCurrentPath.ReplaceName( *NewName ) ) )
+			else if (UE::FUsdPrim TargetPrim = Stage.GetPrimAtPath(UsdCurrentPath.ReplaceName(*NewName)))
 			{
 				// We couldn't find a prim at the old path but found one at the new path, so just assume its the prim that we
 				// wanted to rename anyway, as USD wouldn't have let us rename a prim onto an existing path in the first place.
@@ -341,56 +350,58 @@ namespace UsdUtils
 				return true;
 			}
 
-			UE_LOG( LogUsd, Warning, TEXT( "Failed to rename prim at path '%s' to name '%s'" ),
-				*CurrentPath,
-				*NewName
-			);
+			UE_LOG(LogUsd, Warning, TEXT("Failed to rename prim at path '%s' to name '%s'"), *CurrentPath, *NewName);
 		}
 
 		return false;
 	}
 
-	bool ApplyAttributeTimeSamples( const FTransactorAttributeChange& AttributeChange, UE::FUsdPrim& Prim )
+	bool ApplyAttributeTimeSamples(const FTransactorAttributeChange& AttributeChange, UE::FUsdPrim& Prim)
 	{
-		if ( !Prim ||
-			AttributeChange.TimeSamples.Num() == 0 ||
-			AttributeChange.TimeSamples.Num() != AttributeChange.TimeValues.Num() ||
-			AttributeChange.Field != TEXT( "timeSamples" ) )
+		if (!Prim || AttributeChange.TimeSamples.Num() == 0 || AttributeChange.TimeSamples.Num() != AttributeChange.TimeValues.Num()
+			|| AttributeChange.Field != TEXT("timeSamples"))
 		{
 			return false;
 		}
 
-		// Try Getting first because we shouldn't trust our AttributeTypeName to always just CreateAttribute, as it may be just deduced from a value and be different
-		UE::FUsdAttribute Attribute = Prim.GetAttribute( *AttributeChange.PropertyName );
-		if ( !Attribute )
+		// Try Getting first because we shouldn't trust our AttributeTypeName to always just CreateAttribute, as it may be just deduced from a value
+		// and be different
+		UE::FUsdAttribute Attribute = Prim.GetAttribute(*AttributeChange.PropertyName);
+		if (!Attribute)
 		{
-			Attribute = Prim.CreateAttribute( *AttributeChange.PropertyName, *AttributeChange.AttributeTypeName );
+			Attribute = Prim.CreateAttribute(*AttributeChange.PropertyName, *AttributeChange.AttributeTypeName);
 		}
-		if ( !Attribute )
+		if (!Attribute)
 		{
 			return false;
 		}
 
 		// Clear all timesamples because we may have more timesamples than we receive, and we want our old ones to be removed
 		// This corresponds to the token SdfFieldKeys->TimeSamples, and is extracted from UsdAttribute::Clear
-		Attribute.ClearMetadata( TEXT( "timeSamples" ) );
+		Attribute.ClearMetadata(TEXT("timeSamples"));
 
-		UE_LOG( LogUsd, Log, TEXT( "Applying '%d' timeSamples for attribute '%s' of prim '%s'" ),
+		UE_LOG(
+			LogUsd,
+			Log,
+			TEXT("Applying '%d' timeSamples for attribute '%s' of prim '%s'"),
 			AttributeChange.TimeSamples.Num(),
 			*AttributeChange.PropertyName,
 			*Prim.GetPrimPath().GetString()
 		);
 
 		bool bSuccess = true;
-		for ( int32 Index = 0; Index < AttributeChange.TimeSamples.Num(); Index++ )
+		for (int32 Index = 0; Index < AttributeChange.TimeSamples.Num(); Index++)
 		{
 			UE::FVtValue Value;
-			if ( UnrealToUsd::ConvertValue( AttributeChange.TimeValues[ Index ], Value ) )
+			if (UnrealToUsd::ConvertValue(AttributeChange.TimeValues[Index], Value))
 			{
-				if ( !Attribute.Set( Value, AttributeChange.TimeSamples[ Index ] ) )
+				if (!Attribute.Set(Value, AttributeChange.TimeSamples[Index]))
 				{
-					UE_LOG(LogUsd, Warning, TEXT("Failed to apply value '%s' at timesample '%f' for attribute '%s' of prim '%s'"),
-						*UsdUtils::Stringify( Value ),
+					UE_LOG(
+						LogUsd,
+						Warning,
+						TEXT("Failed to apply value '%s' at timesample '%f' for attribute '%s' of prim '%s'"),
+						*UsdUtils::Stringify(Value),
 						AttributeChange.TimeSamples[Index],
 						*AttributeChange.PropertyName,
 						*Prim.GetPrimPath().GetString()
@@ -400,8 +411,11 @@ namespace UsdUtils
 			}
 			else
 			{
-				UE_LOG( LogUsd, Warning, TEXT( "Failed to convert value for timesample '%f' for attribute '%s' of prim '%s'" ),
-					AttributeChange.TimeSamples[ Index ],
+				UE_LOG(
+					LogUsd,
+					Warning,
+					TEXT("Failed to convert value for timesample '%f' for attribute '%s' of prim '%s'"),
+					AttributeChange.TimeSamples[Index],
 					*AttributeChange.PropertyName,
 					*Prim.GetPrimPath().GetString()
 				);
@@ -412,9 +426,17 @@ namespace UsdUtils
 		return bSuccess;
 	}
 
-	bool ApplyAttributeChange( const FString& PropertyName, const FString& Field, const FString& AttributeTypeName, bool bRemoveProperty, const FConvertedVtValue& Value, UE::FUsdPrim& Prim, TOptional<double> Time = {} )
+	bool ApplyAttributeChange(
+		const FString& PropertyName,
+		const FString& Field,
+		const FString& AttributeTypeName,
+		bool bRemoveProperty,
+		const FConvertedVtValue& Value,
+		UE::FUsdPrim& Prim,
+		TOptional<double> Time = {}
+	)
 	{
-		if ( !Prim )
+		if (!Prim)
 		{
 			return false;
 		}
@@ -422,26 +444,29 @@ namespace UsdUtils
 		bool bCreated = false;
 
 		UE::FUsdAttribute Attribute;
-		if ( PropertyName != TEXT( "kind" ) ) // Kind is prim metadata, not an attribute
+		if (PropertyName != TEXT("kind"))	 // Kind is prim metadata, not an attribute
 		{
-			if ( bRemoveProperty )
+			if (bRemoveProperty)
 			{
-				Attribute = Prim.GetAttribute( *PropertyName );
-				if ( !Attribute )
+				Attribute = Prim.GetAttribute(*PropertyName);
+				if (!Attribute)
 				{
 					return true;
 				}
 			}
 			else
 			{
-				bool bHadAttr = Prim.HasAttribute( *PropertyName );
-				Attribute = Prim.CreateAttribute( *PropertyName, *AttributeTypeName );
-				if ( !Attribute )
+				bool bHadAttr = Prim.HasAttribute(*PropertyName);
+				Attribute = Prim.CreateAttribute(*PropertyName, *AttributeTypeName);
+				if (!Attribute)
 				{
 					// We expect to fail to create an attribute if we have no typename here (e.g. undo remove property)
-					if ( AttributeTypeName.IsEmpty() )
+					if (AttributeTypeName.IsEmpty())
 					{
-						UE_LOG( LogUsd, Warning, TEXT( "Failed to create attribute '%s' with typename '%s' for prim '%s'" ),
+						UE_LOG(
+							LogUsd,
+							Warning,
+							TEXT("Failed to create attribute '%s' with typename '%s' for prim '%s'"),
 							*PropertyName,
 							*AttributeTypeName,
 							*Prim.GetPrimPath().GetString()
@@ -456,58 +481,66 @@ namespace UsdUtils
 		}
 
 		UE::FVtValue WrapperValue;
-		if ( !UnrealToUsd::ConvertValue( Value, WrapperValue ) )
+		if (!UnrealToUsd::ConvertValue(Value, WrapperValue))
 		{
-			UE_LOG( LogUsd, Warning, TEXT( "Failed to convert VtValue back to USD when applying it to attribute '%s' of prim '%s'" ),
+			UE_LOG(
+				LogUsd,
+				Warning,
+				TEXT("Failed to convert VtValue back to USD when applying it to attribute '%s' of prim '%s'"),
 				*PropertyName,
 				*Prim.GetPrimPath().GetString()
 			);
 			return false;
 		}
 
-		UE_LOG( LogUsd, Log, TEXT( "%s property '%s' (typename '%s'), field '%s' of prim '%s' with value '%s' at time '%s'" ),
-			bCreated ? TEXT("Creating") : WrapperValue.IsEmpty() ? bRemoveProperty ? TEXT("Removing") : TEXT( "Clearing" ) : TEXT( "Setting" ),
+		UE_LOG(
+			LogUsd,
+			Log,
+			TEXT("%s property '%s' (typename '%s'), field '%s' of prim '%s' with value '%s' at time '%s'"),
+			bCreated				 ? TEXT("Creating")
+			: WrapperValue.IsEmpty() ? bRemoveProperty ? TEXT("Removing") : TEXT("Clearing")
+									 : TEXT("Setting"),
 			*PropertyName,
 			*AttributeTypeName,
 			*Field,
 			*Prim.GetPrimPath().GetString(),
-			WrapperValue.IsEmpty() ? TEXT("<empty>") : *UsdUtils::Stringify( WrapperValue ),
-			Time.IsSet() ? *LexToString( Time.GetValue() ) : TEXT( "<unset>" )
+			WrapperValue.IsEmpty() ? TEXT("<empty>") : *UsdUtils::Stringify(WrapperValue),
+			Time.IsSet() ? *LexToString(Time.GetValue()) : TEXT("<unset>")
 		);
 
-		if ( PropertyName == TEXT( "kind" ) )
+		if (PropertyName == TEXT("kind"))
 		{
 #if USE_USD_SDK
-			if ( Value.Entries.Num() == 1 && Value.Entries[ 0 ].Num() == 1 )
+			if (Value.Entries.Num() == 1 && Value.Entries[0].Num() == 1)
 			{
-				if ( const FString* KindString = Value.Entries[ 0 ][ 0 ].TryGet<FString>() )
+				if (const FString* KindString = Value.Entries[0][0].TryGet<FString>())
 				{
-					if ( !IUsdPrim::SetKind( Prim, UnrealToUsd::ConvertToken( **KindString ).Get() ) )
+					if (!IUsdPrim::SetKind(Prim, UnrealToUsd::ConvertToken(**KindString).Get()))
 					{
-						UE_LOG( LogUsd, Warning, TEXT( "Failed to set Kind '%s' for prim '%s'" ), **KindString, *Prim.GetPrimPath().GetString() );
+						UE_LOG(LogUsd, Warning, TEXT("Failed to set Kind '%s' for prim '%s'"), **KindString, *Prim.GetPrimPath().GetString());
 					}
 				}
 			}
 			else
 			{
-				if ( !IUsdPrim::ClearKind( Prim ) )
+				if (!IUsdPrim::ClearKind(Prim))
 				{
-					UE_LOG( LogUsd, Warning, TEXT( "Failed to clear Kind for prim '%s'" ), *Prim.GetPrimPath().GetString() );
+					UE_LOG(LogUsd, Warning, TEXT("Failed to clear Kind for prim '%s'"), *Prim.GetPrimPath().GetString());
 				}
 			}
-#endif // USE_USD_SDK
+#endif	  // USE_USD_SDK
 		}
-		else if ( Field == TEXT( "default" ) )
+		else if (Field == TEXT("default"))
 		{
-			if ( WrapperValue.IsEmpty() )
+			if (WrapperValue.IsEmpty())
 			{
-				if ( bRemoveProperty )
+				if (bRemoveProperty)
 				{
-					Prim.RemoveProperty( *PropertyName );
+					Prim.RemoveProperty(*PropertyName);
 				}
-				if ( Time.IsSet() )
+				if (Time.IsSet())
 				{
-					Attribute.ClearAtTime( Time.GetValue() );
+					Attribute.ClearAtTime(Time.GetValue());
 				}
 				else
 				{
@@ -516,66 +549,66 @@ namespace UsdUtils
 			}
 			else
 			{
-				Attribute.Set( WrapperValue, Time );
+				Attribute.Set(WrapperValue, Time);
 			}
 		}
-		else // variability, colorSpace, etc.
+		else	// variability, colorSpace, etc.
 		{
-			if ( WrapperValue.IsEmpty() )
+			if (WrapperValue.IsEmpty())
 			{
 				Attribute.Clear();
 			}
 			else
 			{
-				Attribute.SetMetadata( *Field, WrapperValue );
+				Attribute.SetMetadata(*Field, WrapperValue);
 			}
 		}
 
 		return true;
 	}
 
-	bool ApplyStageMetadataChange( const FString& PropertyName, const FConvertedVtValue& Value, UE::FUsdStage& Stage )
+	bool ApplyStageMetadataChange(const FString& PropertyName, const FConvertedVtValue& Value, UE::FUsdStage& Stage)
 	{
-		if ( !Stage || PropertyName.IsEmpty() )
+		if (!Stage || PropertyName.IsEmpty())
 		{
 			return false;
 		}
 
 		UE::FVtValue WrapperValue;
-		if ( !UnrealToUsd::ConvertValue( Value, WrapperValue ) )
+		if (!UnrealToUsd::ConvertValue(Value, WrapperValue))
 		{
-			UE_LOG( LogUsd, Warning, TEXT( "Failed to convert VtValue back to USD when applying it to stage metadata field '%s'" ),
-				*PropertyName
-			);
+			UE_LOG(LogUsd, Warning, TEXT("Failed to convert VtValue back to USD when applying it to stage metadata field '%s'"), *PropertyName);
 			return false;
 		}
 
-		UE_LOG( LogUsd, Log, TEXT( "Setting stage metadata '%s', with value '%s'" ),
-			*PropertyName,
-			*UsdUtils::Stringify( WrapperValue )
-		);
+		UE_LOG(LogUsd, Log, TEXT("Setting stage metadata '%s', with value '%s'"), *PropertyName, *UsdUtils::Stringify(WrapperValue));
 
 		UE::FSdfLayer OldEditTarget = Stage.GetEditTarget();
-		Stage.SetEditTarget( Stage.GetRootLayer() ); // Stage metadata always needs to be set at the root layer
+		Stage.SetEditTarget(Stage.GetRootLayer());	  // Stage metadata always needs to be set at the root layer
 		{
-			if ( WrapperValue.IsEmpty() )
+			if (WrapperValue.IsEmpty())
 			{
-				Stage.ClearMetadata( *PropertyName );
+				Stage.ClearMetadata(*PropertyName);
 			}
 			else
 			{
-				Stage.SetMetadata( *PropertyName, WrapperValue );
+				Stage.SetMetadata(*PropertyName, WrapperValue);
 			}
 		}
-		Stage.SetEditTarget( OldEditTarget );
+		Stage.SetEditTarget(OldEditTarget);
 
 		return true;
 	}
 
 	/** Applies the field value pairs to all prims on the stage, and returns a list of prim paths for modified prims */
-	TArray<FString> ApplyFieldMapToStage( const FTransactorEditStorage& EditStorage, EApplicationDirection Direction, UE::FUsdStage& Stage, double Time )
+	TArray<FString> ApplyFieldMapToStage(
+		const FTransactorEditStorage& EditStorage,
+		EApplicationDirection Direction,
+		UE::FUsdStage& Stage,
+		double Time
+	)
 	{
-		if ( !Stage )
+		if (!Stage)
 		{
 			return {};
 		}
@@ -585,7 +618,7 @@ namespace UsdUtils
 		int32 Start = 0;
 		int32 End = 0;
 		int32 Incr = 0;
-		if ( Direction == EApplicationDirection::Forward )
+		if (Direction == EApplicationDirection::Forward)
 		{
 			Start = 0;
 			End = EditStorage.Num();
@@ -601,39 +634,44 @@ namespace UsdUtils
 		UE::FSdfLayer OldEditTarget = Stage.GetEditTarget();
 		FString LastEditTargetIdentifier;
 
-		for ( int32 EditIndex = Start; EditIndex != End; EditIndex += Incr )
+		for (int32 EditIndex = Start; EditIndex != End; EditIndex += Incr)
 		{
-			const FTransactorRecordedEdits& Edits = EditStorage[ EditIndex ];
+			const FTransactorRecordedEdits& Edits = EditStorage[EditIndex];
 
-			if ( Edits.EditTargetIdentifier != LastEditTargetIdentifier )
+			if (Edits.EditTargetIdentifier != LastEditTargetIdentifier)
 			{
 #if USE_USD_SDK
-				UE::FSdfLayer EditTarget = UsdUtils::FindLayerForIdentifier( *Edits.EditTargetIdentifier, Stage );
-				if ( !EditTarget )
+				UE::FSdfLayer EditTarget = UsdUtils::FindLayerForIdentifier(*Edits.EditTargetIdentifier, Stage);
+				if (!EditTarget)
 				{
-					UE_LOG(LogUsd, Warning, TEXT("Ignoring application of recorded USD stage changes as the edit target with identifier '%s' cannot be found or opened"), *Edits.EditTargetIdentifier )
+					UE_LOG(
+						LogUsd,
+						Warning,
+						TEXT("Ignoring application of recorded USD stage changes as the edit target with identifier '%s' cannot be found or opened"),
+						*Edits.EditTargetIdentifier
+					)
 					continue;
 				}
 
-				Stage.SetEditTarget( EditTarget );
+				Stage.SetEditTarget(EditTarget);
 				LastEditTargetIdentifier = Edits.EditTargetIdentifier;
-#endif // USE_USD_SDK
+#endif	  // USE_USD_SDK
 			}
 
-			for ( const FTransactorRecordedEdit& Edit : Edits.Edits )
+			for (const FTransactorRecordedEdit& Edit : Edits.Edits)
 			{
-				if ( Edit.ObjectChanges.Num() == 0 )
+				if (Edit.ObjectChanges.Num() == 0)
 				{
 					continue;
 				}
 
-				PrimsChanged.Add( Edit.ObjectPath );
+				PrimsChanged.Add(Edit.ObjectPath);
 
-				if ( Edit.ObjectPath == TEXT("/") )
+				if (Edit.ObjectPath == TEXT("/"))
 				{
-					for ( const FTransactorObjectChange& PrimChange : Edit.ObjectChanges )
+					for (const FTransactorObjectChange& PrimChange : Edit.ObjectChanges)
 					{
-						for ( const FTransactorAttributeChange& AttributeChange : PrimChange.AttributeChanges )
+						for (const FTransactorAttributeChange& AttributeChange : PrimChange.AttributeChanges)
 						{
 							ApplyStageMetadataChange(
 								AttributeChange.PropertyName,
@@ -645,40 +683,40 @@ namespace UsdUtils
 				}
 				else
 				{
-					for ( const FTransactorObjectChange& PrimChange : Edit.ObjectChanges )
+					for (const FTransactorObjectChange& PrimChange : Edit.ObjectChanges)
 					{
-						UE::FSdfPath PrimPath = UE::FSdfPath{ *Edit.ObjectPath };
-						if ( ApplyPrimChange( PrimPath, PrimChange, Stage, Direction ) )
+						UE::FSdfPath PrimPath = UE::FSdfPath{*Edit.ObjectPath};
+						if (ApplyPrimChange(PrimPath, PrimChange, Stage, Direction))
 						{
 							// If we managed to apply a prim change, we know there aren't any other attribute changes in the same ObjectChange
 							continue;
 						}
 
-						UE::FUsdPrim Prim = Stage.GetPrimAtPath( PrimPath );
-						if ( !Prim )
+						UE::FUsdPrim Prim = Stage.GetPrimAtPath(PrimPath);
+						if (!Prim)
 						{
 							continue;
 						}
 
 						// Whether we should remove a property instead of clearing the opinion, when asked to apply an empty value
 						bool bShouldRemove = Direction == EApplicationDirection::Forward
-							? PrimChange.Flags.bDidRemoveProperty || PrimChange.Flags.bDidRemovePropertyWithOnlyRequiredFields
-							: PrimChange.Flags.bDidAddProperty || PrimChange.Flags.bDidAddPropertyWithOnlyRequiredFields;
+												 ? PrimChange.Flags.bDidRemoveProperty || PrimChange.Flags.bDidRemovePropertyWithOnlyRequiredFields
+												 : PrimChange.Flags.bDidAddProperty || PrimChange.Flags.bDidAddPropertyWithOnlyRequiredFields;
 
 						// Can't block more than this, as Defining prims (from within ApplyPrimChange) needs to trigger its notices immediately,
 						// and our changes/edits may depend on previous changes/edits triggering
 						UE::FSdfChangeBlock ChangeBlock;
 
-						if ( PrimChange.Flags.bDidChangeAttributeTimeSamples )
+						if (PrimChange.Flags.bDidChangeAttributeTimeSamples)
 						{
-							for ( const FTransactorAttributeChange& AttributeChange : PrimChange.AttributeChanges )
+							for (const FTransactorAttributeChange& AttributeChange : PrimChange.AttributeChanges)
 							{
-								ApplyAttributeTimeSamples( AttributeChange, Prim );
+								ApplyAttributeTimeSamples(AttributeChange, Prim);
 							}
 						}
 						else
 						{
-							for ( const FTransactorAttributeChange& AttributeChange : PrimChange.AttributeChanges )
+							for (const FTransactorAttributeChange& AttributeChange : PrimChange.AttributeChanges)
 							{
 								ApplyAttributeChange(
 									AttributeChange.PropertyName,
@@ -695,7 +733,7 @@ namespace UsdUtils
 			}
 		}
 
-		Stage.SetEditTarget( OldEditTarget );
+		Stage.SetEditTarget(OldEditTarget);
 
 		return PrimsChanged;
 	}
@@ -711,33 +749,38 @@ namespace UsdUtils
 		FUsdTransactorImpl();
 		~FUsdTransactorImpl();
 
-		void Update( FTransactorRecordedEdits&& NewEdits );
-		void Serialize( FArchive& Ar );
-		void PreEditUndo( AUsdStageActor* StageActor );
-		void PostEditUndo( AUsdStageActor* StageActor );
+		void Update(FTransactorRecordedEdits&& NewEdits);
+		void Serialize(FArchive& Ar);
+		void PreEditUndo(AUsdStageActor* StageActor);
+		void PostEditUndo(AUsdStageActor* StageActor);
 
 		// Returns whether the transaction buffer is currently in the middle of an Undo operation.
-		// WARNING: This approach is only accurate if we're checking from within PreEditUndo/PostEditUndo/PostTransacted/Serialize (which we always do in this file)
+		// WARNING: This approach is only accurate if we're checking from within PreEditUndo/PostEditUndo/PostTransacted/Serialize (which we always do
+		// in this file)
 		bool IsTransactionUndoing();
 
 		// Returns whether the transaction buffer is currently in the middle of a Redo operation. Returns false when we're applying
 		// a ConcertSync transaction, even though concert sync sort of works by applying transactions via Redo
-		// WARNING: This approach is only accurate if we're checking from within PreEditUndo/PostEditUndo/PostTransacted/Serialize (which we always do in this file)
+		// WARNING: This approach is only accurate if we're checking from within PreEditUndo/PostEditUndo/PostTransacted/Serialize (which we always do
+		// in this file)
 		bool IsTransactionRedoing();
 
 		// Returns whether ConcertSync (multi-user) is currently applying a transaction received from the network
-		bool IsApplyingConcertSyncTransaction() { return bApplyingConcertSync; };
+		bool IsApplyingConcertSyncTransaction()
+		{
+			return bApplyingConcertSync;
+		};
 
 	private:
 		// This is called after *any* undo/redo transaction is finalized, so our LastFinalizedUndoCount is kept updated
-		void HandleTransactionStateChanged( const FTransactionContext& InTransactionContext, const ETransactionStateEventType InTransactionState );
+		void HandleTransactionStateChanged(const FTransactionContext& InTransactionContext, const ETransactionStateEventType InTransactionState);
 		int32 LastFinalizedUndoCount = 0;
 
 		// We use this to detect when a ConcertSync transaction is starting, as it has a particular title
-		void HandleBeforeOnRedoUndo( const FTransactionContext& TransactionContext );
+		void HandleBeforeOnRedoUndo(const FTransactionContext& TransactionContext);
 
 		// We use this to detect when a ConcertSync transaction has ended, as it has a particular title
-		void HandleOnRedo( const FTransactionContext& TransactionContext, bool bSucceeded );
+		void HandleOnRedo(const FTransactionContext& TransactionContext, bool bSucceeded);
 
 	private:
 		// Main data storage container
@@ -752,11 +795,13 @@ namespace UsdUtils
 		// When ClientA undoes a change, it handles it's own undo changes from its PreEditUndo, but its final state after the undo transaction
 		// is complete will have the *previous* OldValues/NewValues. This final state is what is sent over the network. ClientB that receives this
 		// can't use these previous OldValues/NewValues to undo the change that ClientA just undone: It needs something else, which this
-		// member provides. When ClientA starts to undo, it will stash it's *current* OldValues in here, and make sure they are visible when serialized
-		// by ConcertSync. ClientB will receive these, and when available will apply those to the scene instead, undoing the same changes that ClientA undone.
-		TOptional< FTransactorEditStorage > ReceivedValuesBeforeUndo;
+		// member provides. When ClientA starts to undo, it will stash it's *current* OldValues in here, and make sure they are visible when
+		// serialized by ConcertSync. ClientB will receive these, and when available will apply those to the scene instead, undoing the same changes
+		// that ClientA undone.
+		TOptional<FTransactorEditStorage> ReceivedValuesBeforeUndo;
 
-		// During the same transaction we continuously append the received change info into the same storage. When the transaction changes, we clear it
+		// During the same transaction we continuously append the received change info into the same storage. When the transaction changes, we clear
+		// it
 		FGuid LastTransactionId;
 
 		bool bApplyingConcertSync = false;
@@ -765,58 +810,59 @@ namespace UsdUtils
 	FUsdTransactorImpl::FUsdTransactorImpl()
 	{
 #if WITH_EDITOR
-		if ( GEditor )
+		if (GEditor)
 		{
-			if ( UTransBuffer* Transactor = Cast<UTransBuffer>( GEditor->Trans ) )
+			if (UTransBuffer* Transactor = Cast<UTransBuffer>(GEditor->Trans))
 			{
-				Transactor->OnTransactionStateChanged().AddRaw( this, &FUsdTransactorImpl::HandleTransactionStateChanged );
-				Transactor->OnBeforeRedoUndo().AddRaw( this, &FUsdTransactorImpl::HandleBeforeOnRedoUndo );
-				Transactor->OnRedo().AddRaw( this, &FUsdTransactorImpl::HandleOnRedo );
+				Transactor->OnTransactionStateChanged().AddRaw(this, &FUsdTransactorImpl::HandleTransactionStateChanged);
+				Transactor->OnBeforeRedoUndo().AddRaw(this, &FUsdTransactorImpl::HandleBeforeOnRedoUndo);
+				Transactor->OnRedo().AddRaw(this, &FUsdTransactorImpl::HandleOnRedo);
 			}
 		}
-#endif // WITH_EDITOR
+#endif	  // WITH_EDITOR
 	}
 
 	FUsdTransactorImpl::~FUsdTransactorImpl()
 	{
 #if WITH_EDITOR
-		if ( GEditor )
+		if (GEditor)
 		{
-			if ( UTransBuffer* Transactor = Cast<UTransBuffer>( GEditor->Trans ) )
+			if (UTransBuffer* Transactor = Cast<UTransBuffer>(GEditor->Trans))
 			{
-				Transactor->OnTransactionStateChanged().RemoveAll( this );
-				Transactor->OnBeforeRedoUndo().RemoveAll( this );
-				Transactor->OnRedo().RemoveAll( this );
+				Transactor->OnTransactionStateChanged().RemoveAll(this);
+				Transactor->OnBeforeRedoUndo().RemoveAll(this);
+				Transactor->OnRedo().RemoveAll(this);
 			}
 		}
-#endif // WITH_EDITOR
+#endif	  // WITH_EDITOR
 	}
 
-	void FUsdTransactorImpl::Update( FTransactorRecordedEdits&& NewEdits )
+	void FUsdTransactorImpl::Update(FTransactorRecordedEdits&& NewEdits)
 	{
-		if ( !GUndo )
+		if (!GUndo)
 		{
 			return;
 		}
 
 		// New transaction -> Start a new storage
 		FTransactionContext Context = GUndo->GetContext();
-		if ( Context.TransactionId != LastTransactionId )
+		if (Context.TransactionId != LastTransactionId)
 		{
 			LastTransactionId = Context.TransactionId;
 			Values.Reset();
 		}
 
-		Values.Add( NewEdits );
+		Values.Add(NewEdits);
 	}
 
-	void FUsdTransactorImpl::Serialize( FArchive& Ar )
+	void FUsdTransactorImpl::Serialize(FArchive& Ar)
 	{
 		Ar << Values;
 
-		// If we have some ReceivedValuesBeforeUndo and the undo system is trying to overwrite it with its old version to apply the undo, keep our values instead!
-		// We need this data to be with us whenever the ConcertSync serializes us to send it over the network during an undo, which happens shortly after this
-		if ( Ar.IsTransacting() && Ar.IsLoading() && IsTransactionUndoing() && ReceivedValuesBeforeUndo.IsSet() )
+		// If we have some ReceivedValuesBeforeUndo and the undo system is trying to overwrite it with its old version to apply the undo, keep our
+		// values instead! We need this data to be with us whenever the ConcertSync serializes us to send it over the network during an undo, which
+		// happens shortly after this
+		if (Ar.IsTransacting() && Ar.IsLoading() && IsTransactionUndoing() && ReceivedValuesBeforeUndo.IsSet())
 		{
 			TOptional<FTransactorEditStorage> Dummy = {};
 			Ar << Dummy;
@@ -827,33 +873,30 @@ namespace UsdUtils
 		}
 	}
 
-	void FUsdTransactorImpl::PreEditUndo( AUsdStageActor* StageActor )
+	void FUsdTransactorImpl::PreEditUndo(AUsdStageActor* StageActor)
 	{
-		if ( !StageActor )
+		if (!StageActor)
 		{
 			return;
 		}
 
-		if ( IsTransactionUndoing() )
+		if (IsTransactionUndoing())
 		{
-			// We can't respond to notices from the attribute that we'll set. Whatever changes setting the attribute causes in UE actors/components/assets
-			// will already be accounted for by those actors/components/assets undoing/redoing by themselves via the UE transaction buffer.
-			FScopedBlockNotices BlockNotices( StageActor->GetUsdListener() );
+			// We can't respond to notices from the attribute that we'll set. Whatever changes setting the attribute causes in UE
+			// actors/components/assets will already be accounted for by those actors/components/assets undoing/redoing by themselves via the UE
+			// transaction buffer.
+			FScopedBlockNotices BlockNotices(StageActor->GetUsdListener());
 
 			UE::FUsdStage& Stage = StageActor->GetOrOpenUsdStage();
 
-			TArray<FString> PrimsChanged = UsdUtils::ApplyFieldMapToStage(
-				Values,
-				UsdUtils::EApplicationDirection::Reverse,
-				Stage,
-				StageActor->GetTime()
-			);
+			TArray<FString>
+				PrimsChanged = UsdUtils::ApplyFieldMapToStage(Values, UsdUtils::EApplicationDirection::Reverse, Stage, StageActor->GetTime());
 
-			if ( PrimsChanged.Num() > 0 )
+			if (PrimsChanged.Num() > 0)
 			{
-				for ( const FString& Prim : PrimsChanged )
+				for (const FString& Prim : PrimsChanged)
 				{
-					StageActor->OnPrimChanged.Broadcast( Prim, false );
+					StageActor->OnPrimChanged.Broadcast(Prim, false);
 				}
 			}
 
@@ -865,7 +908,7 @@ namespace UsdUtils
 			ReceivedValuesBeforeUndo.Reset();
 
 			// ConcertSync calls PreEditUndo, then updates our data with the received data, then calls PostEditUndo
-			if ( IsApplyingConcertSyncTransaction() )
+			if (IsApplyingConcertSyncTransaction())
 			{
 				// Make sure that our own Values survive when overwritten by values that we will receive from ConcertSync.
 				// We'll restore this to our values once the ConcertSync action has finished applying
@@ -874,30 +917,30 @@ namespace UsdUtils
 		}
 	}
 
-	void FUsdTransactorImpl::PostEditUndo( AUsdStageActor* StageActor )
+	void FUsdTransactorImpl::PostEditUndo(AUsdStageActor* StageActor)
 	{
 		const bool bIsRedoing = IsTransactionRedoing();
 		const bool bIsApplyingConcertSync = IsApplyingConcertSyncTransaction();
 
-		if ( StageActor && ( bIsRedoing || bIsApplyingConcertSync ) )
+		if (StageActor && (bIsRedoing || bIsApplyingConcertSync))
 		{
 			// If we're just redoing it's a bit of a waste to let the AUsdStageActor respond to notices from the fields that we'll set,
 			// because any relevant changes caused to the level/assets would be redone by themselves if the actors/assets are also in the transaction
 			// buffer. If we're receiving a ConcertSync transaction, however, we do want to respond to notices because transient actors/assets
 			// aren't tracked by ConcertSync
 			TOptional<FScopedBlockNotices> BlockNotices;
-			if ( bIsRedoing )
+			if (bIsRedoing)
 			{
-				BlockNotices.Emplace( StageActor->GetUsdListener() );
+				BlockNotices.Emplace(StageActor->GetUsdListener());
 			}
 
 			UE::FUsdStage& Stage = StageActor->GetOrOpenUsdStage();
 
 			TArray<FString> PrimsChanged;
-			if ( bIsApplyingConcertSync && ReceivedValuesBeforeUndo.IsSet() )
+			if (bIsApplyingConcertSync && ReceivedValuesBeforeUndo.IsSet())
 			{
-				// If we're applying a received ConcertSync transaction that actually is an undo on the source client then we want to use it's ReceivedValuesBeforeUndo
-				// to replicate the same undo that they did
+				// If we're applying a received ConcertSync transaction that actually is an undo on the source client then we want to use it's
+				// ReceivedValuesBeforeUndo to replicate the same undo that they did
 				PrimsChanged = UsdUtils::ApplyFieldMapToStage(
 					ReceivedValuesBeforeUndo.GetValue(),
 					UsdUtils::EApplicationDirection::Reverse,
@@ -908,27 +951,23 @@ namespace UsdUtils
 			else
 			{
 				// Just a common Redo operation or any other type of ConcertSync transaction, so just apply the new values
-				PrimsChanged = UsdUtils::ApplyFieldMapToStage(
-					Values,
-					UsdUtils::EApplicationDirection::Forward,
-					Stage,
-					StageActor->GetTime()
-				);
+				PrimsChanged = UsdUtils::ApplyFieldMapToStage(Values, UsdUtils::EApplicationDirection::Forward, Stage, StageActor->GetTime());
 			}
 
-			// If we're redoing or applying ConcertSync we don't want to end up with these values when the transaction finalizes as it could be replicated to other clients
+			// If we're redoing or applying ConcertSync we don't want to end up with these values when the transaction finalizes as it could be
+			// replicated to other clients
 			ReceivedValuesBeforeUndo.Reset();
 
-			if ( PrimsChanged.Num() > 0 )
+			if (PrimsChanged.Num() > 0)
 			{
-				for ( const FString& Prim : PrimsChanged )
+				for (const FString& Prim : PrimsChanged)
 				{
-					StageActor->OnPrimChanged.Broadcast( Prim, false );
+					StageActor->OnPrimChanged.Broadcast(Prim, false);
 				}
 			}
 		}
 
-		if ( bIsApplyingConcertSync )
+		if (bIsApplyingConcertSync)
 		{
 			// If we're finishing applying a ConcertSync transaction, revert our Values to the state that they were before we received
 			// the ConcertSync transaction. This is important so that if we undo now, we undo the last change that *we* made
@@ -939,15 +978,15 @@ namespace UsdUtils
 	bool FUsdTransactorImpl::IsTransactionUndoing()
 	{
 #if WITH_EDITOR
-		if ( UTransBuffer* Transactor = Cast<UTransBuffer>( GEditor->Trans ) )
+		if (UTransBuffer* Transactor = Cast<UTransBuffer>(GEditor->Trans))
 		{
 			// We moved away from the end of the transaction buffer -> We're undoing
-			if ( GIsTransacting && Transactor->UndoCount > LastFinalizedUndoCount )
+			if (GIsTransacting && Transactor->UndoCount > LastFinalizedUndoCount)
 			{
 				return true;
 			}
 		}
-#endif // WITH_EDITOR
+#endif	  // WITH_EDITOR
 
 		return false;
 	}
@@ -955,80 +994,84 @@ namespace UsdUtils
 	bool FUsdTransactorImpl::IsTransactionRedoing()
 	{
 #if WITH_EDITOR
-		if ( UTransBuffer* Transactor = Cast<UTransBuffer>( GEditor->Trans ) )
+		if (UTransBuffer* Transactor = Cast<UTransBuffer>(GEditor->Trans))
 		{
 			// We moved towards the end of the transaction buffer -> We're redoing
-			if ( GIsTransacting && Transactor->UndoCount < LastFinalizedUndoCount )
+			if (GIsTransacting && Transactor->UndoCount < LastFinalizedUndoCount)
 			{
 				return true;
 			}
 		}
-#endif // WITH_EDITOR
+#endif	  // WITH_EDITOR
 
 		return false;
 	}
 
-	void FUsdTransactorImpl::HandleTransactionStateChanged( const FTransactionContext& InTransactionContext, const ETransactionStateEventType InTransactionState )
+	void FUsdTransactorImpl::HandleTransactionStateChanged(
+		const FTransactionContext& InTransactionContext,
+		const ETransactionStateEventType InTransactionState
+	)
 	{
 #if WITH_EDITOR
-		if ( InTransactionState == ETransactionStateEventType::UndoRedoFinalized || InTransactionState == ETransactionStateEventType::TransactionFinalized )
+		if (InTransactionState == ETransactionStateEventType::UndoRedoFinalized
+			|| InTransactionState == ETransactionStateEventType::TransactionFinalized)
 		{
-			if ( UTransBuffer* Transactor = Cast<UTransBuffer>( GEditor->Trans ) )
+			if (UTransBuffer* Transactor = Cast<UTransBuffer>(GEditor->Trans))
 			{
-				// Recording UndoCount works because UTransBuffer::Undo preemptively updates it *before* calling any object function, like PreEditUndo/PostEditUndo, so from
-				// within PreEditUndo/PostEditUndo we will always have a delta from this value to the value that is recorded after any transaction was finalized,
-				// which we record right here
+				// Recording UndoCount works because UTransBuffer::Undo preemptively updates it *before* calling any object function, like
+				// PreEditUndo/PostEditUndo, so from within PreEditUndo/PostEditUndo we will always have a delta from this value to the value that is
+				// recorded after any transaction was finalized, which we record right here
 				LastFinalizedUndoCount = Transactor->UndoCount;
 			}
 		}
-#endif // WITH_EDITOR
+#endif	  // WITH_EDITOR
 	}
 
 	// Must be the same as the title used within ConcertClientTransactionBridgeUtil::ProcessTransactionEvent
-	const FText ConcertSyncTransactionTitle = LOCTEXT( "ConcertTransactionEvent", "Concert Transaction Event" );
+	const FText ConcertSyncTransactionTitle = LOCTEXT("ConcertTransactionEvent", "Concert Transaction Event");
 
-	void FUsdTransactorImpl::HandleBeforeOnRedoUndo( const FTransactionContext& TransactionContext )
+	void FUsdTransactorImpl::HandleBeforeOnRedoUndo(const FTransactionContext& TransactionContext)
 	{
-		if ( TransactionContext.Title.EqualTo( ConcertSyncTransactionTitle ) )
+		if (TransactionContext.Title.EqualTo(ConcertSyncTransactionTitle))
 		{
 			bApplyingConcertSync = true;
 		}
 	}
 
-	void FUsdTransactorImpl::HandleOnRedo( const FTransactionContext& TransactionContext, bool bSucceeded )
+	void FUsdTransactorImpl::HandleOnRedo(const FTransactionContext& TransactionContext, bool bSucceeded)
 	{
-		if ( bApplyingConcertSync && TransactionContext.Title.EqualTo( ConcertSyncTransactionTitle ) )
+		if (bApplyingConcertSync && TransactionContext.Title.EqualTo(ConcertSyncTransactionTitle))
 		{
 			bApplyingConcertSync = false;
 		}
 	}
-}
+}	 // namespace UsdUtils
 
 // Boilerplate to support Pimpl in an UObject
-UUsdTransactor::UUsdTransactor( FVTableHelper& Helper )
-	: Super( Helper )
+UUsdTransactor::UUsdTransactor(FVTableHelper& Helper)
+	: Super(Helper)
 {
 }
 UUsdTransactor::UUsdTransactor() = default;
 UUsdTransactor::~UUsdTransactor() = default;
 
-void UUsdTransactor::Initialize( AUsdStageActor* InStageActor )
+void UUsdTransactor::Initialize(AUsdStageActor* InStageActor)
 {
 	StageActor = InStageActor;
 
 #if USE_USD_SDK
-	if ( !IsTemplate() )
+	if (!IsTemplate())
 	{
 		Impl = MakeUnique<UsdUtils::FUsdTransactorImpl>();
 	}
-#endif // USE_USD_SDK
+#endif	  // USE_USD_SDK
 }
 
-void UUsdTransactor::Update( const UsdUtils::FObjectChangesByPath& NewInfoChanges, const UsdUtils::FObjectChangesByPath& NewResyncChanges )
+void UUsdTransactor::Update(const UsdUtils::FObjectChangesByPath& NewInfoChanges, const UsdUtils::FObjectChangesByPath& NewResyncChanges)
 {
 	// We always send notices even when we're undoing/redoing changes (so that multi-user can broadcast them).
 	// Make sure that we only ever update our OldValues/NewValues when we receive *new* updates though
-	if ( Impl.IsValid() && ( Impl->IsTransactionUndoing() || Impl->IsTransactionRedoing() || Impl->IsApplyingConcertSyncTransaction() ) )
+	if (Impl.IsValid() && (Impl->IsTransactionUndoing() || Impl->IsTransactionRedoing() || Impl->IsApplyingConcertSyncTransaction()))
 	{
 		return;
 	}
@@ -1036,7 +1079,7 @@ void UUsdTransactor::Update( const UsdUtils::FObjectChangesByPath& NewInfoChange
 	// In case we close a stage in the same transaction where the actor is destroyed - our UE::FUsdStage could turn invalid at any point otherwise
 	// Not much else we can do as this will get to us before the StageActor's destructor/Destroyed are called
 	AUsdStageActor* StageActorPtr = StageActor.Get();
-	if ( !StageActorPtr || StageActorPtr->IsActorBeingDestroyed() )
+	if (!StageActorPtr || StageActorPtr->IsActorBeingDestroyed())
 	{
 		return;
 	}
@@ -1044,7 +1087,7 @@ void UUsdTransactor::Update( const UsdUtils::FObjectChangesByPath& NewInfoChange
 	Modify();
 
 	const UE::FUsdStage& Stage = StageActorPtr->GetOrOpenUsdStage();
-	if ( !Stage )
+	if (!Stage)
 	{
 		return;
 	}
@@ -1054,25 +1097,25 @@ void UUsdTransactor::Update( const UsdUtils::FObjectChangesByPath& NewInfoChange
 	UsdUtils::FTransactorRecordedEdits NewEdits;
 	NewEdits.EditTargetIdentifier = EditTarget.GetIdentifier();
 
-	UsdUtils::ConvertFieldValueMap( NewInfoChanges, Stage, NewEdits );
-	UsdUtils::ConvertFieldValueMap( NewResyncChanges, Stage, NewEdits );
+	UsdUtils::ConvertFieldValueMap(NewInfoChanges, Stage, NewEdits);
+	UsdUtils::ConvertFieldValueMap(NewResyncChanges, Stage, NewEdits);
 
-	Impl->Update( MoveTemp(NewEdits) );
+	Impl->Update(MoveTemp(NewEdits));
 }
 
-void UUsdTransactor::Serialize( FArchive& Ar )
+void UUsdTransactor::Serialize(FArchive& Ar)
 {
-	Super::Serialize( Ar );
+	Super::Serialize(Ar);
 
-	if ( Impl.IsValid() )
+	if (Impl.IsValid())
 	{
-		Impl->Serialize( Ar );
+		Impl->Serialize(Ar);
 	}
 	else
 	{
 		// In case we somehow serialize before we receive a valid Impl, and then later do receive one
 		UsdUtils::FUsdTransactorImpl Dummy;
-		Dummy.Serialize( Ar );
+		Dummy.Serialize(Ar);
 	}
 }
 
@@ -1080,9 +1123,9 @@ void UUsdTransactor::Serialize( FArchive& Ar )
 
 void UUsdTransactor::PreEditUndo()
 {
-	if( Impl.IsValid() )
+	if (Impl.IsValid())
 	{
-		Impl->PreEditUndo( StageActor.Get() );
+		Impl->PreEditUndo(StageActor.Get());
 	}
 
 	Super::PreEditUndo();
@@ -1090,13 +1133,13 @@ void UUsdTransactor::PreEditUndo()
 
 void UUsdTransactor::PostEditUndo()
 {
-	if ( Impl.IsValid() )
+	if (Impl.IsValid())
 	{
-		Impl->PostEditUndo( StageActor.Get() );
+		Impl->PostEditUndo(StageActor.Get());
 	}
 
 	Super::PostEditUndo();
 }
-#endif // WITH_EDITOR
+#endif	  // WITH_EDITOR
 
 #undef LOCTEXT_NAMESPACE
diff --git a/Engine/Plugins/Importers/USDImporter/Source/USDStage/Public/USDLevelSequenceHelper.h b/Engine/Plugins/Importers/USDImporter/Source/USDStage/Public/USDLevelSequenceHelper.h
--- a/Engine/Plugins/Importers/USDImporter/Source/USDStage/Public/USDLevelSequenceHelper.h
+++ b/Engine/Plugins/Importers/USDImporter/Source/USDStage/Public/USDLevelSequenceHelper.h
@@ -43,7 +43,8 @@ public:
 	/** Sets the asset cache to use when fetching assets and asset info required for the level sequence animation, like UAnimSequences */
 	void SetInfoCache(TSharedPtr<FUsdInfoCache> InInfoCache);
 
-	/** Sets the BBoxCache to use when importing bound animations for prims. Needed for importing, where we don't have a stage actor to take the BBoxCache from */
+	/** Sets the BBoxCache to use when importing bound animations for prims. Needed for importing, where we don't have a stage actor to take the
+	 * BBoxCache from */
 	void SetBBoxCache(TSharedPtr<UE::FUsdGeomBBoxCache> InBBoxCache);
 
 	/* Returns true if we have at least one possessable or a reference to a subsequence */
@@ -85,7 +86,7 @@ public:
 	/** Removes any track associated with this prim */
 	void RemovePrim(const UUsdPrimTwin& PrimTwin);
 
-	void UpdateControlRigTracks( UUsdPrimTwin& PrimTwin );
+	void UpdateControlRigTracks(UUsdPrimTwin& PrimTwin);
 
 	/** Blocks updating the level sequences & tracks from object changes. */
 	void StartMonitoringChanges();
@@ -93,9 +94,9 @@ public:
 	void BlockMonitoringChangesForThisTransaction();
 
 	ULevelSequence* GetMainLevelSequence() const;
-	TArray< ULevelSequence* > GetSubSequences() const;
+	TArray<ULevelSequence*> GetSubSequences() const;
 
-	DECLARE_EVENT_OneParam( FUsdLevelSequenceHelper, FOnSkelAnimationBaked, const FString& /*SkeletonPrimPath*/ );
+	DECLARE_EVENT_OneParam(FUsdLevelSequenceHelper, FOnSkelAnimationBaked, const FString& /*SkeletonPrimPath*/);
 	FOnSkelAnimationBaked& GetOnSkelAnimationBaked();
 
 private:
@@ -106,15 +107,15 @@ private:
 class USDSTAGE_API FScopedBlockMonitoringChangesForTransaction final
 {
 public:
-	explicit FScopedBlockMonitoringChangesForTransaction( FUsdLevelSequenceHelper& InHelper );
-	explicit FScopedBlockMonitoringChangesForTransaction( FUsdLevelSequenceHelperImpl& InHelperImpl );
+	explicit FScopedBlockMonitoringChangesForTransaction(FUsdLevelSequenceHelper& InHelper);
+	explicit FScopedBlockMonitoringChangesForTransaction(FUsdLevelSequenceHelperImpl& InHelperImpl);
 	~FScopedBlockMonitoringChangesForTransaction();
 
 	FScopedBlockMonitoringChangesForTransaction() = delete;
-	FScopedBlockMonitoringChangesForTransaction( const FScopedBlockMonitoringChangesForTransaction& ) = delete;
-	FScopedBlockMonitoringChangesForTransaction( FScopedBlockMonitoringChangesForTransaction&& ) = delete;
-	FScopedBlockMonitoringChangesForTransaction& operator=( const FScopedBlockMonitoringChangesForTransaction& ) = delete;
-	FScopedBlockMonitoringChangesForTransaction& operator=( FScopedBlockMonitoringChangesForTransaction&& ) = delete;
+	FScopedBlockMonitoringChangesForTransaction(const FScopedBlockMonitoringChangesForTransaction&) = delete;
+	FScopedBlockMonitoringChangesForTransaction(FScopedBlockMonitoringChangesForTransaction&&) = delete;
+	FScopedBlockMonitoringChangesForTransaction& operator=(const FScopedBlockMonitoringChangesForTransaction&) = delete;
+	FScopedBlockMonitoringChangesForTransaction& operator=(FScopedBlockMonitoringChangesForTransaction&&) = delete;
 
 private:
 	FUsdLevelSequenceHelperImpl& HelperImpl;
diff --git a/Engine/Plugins/Importers/USDImporter/Source/USDStage/Public/USDPrimTwin.h b/Engine/Plugins/Importers/USDImporter/Source/USDStage/Public/USDPrimTwin.h
--- a/Engine/Plugins/Importers/USDImporter/Source/USDStage/Public/USDPrimTwin.h
+++ b/Engine/Plugins/Importers/USDImporter/Source/USDStage/Public/USDPrimTwin.h
@@ -3,70 +3,73 @@
 #pragma once
 
 #include "CoreMinimal.h"
-
-#include "Containers/Map.h"
-#include "Containers/UnrealString.h"
 #include "Delegates/Delegate.h"
 #include "Delegates/DelegateCombinations.h"
-#include "Templates/UniquePtr.h"
 #include "UObject/WeakObjectPtr.h"
-#include "UObject/WeakObjectPtrTemplates.h"
-#include "UObject/SoftObjectPtr.h"
 
 #include "USDPrimTwin.generated.h"
 
 /** The Unreal equivalent (twin) of a USD prim */
-UCLASS( Transient, MinimalAPI )
+UCLASS(Transient, MinimalAPI)
 class UUsdPrimTwin final : public UObject
 {
 	GENERATED_BODY()
 
 public:
-	UUsdPrimTwin& AddChild( const FString& InPrimPath );
-	void RemoveChild( const TCHAR* InPrimPath );
-	const TMap<FString, TObjectPtr<UUsdPrimTwin>>& GetChildren() const { return Children; }
+	UUsdPrimTwin& AddChild(const FString& InPrimPath);
+	void RemoveChild(const TCHAR* InPrimPath);
+	const TMap<FString, TObjectPtr<UUsdPrimTwin>>& GetChildren() const
+	{
+		return Children;
+	}
 
-	UUsdPrimTwin* GetParent() const { return Parent.Get(); }
+	UUsdPrimTwin* GetParent() const
+	{
+		return Parent.Get();
+	}
 
 	void Clear();
 
-	void Iterate( TFunction< void( UUsdPrimTwin& ) > Func, bool bRecursive )
+	void Iterate(TFunction<void(UUsdPrimTwin&)> Func, bool bRecursive)
 	{
-		for ( TMap<FString, TObjectPtr<UUsdPrimTwin>>::TIterator It = Children.CreateIterator(); It; ++It )
+		for (TMap<FString, TObjectPtr<UUsdPrimTwin>>::TIterator It = Children.CreateIterator(); It; ++It)
 		{
 			UUsdPrimTwin* Child = It->Value;
 
-			Func( *Child );
+			Func(*Child);
 
-			if ( bRecursive )
+			if (bRecursive)
 			{
-				Child->Iterate( Func, bRecursive );
+				Child->Iterate(Func, bRecursive);
 			}
 		}
 	}
 
-	UUsdPrimTwin* Find( const FString& InPrimPath );
-	UUsdPrimTwin* Find( const USceneComponent* InSceneComponent );
+	UUsdPrimTwin* Find(const FString& InPrimPath);
+	UUsdPrimTwin* Find(const USceneComponent* InSceneComponent);
 
 	USceneComponent* GetSceneComponent() const;
 
-	DECLARE_EVENT_OneParam( UUsdPrimTwin, FOnUsdPrimTwinDestroyed, const UUsdPrimTwin& );
+	DECLARE_EVENT_OneParam(UUsdPrimTwin, FOnUsdPrimTwinDestroyed, const UUsdPrimTwin&);
 	FOnUsdPrimTwinDestroyed OnDestroyed;
 
 public:
-	const static FName GetChildrenPropertyName() { return GET_MEMBER_NAME_CHECKED( UUsdPrimTwin, Children ); }
+	const static FName GetChildrenPropertyName()
+	{
+		return GET_MEMBER_NAME_CHECKED(UUsdPrimTwin, Children);
+	}
 
 public:
 	UPROPERTY()
 	FString PrimPath;
 
 	UPROPERTY()
-	TWeakObjectPtr< class USceneComponent > SceneComponent;
+	TWeakObjectPtr<class USceneComponent> SceneComponent;
 
 private:
-	UPROPERTY( Transient )
-	TMap< FString, TObjectPtr<UUsdPrimTwin> > Children;
+	UPROPERTY(Transient)
+	TMap<FString, TObjectPtr<UUsdPrimTwin>> Children;
 
-	UPROPERTY( Transient )
-	TWeakObjectPtr< UUsdPrimTwin > Parent;
+	UPROPERTY(Transient)
+	TWeakObjectPtr<UUsdPrimTwin> Parent;
 };
diff --git a/Engine/Plugins/Importers/USDImporter/Source/USDStage/Public/USDStageActor.h b/Engine/Plugins/Importers/USDImporter/Source/USDStage/Public/USDStageActor.h
--- a/Engine/Plugins/Importers/USDImporter/Source/USDStage/Public/USDStageActor.h
+++ b/Engine/Plugins/Importers/USDImporter/Source/USDStage/Public/USDStageActor.h
@@ -37,7 +37,7 @@ enum class EUsdStageState : uint8
 	OpenedAndLoaded
 };
 
-UCLASS( MinimalAPI, config = USDImporter )
+UCLASS(MinimalAPI, config = USDImporter)
 class AUsdStageActor : public AActor
 {
 	GENERATED_BODY()
@@ -59,24 +59,24 @@ public:
 	// This allows the defaults for them to be configured on EditorPerProjectUserSettings.ini, and allows us to write
 	// to that config from the USD Stage Editor, specifying our options before the editor is attached to any stage actor.
 
-	UPROPERTY(EditAnywhere, BlueprintReadOnly, Category = "USD", config )
+	UPROPERTY(EditAnywhere, BlueprintReadOnly, Category = "USD", config)
 	EUsdInitialLoadSet InitialLoadSet;
 
-	UPROPERTY( EditAnywhere, BlueprintReadOnly, Category = "USD", config )
+	UPROPERTY(EditAnywhere, BlueprintReadOnly, Category = "USD", config)
 	EUsdInterpolationType InterpolationType;
 
 	/**
 	 * Whether to try to combine individual assets and components of the same type on a kind-per-kind basis,
 	 * like multiple Mesh prims into a single Static Mesh
 	 */
-	UPROPERTY(EditAnywhere, BlueprintReadOnly, Category = "USD", config, meta = (Bitmask, BitmaskEnum="/Script/UnrealUSDWrapper.EUsdDefaultKind"))
+	UPROPERTY(EditAnywhere, BlueprintReadOnly, Category = "USD", config, meta = (Bitmask, BitmaskEnum = "/Script/UnrealUSDWrapper.EUsdDefaultKind"))
 	int32 KindsToCollapse;
 
 	/**
 	 * If enabled, when multiple mesh prims are collapsed into a single static mesh, identical material slots are merged into one slot.
 	 * Otherwise, materials slots are simply appended to the list.
 	 */
-	UPROPERTY(EditAnywhere, BlueprintReadOnly, Category = "USD", config )
+	UPROPERTY(EditAnywhere, BlueprintReadOnly, Category = "USD", config)
 	bool bMergeIdenticalMaterialSlots;
 
 	/**
@@ -84,7 +84,7 @@ public:
 	 * that asset is generated, and the asset is shared by the components generated for both prims.
 	 * If false, we will always generate a dedicated asset for each prim.
 	 */
-	UPROPERTY(EditAnywhere, BlueprintReadOnly, Category = "USD", config )
+	UPROPERTY(EditAnywhere, BlueprintReadOnly, Category = "USD", config)
 	bool bReuseIdenticalAssets;
 
 	/**
@@ -92,28 +92,28 @@ public:
 	 * If false, will cause us to use HierarchicalInstancedStaticMeshComponents to replicate the instancing behavior.
 	 * Point instancers inside other point instancer prototypes are *always* collapsed into the prototype's static mesh.
 	 */
-	UE_DEPRECATED( 5.2, "This option is now controlled via the cvar 'USD.CollapseTopLevelPointInstancers'" )
+	UE_DEPRECATED(5.2, "This option is now controlled via the cvar 'USD.CollapseTopLevelPointInstancers'")
 	UPROPERTY()
 	bool bCollapseTopLevelPointInstancers;
 
 	/* Only load prims with these specific purposes from the USD file */
-	UPROPERTY(EditAnywhere, BlueprintReadOnly, Category = "USD", config, meta = (Bitmask, BitmaskEnum="/Script/UnrealUSDWrapper.EUsdPurpose"))
+	UPROPERTY(EditAnywhere, BlueprintReadOnly, Category = "USD", config, meta = (Bitmask, BitmaskEnum = "/Script/UnrealUSDWrapper.EUsdPurpose"))
 	int32 PurposesToLoad;
 
 	/** Try enabling Nanite for static meshes that are generated with at least this many triangles */
-	UPROPERTY( EditAnywhere, BlueprintReadOnly, Category = "USD", config, meta = ( NoSpinbox = "true", UIMin = "0", ClampMin = "0" ))
+	UPROPERTY(EditAnywhere, BlueprintReadOnly, Category = "USD", config, meta = (NoSpinbox = "true", UIMin = "0", ClampMin = "0"))
 	int32 NaniteTriangleThreshold;
 
 	/** Specifies which set of shaders to use when parsing USD materials, in addition to the universal render context. */
-	UPROPERTY(EditAnywhere, BlueprintReadOnly, Category = "USD", config )
+	UPROPERTY(EditAnywhere, BlueprintReadOnly, Category = "USD", config)
 	FName RenderContext;
 
 	/** Specifies which material purbose to use when parsing USD material bindings, in addition to the "allPurpose" fallback. */
-	UPROPERTY( EditAnywhere, BlueprintReadOnly, Category = "USD", config )
+	UPROPERTY(EditAnywhere, BlueprintReadOnly, Category = "USD", config)
 	FName MaterialPurpose;
 
 	// Describes what to add to the root bone animation within generated AnimSequences, if anything
-	UPROPERTY( EditAnywhere, BlueprintReadOnly, Category = "USD", config )
+	UPROPERTY(EditAnywhere, BlueprintReadOnly, Category = "USD", config)
 	EUsdRootMotionHandling RootMotionHandling;
 
 	/**
@@ -149,7 +149,7 @@ public:
 
 public:
 	UFUNCTION(BlueprintCallable, Category = "USD", meta = (CallInEditor = "true"))
-	USDSTAGE_API void SetRootLayer(const FString& RootFilePath );
+	USDSTAGE_API void SetRootLayer(const FString& RootFilePath);
 
 	UFUNCTION(BlueprintCallable, Category = "USD", meta = (CallInEditor = "true"))
 	USDSTAGE_API void SetStageState(EUsdStageState NewStageState);
@@ -158,38 +158,38 @@ public:
 	USDSTAGE_API void SetAssetCache(UUsdAssetCache2* NewCache);
 
 	UFUNCTION(BlueprintCallable, Category = "USD", meta = (CallInEditor = "true"))
-	USDSTAGE_API void SetInitialLoadSet( EUsdInitialLoadSet NewLoadSet );
+	USDSTAGE_API void SetInitialLoadSet(EUsdInitialLoadSet NewLoadSet);
 
-	UFUNCTION( BlueprintCallable, Category = "USD", meta = ( CallInEditor = "true" ) )
-	USDSTAGE_API void SetInterpolationType( EUsdInterpolationType NewType );
+	UFUNCTION(BlueprintCallable, Category = "USD", meta = (CallInEditor = "true"))
+	USDSTAGE_API void SetInterpolationType(EUsdInterpolationType NewType);
 
-	UFUNCTION( BlueprintCallable, Category = "USD", meta = ( CallInEditor = "true" ) )
-	USDSTAGE_API void SetKindsToCollapse( int32 NewKindsToCollapse );
+	UFUNCTION(BlueprintCallable, Category = "USD", meta = (CallInEditor = "true"))
+	USDSTAGE_API void SetKindsToCollapse(int32 NewKindsToCollapse);
 
-	UFUNCTION( BlueprintCallable, Category = "USD", meta = ( CallInEditor = "true" ) )
-	USDSTAGE_API void SetMergeIdenticalMaterialSlots( bool bMerge );
+	UFUNCTION(BlueprintCallable, Category = "USD", meta = (CallInEditor = "true"))
+	USDSTAGE_API void SetMergeIdenticalMaterialSlots(bool bMerge);
 
 	UFUNCTION(BlueprintCallable, Category = "USD", meta = (CallInEditor = "true"))
 	USDSTAGE_API void SetReuseIdenticalAssets(bool bReuse);
 
-	UE_DEPRECATED( 5.2, "This option is now controlled via the cvar 'USD.CollapseTopLevelPointInstancers'" )
-	UFUNCTION( BlueprintCallable, Category = "USD", meta = ( CallInEditor = "true" ) )
-	USDSTAGE_API void SetCollapseTopLevelPointInstancers( bool bCollapse );
+	UE_DEPRECATED(5.2, "This option is now controlled via the cvar 'USD.CollapseTopLevelPointInstancers'")
+	UFUNCTION(BlueprintCallable, Category = "USD", meta = (CallInEditor = "true"))
+	USDSTAGE_API void SetCollapseTopLevelPointInstancers(bool bCollapse);
 
 	UFUNCTION(BlueprintCallable, Category = "USD", meta = (CallInEditor = "true"))
-	USDSTAGE_API void SetPurposesToLoad( int32 NewPurposesToLoad );
+	USDSTAGE_API void SetPurposesToLoad(int32 NewPurposesToLoad);
 
-	UFUNCTION( BlueprintCallable, Category = "USD", meta = ( CallInEditor = "true" ) )
-	USDSTAGE_API void SetNaniteTriangleThreshold( int32 NewNaniteTriangleThreshold );
+	UFUNCTION(BlueprintCallable, Category = "USD", meta = (CallInEditor = "true"))
+	USDSTAGE_API void SetNaniteTriangleThreshold(int32 NewNaniteTriangleThreshold);
 
 	UFUNCTION(BlueprintCallable, Category = "USD", meta = (CallInEditor = "true"))
-	USDSTAGE_API void SetRenderContext( const FName& NewRenderContext );
+	USDSTAGE_API void SetRenderContext(const FName& NewRenderContext);
 
-	UFUNCTION( BlueprintCallable, Category = "USD", meta = ( CallInEditor = "true" ) )
-	USDSTAGE_API void SetMaterialPurpose( const FName& NewMaterialPurpose );
+	UFUNCTION(BlueprintCallable, Category = "USD", meta = (CallInEditor = "true"))
+	USDSTAGE_API void SetMaterialPurpose(const FName& NewMaterialPurpose);
 
-	UFUNCTION( BlueprintCallable, Category = "USD", meta = ( CallInEditor = "true" ) )
-	USDSTAGE_API void SetRootMotionHandling( EUsdRootMotionHandling NewHandlingStrategy );
+	UFUNCTION(BlueprintCallable, Category = "USD", meta = (CallInEditor = "true"))
+	USDSTAGE_API void SetRootMotionHandling(EUsdRootMotionHandling NewHandlingStrategy);
 
 	UFUNCTION(BlueprintCallable, Category = "USD", meta = (CallInEditor = "true"))
 	USDSTAGE_API void SetSubdivisionLevel(int32 NewLevel);
@@ -215,33 +215,37 @@ public:
 	UFUNCTION(BlueprintCallable, Category = "USD", meta = (CallInEditor = "true"))
 	USDSTAGE_API void SetTime(float InTime);
 
-	UFUNCTION( BlueprintCallable, Category = "USD", meta = ( CallInEditor = "true" ) )
+	UFUNCTION(BlueprintCallable, Category = "USD", meta = (CallInEditor = "true"))
 	USDSTAGE_API ULevelSequence* GetLevelSequence();
 
 	/**
 	 * Gets the transient component that was generated for a prim with a given prim path.
-	 * Warning: The lifetime of the component is managed by the AUsdStageActor, and it may be force-destroyed at any time (e.g. when closing the stage)
+	 * Warning: The lifetime of the component is managed by the AUsdStageActor, and it may be force-destroyed at any time (e.g. when closing the
+	 * stage)
 	 * @param PrimPath - Full path to the source prim, e.g. "/root_prim/my_prim"
-	 * @return The corresponding spawned component. It may correspond to a parent prim, if the prim at PrimPath was collapsed. Nullptr if path is invalid.
+	 * @return The corresponding spawned component. It may correspond to a parent prim, if the prim at PrimPath was collapsed. Nullptr if path is
+	 * invalid.
 	 */
 	UFUNCTION(BlueprintCallable, Category = "USD", meta = (CallInEditor = "true"))
-	USDSTAGE_API USceneComponent* GetGeneratedComponent( const FString& PrimPath );
+	USDSTAGE_API USceneComponent* GetGeneratedComponent(const FString& PrimPath);
 
 	/**
-	 * Gets the transient assets that were generated for a prim with a given prim path. Likely one asset (e.g. UStaticMesh), but can be multiple (USkeletalMesh, USkeleton, etc.)
+	 * Gets the transient assets that were generated for a prim with a given prim path. Likely one asset (e.g. UStaticMesh), but can be multiple
+	 * (USkeletalMesh, USkeleton, etc.)
 	 * @param PrimPath - Full path to the source prim, e.g. "/root_prim/my_mesh"
 	 * @return The corresponding generated assets. May be empty if path is invalid or if that prim led to no generated assets.
 	 */
 	UFUNCTION(BlueprintCallable, Category = "USD", meta = (CallInEditor = "true"))
-	USDSTAGE_API TArray<UObject*> GetGeneratedAssets( const FString& PrimPath );
+	USDSTAGE_API TArray<UObject*> GetGeneratedAssets(const FString& PrimPath);
 
 	/**
 	 * Gets the path to the prim that was parsed to generate the given `Object`.
-	 * @param Object - UObject to query with. Can be one of the transient components generated when a stage was opened, or something like a UStaticMesh.
+	 * @param Object - UObject to query with. Can be one of the transient components generated when a stage was opened, or something like a
+	 * UStaticMesh.
 	 * @return The path to the source prim, e.g. "/root_prim/some_prim". May be empty in case we couldn't find the source prim.
 	 */
 	UFUNCTION(BlueprintCallable, Category = "USD", meta = (CallInEditor = "true"))
-	USDSTAGE_API FString GetSourcePrimPath( UObject* Object );
+	USDSTAGE_API FString GetSourcePrimPath(UObject* Object);
 
 	// Creates a brand new, memory-only USD stage and opens it
 	UFUNCTION(BlueprintCallable, Category = "USD", meta = (CallInEditor = "true"))
@@ -282,14 +286,14 @@ public:
 	USDSTAGE_API void SetUsdStage(const UE::FUsdStage& NewStage);
 
 	// Enters isolated mode by creating a new USD Stage using the provided layer as its root
-	USDSTAGE_API void IsolateLayer( const UE::FSdfLayer& Layer );
+	USDSTAGE_API void IsolateLayer(const UE::FSdfLayer& Layer);
 
 	// Regenerates our LevelSequence from the USD Stage
 	USDSTAGE_API void ReloadAnimations();
 
 	USDSTAGE_API TSharedPtr<FUsdInfoCache> GetInfoCache();
 	USDSTAGE_API TSharedPtr<UE::FUsdGeomBBoxCache> GetBBoxCache();
-	USDSTAGE_API TMap< FString, TMap< FString, int32 > > GetMaterialToPrimvarToUVIndex();
+	USDSTAGE_API TMap<FString, TMap<FString, int32>> GetMaterialToPrimvarToUVIndex();
 	USDSTAGE_API const UsdUtils::FBlendShapeMap& GetBlendShapeMap();
 	USDSTAGE_API FUsdListener& GetUsdListener();
 	USDSTAGE_API const FUsdListener& GetUsdListener() const;
@@ -305,18 +309,18 @@ public:
 	USDSTAGE_API void BlockMonitoringLevelSequenceForThisTransaction();
 
 	// Begin UObject interface
-	USDSTAGE_API virtual void PostDuplicate( bool bDuplicateForPIE ) override;
+	USDSTAGE_API virtual void PostDuplicate(bool bDuplicateForPIE) override;
 	USDSTAGE_API virtual void Serialize(FArchive& Ar) override;
 	USDSTAGE_API virtual void Destroyed() override;
 	USDSTAGE_API virtual void PostActorCreated() override;
-	USDSTAGE_API virtual void PostRename( UObject* OldOuter, const FName OldName ) override;
+	USDSTAGE_API virtual void PostRename(UObject* OldOuter, const FName OldName) override;
 	USDSTAGE_API virtual void BeginDestroy() override;
 #if WITH_EDITOR
-	USDSTAGE_API virtual void PostEditChangeProperty( FPropertyChangedEvent& PropertyChangedEvent ) override;
+	USDSTAGE_API virtual void PostEditChangeProperty(FPropertyChangedEvent& PropertyChangedEvent) override;
 	USDSTAGE_API virtual void PostTransacted(const FTransactionObjectEvent& TransactionEvent) override;
-	USDSTAGE_API virtual void PreEditChange( FProperty* PropertyThatWillChange ) override;
+	USDSTAGE_API virtual void PreEditChange(FProperty* PropertyThatWillChange) override;
 	USDSTAGE_API virtual void PreEditUndo() override;
-#endif // WITH_EDITOR
+#endif	  // WITH_EDITOR
 	// End UObject interface
 
 	// Begin AActor interface
@@ -365,24 +369,29 @@ protected:
 #if WITH_EDITOR
 	void OnBeginPIE(bool bIsSimulating);
 	void OnPostPIEStarted(bool bIsSimulating);
-	void OnObjectsReplaced( const TMap<UObject*, UObject*>& ObjectReplacementMap );
+	void OnObjectsReplaced(const TMap<UObject*, UObject*>& ObjectReplacementMap);
 	void OnLevelActorDeleted(AActor* DeletedActor);
 	void HandleTransactionStateChanged(const FTransactionContext& InTransactionContext, const ETransactionStateEventType InTransactionState);
-#endif // WITH_EDITOR
+#endif	  // WITH_EDITOR
 
 	void OnPreUsdImport(FString FilePath);
 	void OnPostUsdImport(FString FilePath);
-	void OnUsdObjectsChanged( const UsdUtils::FObjectChangesByPath& InfoChanges, const UsdUtils::FObjectChangesByPath& ResyncChanges );
-	void OnUsdPrimTwinDestroyed( const UUsdPrimTwin& UsdPrimTwin );
-	void OnObjectPropertyChanged( UObject* ObjectBeingModified, FPropertyChangedEvent& PropertyChangedEvent );
-	void HandlePropertyChangedEvent( FPropertyChangedEvent& PropertyChangedEvent );
+	void OnUsdObjectsChanged(const UsdUtils::FObjectChangesByPath& InfoChanges, const UsdUtils::FObjectChangesByPath& ResyncChanges);
+	void OnUsdPrimTwinDestroyed(const UUsdPrimTwin& UsdPrimTwin);
+	void OnObjectPropertyChanged(UObject* ObjectBeingModified, FPropertyChangedEvent& PropertyChangedEvent);
+	void HandlePropertyChangedEvent(FPropertyChangedEvent& PropertyChangedEvent);
 	void OnSkelAnimationBaked(const FString& SkeletonPrimPath);
 
 protected:
 	friend struct FUsdStageActorImpl;
 	friend class FUsdLevelSequenceHelperImpl;
 
-	UPROPERTY(Category = UsdStageActor, VisibleAnywhere, BlueprintReadOnly, meta = (ExposeFunctionCategories = "Mesh,Rendering,Physics,Components|StaticMesh", AllowPrivateAccess = "true"))
+	UPROPERTY(
+		Category = UsdStageActor,
+		VisibleAnywhere,
+		BlueprintReadOnly,
+		meta = (ExposeFunctionCategories = "Mesh,Rendering,Physics,Components|StaticMesh", AllowPrivateAccess = "true")
+	)
 	TObjectPtr<class USceneComponent> SceneComponent;
 
 	/* TimeCode to evaluate the USD stage at */
@@ -393,20 +402,23 @@ protected:
 	TObjectPtr<ULevelSequence> LevelSequence;
 
 	// Note: This cannot be instanced: Read the comment in the constructor
-	UPROPERTY( Transient )
+	UPROPERTY(Transient)
 	TObjectPtr<UUsdPrimTwin> RootUsdTwin;
 
-	UPROPERTY( Transient )
-	TSet< FString > PrimsToAnimate;
+	UPROPERTY(Transient)
+	TSet<FString> PrimsToAnimate;
 
-	UPROPERTY( Transient )
-	TMap< TObjectPtr<UObject>, FString > ObjectsToWatch;
+	UPROPERTY(Transient)
+	TMap<TObjectPtr<UObject>, FString> ObjectsToWatch;
 
 	UE_DEPRECATED(5.2, "Use the new AssetCache property instead, that uses UUsdAssetCache2 objects")
-	UPROPERTY(AdvancedDisplay, meta=(DeprecatedProperty, DeprecationMessage = "Use the new AssetCache property instead, that uses UUsdAssetCache2 objects"))
+	UPROPERTY(
+		AdvancedDisplay,
+		meta = (DeprecatedProperty, DeprecationMessage = "Use the new AssetCache property instead, that uses UUsdAssetCache2 objects")
+	)
 	TObjectPtr<UUsdAssetCache> AssetCache;
 
-	UPROPERTY( Transient )
+	UPROPERTY(Transient)
 	TObjectPtr<UUsdTransactor> Transactor;
 
 	/** Caches various information about prims that are expensive to query */
@@ -433,7 +445,7 @@ protected:
 	 * When parsing meshes later, we use this data to place the correct primvar values in each UV channel.
 	 * We keep this here as these are generated when the materials stored in AssetsCache are parsed, so it should accompany them
 	 */
-	TMap< FString, TMap< FString, int32 > > MaterialToPrimvarToUVIndex;
+	TMap<FString, TMap<FString, int32>> MaterialToPrimvarToUVIndex;
 
 	/**
 	 * We use PostRegisterAllComponents and PostUnregisterAllComponents as main entry points to decide when to load/unload
@@ -464,16 +476,15 @@ protected:
 class USDSTAGE_API FScopedBlockNoticeListening final
 {
 public:
-	explicit FScopedBlockNoticeListening( AUsdStageActor* InStageActor );
+	explicit FScopedBlockNoticeListening(AUsdStageActor* InStageActor);
 	~FScopedBlockNoticeListening();
 
 	FScopedBlockNoticeListening() = delete;
-	FScopedBlockNoticeListening( const FScopedBlockNoticeListening& ) = delete;
-	FScopedBlockNoticeListening( FScopedBlockNoticeListening&& ) = delete;
-	FScopedBlockNoticeListening& operator=( const FScopedBlockNoticeListening& ) = delete;
-	FScopedBlockNoticeListening& operator=( FScopedBlockNoticeListening&& ) = delete;
+	FScopedBlockNoticeListening(const FScopedBlockNoticeListening&) = delete;
+	FScopedBlockNoticeListening(FScopedBlockNoticeListening&&) = delete;
+	FScopedBlockNoticeListening& operator=(const FScopedBlockNoticeListening&) = delete;
+	FScopedBlockNoticeListening& operator=(FScopedBlockNoticeListening&&) = delete;
 
 private:
 	TWeakObjectPtr<AUsdStageActor> StageActor;
 };
-
diff --git a/Engine/Plugins/Importers/USDImporter/Source/USDStage/Public/USDTransactor.h b/Engine/Plugins/Importers/USDImporter/Source/USDStage/Public/USDTransactor.h
--- a/Engine/Plugins/Importers/USDImporter/Source/USDStage/Public/USDTransactor.h
+++ b/Engine/Plugins/Importers/USDImporter/Source/USDStage/Public/USDTransactor.h
@@ -3,11 +3,8 @@
 #pragma once
 
 #include "CoreMinimal.h"
-#include "Misc/ITransaction.h"
 #include "UObject/WeakObjectPtr.h"
 
-#include "UsdWrappers/VtValue.h"
-
 #include "USDTransactor.generated.h"
 
 class AUsdStageActor;
@@ -45,22 +42,22 @@ class USDSTAGE_API UUsdTransactor : public UObject
 public:
 	// Boilerplate for Pimpl usage with UObject
 	UUsdTransactor();
-	UUsdTransactor( FVTableHelper& Helper );
+	UUsdTransactor(FVTableHelper& Helper);
 	~UUsdTransactor();
 
-	void Initialize( AUsdStageActor* InStageActor );
-	void Update( const UsdUtils::FObjectChangesByPath& NewInfoChanges, const UsdUtils::FObjectChangesByPath& NewResyncChanges );
+	void Initialize(AUsdStageActor* InStageActor);
+	void Update(const UsdUtils::FObjectChangesByPath& NewInfoChanges, const UsdUtils::FObjectChangesByPath& NewResyncChanges);
 
 	// Begin UObject interface
-	virtual void Serialize( FArchive& Ar ) override;
+	virtual void Serialize(FArchive& Ar) override;
 #if WITH_EDITOR
 	virtual void PreEditUndo();
 	virtual void PostEditUndo();
-#endif // WITH_EDITOR
-	//~ End UObject interface
+#endif	  // WITH_EDITOR
+		  //~ End UObject interface
 
 private:
 	TWeakObjectPtr<AUsdStageActor> StageActor;
 
-	TUniquePtr< UsdUtils::FUsdTransactorImpl > Impl;
+	TUniquePtr<UsdUtils::FUsdTransactorImpl> Impl;
 };
