USD: Reformat USDStageEditor module.

#rnx

diff --git a/Engine/Plugins/Importers/USDImporter/Source/USDStageEditor/Private/SUSDLayersTreeView.cpp b/Engine/Plugins/Importers/USDImporter/Source/USDStageEditor/Private/SUSDLayersTreeView.cpp
--- a/Engine/Plugins/Importers/USDImporter/Source/USDStageEditor/Private/SUSDLayersTreeView.cpp
+++ b/Engine/Plugins/Importers/USDImporter/Source/USDStageEditor/Private/SUSDLayersTreeView.cpp
@@ -11,24 +11,17 @@
 #include "USDLog.h"
 #include "USDMemory.h"
 #include "USDProjectSettings.h"
-#include "USDStageActor.h"
-#include "USDStageModule.h"
-#include "USDTypesConversion.h"
 #include "UsdWrappers/SdfChangeBlock.h"
 #include "UsdWrappers/SdfLayer.h"
 #include "UsdWrappers/UsdStage.h"
 
 #include "DesktopPlatformModule.h"
-#include "Engine/World.h"
 #include "EngineAnalytics.h"
 #include "Framework/MultiBox/MultiBoxBuilder.h"
 #include "Framework/Notifications/NotificationManager.h"
 #include "IDesktopPlatform.h"
-#include "Modules/ModuleManager.h"
 #include "ScopedTransaction.h"
 #include "Styling/AppStyle.h"
-#include "Styling/SlateIconFinder.h"
-#include "Styling/SlateTypes.h"
 #include "Templates/SharedPointer.h"
 #include "Widgets/Images/SImage.h"
 #include "Widgets/Input/SButton.h"
@@ -41,104 +34,104 @@
 
 namespace UE::USDLayersTreeViewImpl::Private
 {
-	void ExportLayerToPath( const UE::FSdfLayer& LayerToExport, const FString& TargetPath )
+	void ExportLayerToPath(const UE::FSdfLayer& LayerToExport, const FString& TargetPath)
 	{
-		if ( !LayerToExport )
+		if (!LayerToExport)
 		{
 			return;
 		}
 
 		// Clone the layer so that we don't modify the currently opened stage when we do the remapping below
-		UE::FSdfLayer OutputLayer = UE::FSdfLayer::CreateNew( *TargetPath );
-		OutputLayer.TransferContent( LayerToExport );
+		UE::FSdfLayer OutputLayer = UE::FSdfLayer::CreateNew(*TargetPath);
+		OutputLayer.TransferContent(LayerToExport);
 
 		// Update references to assets (e.g. textures) so that they're absolute and also work from the new file
-		UsdUtils::ConvertAssetRelativePathsToAbsolute( OutputLayer, LayerToExport );
+		UsdUtils::ConvertAssetRelativePathsToAbsolute(OutputLayer, LayerToExport);
 
 		// Convert layer references to absolute paths so that it still works at its target location
 		FString LayerPath = LayerToExport.GetRealPath();
-		FPaths::NormalizeFilename( LayerPath );
+		FPaths::NormalizeFilename(LayerPath);
 		const TSet<FString> AssetDependencies = OutputLayer.GetCompositionAssetDependencies();
-		for ( const FString& Ref : AssetDependencies )
+		for (const FString& Ref : AssetDependencies)
 		{
-			FString AbsRef = FPaths::ConvertRelativePathToFull( FPaths::GetPath( LayerPath ), Ref ); // Relative to the original file
-			OutputLayer.UpdateCompositionAssetDependency( *Ref, *AbsRef );
+			FString AbsRef = FPaths::ConvertRelativePathToFull(FPaths::GetPath(LayerPath), Ref);	// Relative to the original file
+			OutputLayer.UpdateCompositionAssetDependency(*Ref, *AbsRef);
 		}
 
 		bool bForce = true;
-		OutputLayer.Save( bForce );
+		OutputLayer.Save(bForce);
 	}
-}
+}	 // namespace UE::USDLayersTreeViewImpl::Private
 
-class FUsdLayerNameColumn : public FUsdTreeViewColumn, public TSharedFromThis< FUsdLayerNameColumn >
+class FUsdLayerNameColumn
+	: public FUsdTreeViewColumn
+	, public TSharedFromThis<FUsdLayerNameColumn>
 {
 public:
-	FSlateColor GetForegroundColor( const FUsdLayerViewModelRef TreeItem ) const
+	FSlateColor GetForegroundColor(const FUsdLayerViewModelRef TreeItem) const
 	{
-		return TreeItem->IsInIsolatedStage()
-			? FSlateColor::UseForeground()
-			: FSlateColor::UseSubduedForeground();
+		return TreeItem->IsInIsolatedStage() ? FSlateColor::UseForeground() : FSlateColor::UseSubduedForeground();
 	}
 
-	virtual TSharedRef< SWidget > GenerateWidget( const TSharedPtr< IUsdTreeViewItem > InTreeItem, const TSharedPtr< ITableRow > TableRow ) override
+	virtual TSharedRef<SWidget> GenerateWidget(const TSharedPtr<IUsdTreeViewItem> InTreeItem, const TSharedPtr<ITableRow> TableRow) override
 	{
-		FUsdLayerViewModelRef TreeItem = StaticCastSharedRef< FUsdLayerViewModel >( InTreeItem.ToSharedRef() );
+		FUsdLayerViewModelRef TreeItem = StaticCastSharedRef<FUsdLayerViewModel>(InTreeItem.ToSharedRef());
 		TWeakPtr<FUsdLayerViewModel> TreeItemWeak = TreeItem;
 
-		return SNew( SBox )
-			.VAlign( VAlign_Center )
+		// clang-format off
+		return SNew(SBox)
+			.VAlign(VAlign_Center)
 			[
-				SNew( STextBlock )
-				.Text( TreeItem, &FUsdLayerViewModel::GetDisplayName )
-				.ColorAndOpacity( this, &FUsdLayerNameColumn::GetForegroundColor, TreeItem )
-				.ToolTipText_Lambda( [TreeItemWeak]
+				SNew(STextBlock)
+				.Text(TreeItem, &FUsdLayerViewModel::GetDisplayName)
+				.ColorAndOpacity(this, &FUsdLayerNameColumn::GetForegroundColor, TreeItem)
+				.ToolTipText_Lambda([TreeItemWeak]
 				{
-					if ( TSharedPtr<FUsdLayerViewModel> PinnedTreeItem = TreeItemWeak.Pin() )
+					if (TSharedPtr<FUsdLayerViewModel> PinnedTreeItem = TreeItemWeak.Pin())
 					{
-						return FText::FromString( PinnedTreeItem->LayerIdentifier );
+						return FText::FromString(PinnedTreeItem->LayerIdentifier);
 					}
 
 					return FText::GetEmpty();
 				})
 			];
+		// clang-format on
 	}
 };
 
-class FUsdLayerMutedColumn : public FUsdTreeViewColumn, public TSharedFromThis< FUsdLayerMutedColumn >
+class FUsdLayerMutedColumn
+	: public FUsdTreeViewColumn
+	, public TSharedFromThis<FUsdLayerMutedColumn>
 {
 public:
-	FReply OnClicked( const FUsdLayerViewModelRef TreeItem )
+	FReply OnClicked(const FUsdLayerViewModelRef TreeItem)
 	{
-		ToggleMuteLayer( TreeItem );
+		ToggleMuteLayer(TreeItem);
 
 		return FReply::Handled();
 	}
 
-	const FSlateBrush* GetBrush( const FUsdLayerViewModelRef TreeItem, const TSharedPtr< SButton > Button ) const
+	const FSlateBrush* GetBrush(const FUsdLayerViewModelRef TreeItem, const TSharedPtr<SButton> Button) const
 	{
 		const bool bIsButtonHovered = Button.IsValid() && Button->IsHovered();
 
-		if ( !CanMuteLayer( TreeItem ) )
+		if (!CanMuteLayer(TreeItem))
 		{
 			return nullptr;
 		}
-		else if ( TreeItem->IsLayerMuted() )
+		else if (TreeItem->IsLayerMuted())
 		{
-			return bIsButtonHovered
-				? FAppStyle::GetBrush( "Level.NotVisibleHighlightIcon16x" )
-				: FAppStyle::GetBrush( "Level.NotVisibleIcon16x" );
+			return bIsButtonHovered ? FAppStyle::GetBrush("Level.NotVisibleHighlightIcon16x") : FAppStyle::GetBrush("Level.NotVisibleIcon16x");
 		}
 		else
 		{
-			return bIsButtonHovered
-				? FAppStyle::GetBrush( "Level.VisibleHighlightIcon16x" )
-				: FAppStyle::GetBrush( "Level.VisibleIcon16x" );
+			return bIsButtonHovered ? FAppStyle::GetBrush("Level.VisibleHighlightIcon16x") : FAppStyle::GetBrush("Level.VisibleIcon16x");
 		}
 	}
 
-	FSlateColor GetForegroundColor( const FUsdLayerViewModelRef TreeItem, const TSharedPtr< ITableRow > TableRow, const TSharedPtr< SButton > Button ) const
+	FSlateColor GetForegroundColor(const FUsdLayerViewModelRef TreeItem, const TSharedPtr<ITableRow> TableRow, const TSharedPtr<SButton> Button) const
 	{
-		if ( !TableRow.IsValid() || !Button.IsValid() )
+		if (!TableRow.IsValid() || !Button.IsValid())
 		{
 			return FSlateColor::UseForeground();
 		}
@@ -148,84 +141,86 @@ public:
 		const bool bIsRowSelected = TableRow->IsItemSelected();
 		const bool bIsLayerMuted = TreeItem->IsLayerMuted();
 
-		if ( !bIsLayerMuted && !bIsRowHovered && !bIsRowSelected )
+		if (!bIsLayerMuted && !bIsRowHovered && !bIsRowSelected)
 		{
 			return FLinearColor::Transparent;
 		}
-		else if ( bIsButtonHovered && !bIsRowSelected )
+		else if (bIsButtonHovered && !bIsRowSelected)
 		{
-			return FAppStyle::GetSlateColor( TEXT( "Colors.ForegroundHover" ) );
+			return FAppStyle::GetSlateColor(TEXT("Colors.ForegroundHover"));
 		}
 
 		return FSlateColor::UseForeground();
 	}
 
-	virtual TSharedRef< SWidget > GenerateWidget( const TSharedPtr< IUsdTreeViewItem > InTreeItem, const TSharedPtr< ITableRow > TableRow ) override
+	virtual TSharedRef<SWidget> GenerateWidget(const TSharedPtr<IUsdTreeViewItem> InTreeItem, const TSharedPtr<ITableRow> TableRow) override
 	{
-		if ( !InTreeItem )
+		if (!InTreeItem)
 		{
 			return SNullWidget::NullWidget;
 		}
 
-		FUsdLayerViewModelRef TreeItem = StaticCastSharedRef< FUsdLayerViewModel >( InTreeItem.ToSharedRef() );
+		FUsdLayerViewModelRef TreeItem = StaticCastSharedRef<FUsdLayerViewModel>(InTreeItem.ToSharedRef());
 		FUsdLayerViewModelWeak TreeItemWeak = TreeItem;
-		const float ItemSize = FUsdStageEditorStyle::Get()->GetFloat( "UsdStageEditor.ListItemHeight" );
+		const float ItemSize = FUsdStageEditorStyle::Get()->GetFloat("UsdStageEditor.ListItemHeight");
 
-		if ( !TreeItem->CanMuteLayer() )
+		// clang-format off
+		if (!TreeItem->CanMuteLayer())
 		{
-			return SNew( SBox )
-				.HeightOverride( ItemSize )
-				.WidthOverride( ItemSize )
-				.Visibility( EVisibility::Visible )
-				.ToolTip( SNew( SToolTip ).Text( LOCTEXT( "CantMuteLayerTooltip", "This layer cannot be muted!" ) ) );
+			return SNew(SBox)
+				.HeightOverride(ItemSize)
+				.WidthOverride(ItemSize)
+				.Visibility(EVisibility::Visible)
+				.ToolTip(SNew(SToolTip).Text(LOCTEXT("CantMuteLayerTooltip", "This layer cannot be muted!")));
 		}
 
-		TSharedPtr<SButton> Button = SNew( SButton )
-			.ContentPadding( 0 )
+		TSharedPtr<SButton> Button = SNew(SButton)
+			.ContentPadding(0)
 			.IsEnabled_Lambda([TreeItemWeak]()->bool
 			{
-				if ( TSharedPtr<FUsdLayerViewModel> PinnedTreeItem = TreeItemWeak.Pin() )
+				if (TSharedPtr<FUsdLayerViewModel> PinnedTreeItem = TreeItemWeak.Pin())
 				{
 					return PinnedTreeItem->IsInIsolatedStage();
 				}
 				return false;
 			})
-			.ButtonStyle( FUsdStageEditorStyle::Get(), TEXT("NoBorder") )
-			.OnClicked( this, &FUsdLayerMutedColumn::OnClicked, TreeItem )
-			.ToolTip( SNew( SToolTip ).Text( LOCTEXT( "MuteLayerTooltip", "Mute or unmute this layer" ) ) )
-			.HAlign( HAlign_Center )
-			.VAlign( VAlign_Center );
-
-		TSharedPtr<SImage> Image = SNew( SImage )
-			.Image( this, &FUsdLayerMutedColumn::GetBrush, TreeItem, Button )
-			.ColorAndOpacity( this, &FUsdLayerMutedColumn::GetForegroundColor, TreeItem, TableRow, Button );
-
-		Button->SetContent( Image.ToSharedRef() );
-
-		return SNew( SBox )
-			.HeightOverride( ItemSize )
-			.WidthOverride( ItemSize )
-			.Visibility( EVisibility::Visible )
+			.ButtonStyle(FUsdStageEditorStyle::Get(), TEXT("NoBorder"))
+			.OnClicked(this, &FUsdLayerMutedColumn::OnClicked, TreeItem)
+			.ToolTip(SNew(SToolTip).Text(LOCTEXT("MuteLayerTooltip", "Mute or unmute this layer")))
+			.HAlign(HAlign_Center)
+			.VAlign(VAlign_Center);
+
+		TSharedPtr<SImage> Image = SNew(SImage)
+			.Image(this, &FUsdLayerMutedColumn::GetBrush, TreeItem, Button)
+			.ColorAndOpacity(this, &FUsdLayerMutedColumn::GetForegroundColor, TreeItem, TableRow, Button);
+
+		Button->SetContent(Image.ToSharedRef());
+
+		return SNew(SBox)
+			.HeightOverride(ItemSize)
+			.WidthOverride(ItemSize)
+			.Visibility(EVisibility::Visible)
 			[
 				Button.ToSharedRef()
 			];
+		// clang-format on
 	}
 
 protected:
-	bool CanMuteLayer( FUsdLayerViewModelRef LayerItem ) const
+	bool CanMuteLayer(FUsdLayerViewModelRef LayerItem) const
 	{
 		return LayerItem->IsValid() && LayerItem->CanMuteLayer();
 	}
 
-	void ToggleMuteLayer( FUsdLayerViewModelRef LayerItem )
+	void ToggleMuteLayer(FUsdLayerViewModelRef LayerItem)
 	{
-		if ( !CanMuteLayer( LayerItem ) )
+		if (!CanMuteLayer(LayerItem))
 		{
 			return;
 		}
 
 		// If the layer is not dirty we can just mute it without worry and early out
-		if ( !LayerItem->IsLayerDirty() )
+		if (!LayerItem->IsLayerDirty())
 		{
 			LayerItem->ToggleMuteLayer();
 			return;
@@ -233,16 +228,19 @@ protected:
 
 		// Show a warning if the layer is dirty, as muting it will discard changes
 		const UUsdProjectSettings* Settings = GetDefault<UUsdProjectSettings>();
-		if ( Settings && Settings->bShowConfirmationWhenMutingDirtyLayers )
+		if (Settings && Settings->bShowConfirmationWhenMutingDirtyLayers)
 		{
 			static TWeakPtr<SNotificationItem> Notification;
 
-			FNotificationInfo Toast( LOCTEXT( "ConfirmMutingLayer", "Muting dirty layer" ) );
+			FNotificationInfo Toast(LOCTEXT("ConfirmMutingLayer", "Muting dirty layer"));
 			Toast.SubText = FText::Format(
-				LOCTEXT( "ConfirmMutingLayer_Subtext", "Layer '{0}' has unsaved changes that will be lost if muted.\n\nDo you wish to proceed muting the layer?" ),
+				LOCTEXT(
+					"ConfirmMutingLayer_Subtext",
+					"Layer '{0}' has unsaved changes that will be lost if muted.\n\nDo you wish to proceed muting the layer?"
+				),
 				LayerItem->GetDisplayName()
 			);
-			Toast.Image = FCoreStyle::Get().GetBrush( TEXT( "MessageLog.Warning" ) );
+			Toast.Image = FCoreStyle::Get().GetBrush(TEXT("MessageLog.Warning"));
 			Toast.bUseLargeFont = false;
 			Toast.bFireAndForget = false;
 			Toast.FadeOutDuration = 0.0f;
@@ -251,55 +249,67 @@ protected:
 			Toast.bUseSuccessFailIcons = false;
 
 			Toast.ButtonDetails.Emplace(
-				LOCTEXT( "ConfirmMutingOkAll", "Always proceed" ),
+				LOCTEXT("ConfirmMutingOkAll", "Always proceed"),
 				FText::GetEmpty(),
-				FSimpleDelegate::CreateLambda( [LayerItem]() {
-				if ( TSharedPtr<SNotificationItem> PinnedNotification = Notification.Pin() )
-				{
-					PinnedNotification->SetCompletionState( SNotificationItem::CS_Success );
-					PinnedNotification->ExpireAndFadeout();
+				FSimpleDelegate::CreateLambda(
+					[LayerItem]()
+					{
+						if (TSharedPtr<SNotificationItem> PinnedNotification = Notification.Pin())
+						{
+							PinnedNotification->SetCompletionState(SNotificationItem::CS_Success);
+							PinnedNotification->ExpireAndFadeout();
 
-					UUsdProjectSettings* Settings = GetMutableDefault<UUsdProjectSettings>();
-					Settings->bShowConfirmationWhenMutingDirtyLayers = false;
-					Settings->SaveConfig();
+							UUsdProjectSettings* Settings = GetMutableDefault<UUsdProjectSettings>();
+							Settings->bShowConfirmationWhenMutingDirtyLayers = false;
+							Settings->SaveConfig();
 
-					LayerItem->ToggleMuteLayer();
-				}
-			} ) );
+							LayerItem->ToggleMuteLayer();
+						}
+					}
+				)
+			);
 
 			Toast.ButtonDetails.Emplace(
-				LOCTEXT( "ConfirmMutingOk", "Proceed" ),
+				LOCTEXT("ConfirmMutingOk", "Proceed"),
 				FText::GetEmpty(),
-				FSimpleDelegate::CreateLambda( [LayerItem]() {
-				if ( TSharedPtr<SNotificationItem> PinnedNotification = Notification.Pin() )
-				{
-					PinnedNotification->SetCompletionState( SNotificationItem::CS_Success );
-					PinnedNotification->ExpireAndFadeout();
+				FSimpleDelegate::CreateLambda(
+					[LayerItem]()
+					{
+						if (TSharedPtr<SNotificationItem> PinnedNotification = Notification.Pin())
+						{
+							PinnedNotification->SetCompletionState(SNotificationItem::CS_Success);
+							PinnedNotification->ExpireAndFadeout();
 
-					LayerItem->ToggleMuteLayer();
-				}
-			} ) );
+							LayerItem->ToggleMuteLayer();
+						}
+					}
+				)
+			);
 
 			Toast.ButtonDetails.Emplace(
-				LOCTEXT( "ConfirmMutingCancel", "Cancel" ),
+				LOCTEXT("ConfirmMutingCancel", "Cancel"),
 				FText::GetEmpty(),
-				FSimpleDelegate::CreateLambda( []() {
-				if ( TSharedPtr<SNotificationItem> PinnedNotification = Notification.Pin() )
-				{
-					PinnedNotification->SetCompletionState( SNotificationItem::CS_Fail );
-					PinnedNotification->ExpireAndFadeout();
-				}
-			} ) );
+				FSimpleDelegate::CreateLambda(
+					[]()
+					{
+						if (TSharedPtr<SNotificationItem> PinnedNotification = Notification.Pin())
+						{
+							PinnedNotification->SetCompletionState(SNotificationItem::CS_Fail);
+							PinnedNotification->ExpireAndFadeout();
+						}
+					}
+				)
+			);
 
 			// Only show one at a time
-			if ( !Notification.IsValid() )
+			if (!Notification.IsValid())
 			{
-				Notification = FSlateNotificationManager::Get().AddNotification( Toast );
+				Notification = FSlateNotificationManager::Get().AddNotification(Toast);
 			}
 
-			if ( TSharedPtr<SNotificationItem> PinnedNotification = Notification.Pin() )
+			if (TSharedPtr<SNotificationItem> PinnedNotification = Notification.Pin())
 			{
-				PinnedNotification->SetCompletionState( SNotificationItem::CS_Pending );
+				PinnedNotification->SetCompletionState(SNotificationItem::CS_Pending);
 			}
 		}
 		else
@@ -310,17 +320,19 @@ protected:
 	}
 };
 
-class FUsdLayerEditColumn : public FUsdTreeViewColumn, public TSharedFromThis< FUsdLayerEditColumn >
+class FUsdLayerEditColumn
+	: public FUsdTreeViewColumn
+	, public TSharedFromThis<FUsdLayerEditColumn>
 {
 public:
-	FReply OnClicked( const FUsdLayerViewModelRef TreeItem )
+	FReply OnClicked(const FUsdLayerViewModelRef TreeItem)
 	{
 		return TreeItem->EditLayer() ? FReply::Handled() : FReply::Unhandled();
 	}
 
-	FSlateColor GetForegroundColor( const FUsdLayerViewModelRef TreeItem, const TSharedPtr< ITableRow > TableRow, const TSharedPtr< SButton > Button ) const
+	FSlateColor GetForegroundColor(const FUsdLayerViewModelRef TreeItem, const TSharedPtr<ITableRow> TableRow, const TSharedPtr<SButton> Button) const
 	{
-		if ( !TableRow.IsValid() || !Button.IsValid() )
+		if (!TableRow.IsValid() || !Button.IsValid())
 		{
 			return FSlateColor::UseForeground();
 		}
@@ -328,95 +340,99 @@ public:
 		const bool bIsButtonHovered = Button->IsHovered();
 		const bool bIsLayerEditTarget = TreeItem->IsEditTarget();
 
-		if ( bIsLayerEditTarget )
+		if (bIsLayerEditTarget)
 		{
 			return FSlateColor::UseForeground();
 		}
 		else
 		{
-			if ( !bIsButtonHovered )
+			if (!bIsButtonHovered)
 			{
 				return FLinearColor::Transparent;
 			}
 			else
 			{
-				return FAppStyle::GetSlateColor( TEXT( "Colors.ForegroundHover" ) );
+				return FAppStyle::GetSlateColor(TEXT("Colors.ForegroundHover"));
 			}
 		}
 	}
 
-	virtual TSharedRef< SWidget > GenerateWidget( const TSharedPtr< IUsdTreeViewItem > InTreeItem, const TSharedPtr< ITableRow > TableRow ) override
+	virtual TSharedRef<SWidget> GenerateWidget(const TSharedPtr<IUsdTreeViewItem> InTreeItem, const TSharedPtr<ITableRow> TableRow) override
 	{
-		const FUsdLayerViewModelRef TreeItem = StaticCastSharedRef< FUsdLayerViewModel >( InTreeItem.ToSharedRef() );
+		const FUsdLayerViewModelRef TreeItem = StaticCastSharedRef<FUsdLayerViewModel>(InTreeItem.ToSharedRef());
 		FUsdLayerViewModelWeak TreeItemWeak = TreeItem;
 
-		float ItemSize = FUsdStageEditorStyle::Get()->GetFloat( "UsdStageEditor.ListItemHeight" );
+		float ItemSize = FUsdStageEditorStyle::Get()->GetFloat("UsdStageEditor.ListItemHeight");
 
-		TSharedPtr<SButton> Button = SNew( SButton )
-			.ContentPadding( 0 )
+		// clang-format off
+		TSharedPtr<SButton> Button = SNew(SButton)
+			.ContentPadding(0)
 			.IsEnabled_Lambda([TreeItemWeak]()->bool
 			{
-				if ( TSharedPtr<FUsdLayerViewModel> PinnedTreeItem = TreeItemWeak.Pin() )
+				if (TSharedPtr<FUsdLayerViewModel> PinnedTreeItem = TreeItemWeak.Pin())
 				{
 					return PinnedTreeItem->IsInIsolatedStage();
 				}
 				return false;
 			})
-			.ButtonStyle( FUsdStageEditorStyle::Get(), TEXT( "NoBorder" ) )
-			.OnClicked( this, &FUsdLayerEditColumn::OnClicked, TreeItem )
-			.ToolTip( SNew( SToolTip ).Text( LOCTEXT( "EditLayerButtonToolTip", "Edit layer" ) ) )
-			.HAlign( HAlign_Center )
-			.VAlign( VAlign_Center );
-
-		TSharedPtr<SImage> Image = SNew( SImage )
-			.Image( FUsdStageEditorStyle::Get()->GetBrush( "UsdStageEditor.CheckBoxImage" ) )
-			.ColorAndOpacity( this, &FUsdLayerEditColumn::GetForegroundColor, TreeItem, TableRow, Button );
-
-		Button->SetContent( Image.ToSharedRef() );
-
-		return SNew( SBox )
-			.HeightOverride( ItemSize )
-			.WidthOverride( ItemSize )
-			.HAlign( HAlign_Center )
-			.VAlign( VAlign_Center )
+			.ButtonStyle(FUsdStageEditorStyle::Get(), TEXT("NoBorder"))
+			.OnClicked(this, &FUsdLayerEditColumn::OnClicked, TreeItem)
+			.ToolTip(SNew(SToolTip).Text(LOCTEXT("EditLayerButtonToolTip", "Edit layer")))
+			.HAlign(HAlign_Center)
+			.VAlign(VAlign_Center);
+
+		TSharedPtr<SImage> Image = SNew(SImage)
+			.Image(FUsdStageEditorStyle::Get()->GetBrush("UsdStageEditor.CheckBoxImage"))
+			.ColorAndOpacity(this, &FUsdLayerEditColumn::GetForegroundColor, TreeItem, TableRow, Button);
+
+		Button->SetContent(Image.ToSharedRef());
+
+		return SNew(SBox)
+			.HeightOverride(ItemSize)
+			.WidthOverride(ItemSize)
+			.HAlign(HAlign_Center)
+			.VAlign(VAlign_Center)
 			[
 				Button.ToSharedRef()
 			];
+		// clang-format on
 	}
 };
 
-void SUsdLayersTreeView::Construct( const FArguments& InArgs )
+void SUsdLayersTreeView::Construct(const FArguments& InArgs)
 {
-	SUsdTreeView::Construct( SUsdTreeView::FArguments() );
+	SUsdTreeView::Construct(SUsdTreeView::FArguments());
 
 	LayerIsolatedDelegate = InArgs._OnLayerIsolated;
 
-	OnContextMenuOpening = FOnContextMenuOpening::CreateSP( this, &SUsdLayersTreeView::ConstructLayerContextMenu );
+	OnContextMenuOpening = FOnContextMenuOpening::CreateSP(this, &SUsdLayersTreeView::ConstructLayerContextMenu);
 
-	OnExpansionChanged = FOnExpansionChanged::CreateLambda( [this]( const FUsdLayerViewModelRef& LayerViewModel, bool bIsExpanded )
-	{
-		if ( !LayerViewModel->IsValid() )
+	OnExpansionChanged = FOnExpansionChanged::CreateLambda(
+		[this](const FUsdLayerViewModelRef& LayerViewModel, bool bIsExpanded)
 		{
-			return;
-		}
+			if (!LayerViewModel->IsValid())
+			{
+				return;
+			}
 
-		TreeItemExpansionStates.Add( LayerViewModel->LayerIdentifier, bIsExpanded );
-	} );
+			TreeItemExpansionStates.Add(LayerViewModel->LayerIdentifier, bIsExpanded);
+		}
+	);
 }
 
-void SUsdLayersTreeView::Refresh( const UE::FUsdStageWeak& NewStage, const UE::FUsdStageWeak& InIsolatedStage, bool bResync )
+void SUsdLayersTreeView::Refresh(const UE::FUsdStageWeak& NewStage, const UE::FUsdStageWeak& InIsolatedStage, bool bResync)
 {
-	if ( bResync )
+	if (bResync)
 	{
 		bool bShouldResetExpansionStates = false;
-		if ( RootItems.Num() > 0 )
+		if (RootItems.Num() > 0)
 		{
-			const FUsdLayerViewModelRef& FirstRootItem = RootItems[ 0 ];
+			const FUsdLayerViewModelRef& FirstRootItem = RootItems[0];
 			const UE::FUsdStageWeak& OldStage = FirstRootItem->UsdStage;
 			bShouldResetExpansionStates = NewStage != OldStage;
 		}
 
-		if ( bShouldResetExpansionStates )
+		if (bShouldResetExpansionStates)
 		{
 			TreeItemExpansionStates.Reset();
 		}
@@ -430,7 +446,7 @@ void SUsdLayersTreeView::Refresh( const UE::FUsdStageWeak& NewStage, const UE::F
 	}
 	else
 	{
-		for ( FUsdLayerViewModelRef TreeItem :  RootItems )
+		for (FUsdLayerViewModelRef TreeItem : RootItems)
 		{
 			TreeItem->RefreshData();
 		}
@@ -439,45 +455,44 @@ void SUsdLayersTreeView::Refresh( const UE::FUsdStageWeak& NewStage, const UE::F
 	RequestTreeRefresh();
 }
 
-FReply SUsdLayersTreeView::OnRowDragDetected( const FGeometry& Geometry, const FPointerEvent& PointerEvent )
+FReply SUsdLayersTreeView::OnRowDragDetected(const FGeometry& Geometry, const FPointerEvent& PointerEvent)
 {
 	TArray<FUsdLayerViewModelRef> Items = GetSelectedItems();
 
 	TSet<TSharedRef<IUsdTreeViewItem>> DraggedItems;
-	DraggedItems.Reserve( Items.Num() );
+	DraggedItems.Reserve(Items.Num());
 
-	for ( const FUsdLayerViewModelRef& Item : Items )
+	for (const FUsdLayerViewModelRef& Item : Items)
 	{
 		// Don't let a drag operation begin if we're dragging the root or session layer
-		if ( Item->ParentItem == nullptr )
+		if (Item->ParentItem == nullptr)
 		{
 			return FReply::Unhandled();
 		}
 
-		DraggedItems.Add( Item );
+		DraggedItems.Add(Item);
 	}
 
-	const FSlateBrush* Icon = FAppStyle::GetBrush( TEXT( "Layer.Icon16x" ) );
-	FText DefaultHoverText = DraggedItems.Num() == 1
-		? FText::Format(
-			LOCTEXT( "DefaultTextSingle", "USD layer '{0}'" ),
-			FText::FromString( StaticCastSharedRef<FUsdLayerViewModel>( *DraggedItems.CreateConstIterator() )->LayerIdentifier )
-		)
-		: FText::Format( LOCTEXT( "DefaultTextMultiple", "{0} USD layers" ), DraggedItems.Num() );
+	const FSlateBrush* Icon = FAppStyle::GetBrush(TEXT("Layer.Icon16x"));
+	FText DefaultHoverText = DraggedItems.Num() == 1 ? FText::Format(
+								 LOCTEXT("DefaultTextSingle", "USD layer '{0}'"),
+								 FText::FromString(StaticCastSharedRef<FUsdLayerViewModel>(*DraggedItems.CreateConstIterator())->LayerIdentifier)
+							 )
+													 : FText::Format(LOCTEXT("DefaultTextMultiple", "{0} USD layers"), DraggedItems.Num());
 
 	TSharedRef<FUsdDragDropOp> Op = MakeShared<FUsdDragDropOp>();
 	Op->OpType = EUsdDragDropOpType::Layers;
 	Op->DraggedItems = DraggedItems;
-	Op->SetToolTip( DefaultHoverText, Icon );
+	Op->SetToolTip(DefaultHoverText, Icon);
 	Op->SetupDefaults();
-	Op->Construct();  // Required to initialize the little window that shows tooltips
+	Op->Construct();	// Required to initialize the little window that shows tooltips
 
-	return FReply::Handled().BeginDragDrop( Op );
+	return FReply::Handled().BeginDragDrop(Op);
 }
 
-void SUsdLayersTreeView::OnRowDragLeave( const FDragDropEvent& Event )
+void SUsdLayersTreeView::OnRowDragLeave(const FDragDropEvent& Event)
 {
-	if ( TSharedPtr<FUsdDragDropOp> Op = Event.GetOperationAs<FUsdDragDropOp>() )
+	if (TSharedPtr<FUsdDragDropOp> Op = Event.GetOperationAs<FUsdDragDropOp>())
 	{
 		Op->ResetToDefaultToolTip();
 	}
@@ -486,21 +501,15 @@ void SUsdLayersTreeView::OnRowDragLeave( const FDragDropEvent& Event )
 namespace UE::USDLayersTreeViewImpl::Private
 {
 	// Checks if we can add all of our dragged layers as sublayers to Parent
-	UsdUtils::ECanInsertSublayerResult CanAddAsSubLayers(
-		const UE::FSdfLayer& Parent,
-		const TSet<TSharedRef<IUsdTreeViewItem>>& SubLayerItems
-	)
+	UsdUtils::ECanInsertSublayerResult CanAddAsSubLayers(const UE::FSdfLayer& Parent, const TSet<TSharedRef<IUsdTreeViewItem>>& SubLayerItems)
 	{
-		for ( const TSharedRef<IUsdTreeViewItem>& Item : SubLayerItems )
+		for (const TSharedRef<IUsdTreeViewItem>& Item : SubLayerItems)
 		{
-			const FUsdLayerViewModelRef& SubLayerItem = StaticCastSharedRef<FUsdLayerViewModel>( Item );
+			const FUsdLayerViewModelRef& SubLayerItem = StaticCastSharedRef<FUsdLayerViewModel>(Item);
 
-			UsdUtils::ECanInsertSublayerResult Result = UsdUtils::CanInsertSubLayer(
-				Parent,
-				*SubLayerItem->LayerIdentifier
-			);
+			UsdUtils::ECanInsertSublayerResult Result = UsdUtils::CanInsertSubLayer(Parent, *SubLayerItem->LayerIdentifier);
 
-			if ( Result != UsdUtils::ECanInsertSublayerResult::Success )
+			if (Result != UsdUtils::ECanInsertSublayerResult::Success)
 			{
 				return Result;
 			}
@@ -511,24 +520,24 @@ namespace UE::USDLayersTreeViewImpl::Private
 
 	// Checks if this operation would effectively do nothing
 	bool IsNoOp(
-		const TArray< FUsdLayerViewModelRef >& ExistingItems,
-		const TSet< TSharedRef<IUsdTreeViewItem> >& DraggedItems,
+		const TArray<FUsdLayerViewModelRef>& ExistingItems,
+		const TSet<TSharedRef<IUsdTreeViewItem>>& DraggedItems,
 		int32 TargetItemIndex,
 		EItemDropZone DropZone
 	)
 	{
-		switch ( DropZone )
+		switch (DropZone)
 		{
 			case EItemDropZone::AboveItem:
 			{
-				ensure( TargetItemIndex != INDEX_NONE );
+				ensure(TargetItemIndex != INDEX_NONE);
 
 				// If dropped, we'll add the layers before the target index within ExistingItems, so here we just have
 				// to check if the previous items of ExistingItems (before the target) are the ones we're dragging
 
 				// We have dragged more items than there are items after the target, so there is
 				// no way this can be a no-op
-				if ( DraggedItems.Num() > TargetItemIndex )
+				if (DraggedItems.Num() > TargetItemIndex)
 				{
 					return false;
 				}
@@ -536,13 +545,13 @@ namespace UE::USDLayersTreeViewImpl::Private
 				// See if all the items immediately before the target are the same as the ones we're
 				// dragging
 				TSet<FUsdLayerViewModelRef> UniqueItemsBeforeTarget;
-				for ( int32 Index = TargetItemIndex - 1; Index >= 0; --Index )
+				for (int32 Index = TargetItemIndex - 1; Index >= 0; --Index)
 				{
-					const FUsdLayerViewModelRef& OtherItem = ExistingItems[ Index ];
+					const FUsdLayerViewModelRef& OtherItem = ExistingItems[Index];
 
-					if ( DraggedItems.Contains( OtherItem ) )
+					if (DraggedItems.Contains(OtherItem))
 					{
-						UniqueItemsBeforeTarget.Add( OtherItem );
+						UniqueItemsBeforeTarget.Add(OtherItem);
 					}
 					else
 					{
@@ -554,7 +563,7 @@ namespace UE::USDLayersTreeViewImpl::Private
 				// This means there is another item before the target that is *not* something we dragged:
 				// That means we're dragging something that is not immediately after the target, so
 				// this operation would actually do something
-				if ( UniqueItemsBeforeTarget.Num() < DraggedItems.Num() )
+				if (UniqueItemsBeforeTarget.Num() < DraggedItems.Num())
 				{
 					return false;
 				}
@@ -562,14 +571,14 @@ namespace UE::USDLayersTreeViewImpl::Private
 			}
 			case EItemDropZone::BelowItem:
 			{
-				ensure( TargetItemIndex != INDEX_NONE );
+				ensure(TargetItemIndex != INDEX_NONE);
 
 				// If dropped, we'll add the layers after the target index within ExistingItems, so here we just have
 				// to check if the next items of ExistingItems (after the target) are the ones we're dragging
 
 				// We have dragged more items than there are items after the target, so there is
 				// no way this can be a no-op
-				if ( DraggedItems.Num() > ExistingItems.Num() - TargetItemIndex - 1 )
+				if (DraggedItems.Num() > ExistingItems.Num() - TargetItemIndex - 1)
 				{
 					return false;
 				}
@@ -577,13 +586,13 @@ namespace UE::USDLayersTreeViewImpl::Private
 				// See if all the items immediately after the target are the same as the ones we're
 				// dragging
 				TSet<FUsdLayerViewModelRef> UniqueItemsAfterTarget;
-				for ( int32 Index = TargetItemIndex + 1; Index < ExistingItems.Num(); ++Index )
+				for (int32 Index = TargetItemIndex + 1; Index < ExistingItems.Num(); ++Index)
 				{
-					const FUsdLayerViewModelRef& OtherItem = ExistingItems[ Index ];
+					const FUsdLayerViewModelRef& OtherItem = ExistingItems[Index];
 
-					if ( DraggedItems.Contains( OtherItem ) )
+					if (DraggedItems.Contains(OtherItem))
 					{
-						UniqueItemsAfterTarget.Add( OtherItem );
+						UniqueItemsAfterTarget.Add(OtherItem);
 					}
 					else
 					{
@@ -595,7 +604,7 @@ namespace UE::USDLayersTreeViewImpl::Private
 				// This means there is another item after the target that is *not* something we dragged:
 				// That means we're dragging something that is not immediately after the target, so
 				// this operation would actually do something
-				if ( UniqueItemsAfterTarget.Num() < DraggedItems.Num() )
+				if (UniqueItemsAfterTarget.Num() < DraggedItems.Num())
 				{
 					return false;
 				}
@@ -606,20 +615,20 @@ namespace UE::USDLayersTreeViewImpl::Private
 				// If dropped, we add the layers at then end of the sublayer list, so here we just
 				// have to check if the last items of the sublayer list are the ones we're dragging
 
-				if ( DraggedItems.Num() > ExistingItems.Num() )
+				if (DraggedItems.Num() > ExistingItems.Num())
 				{
 					return false;
 				}
 
 				// See if the last ExistingItems are the same as the ones we're dragging
 				TSet<FUsdLayerViewModelRef> UniqueExistingItemsTail;
-				for ( int32 Index = ExistingItems.Num() - 1; Index >= 0; --Index )
+				for (int32 Index = ExistingItems.Num() - 1; Index >= 0; --Index)
 				{
-					const FUsdLayerViewModelRef& OtherItem = ExistingItems[ Index ];
+					const FUsdLayerViewModelRef& OtherItem = ExistingItems[Index];
 
-					if ( DraggedItems.Contains( OtherItem ) )
+					if (DraggedItems.Contains(OtherItem))
 					{
-						UniqueExistingItemsTail.Add( OtherItem );
+						UniqueExistingItemsTail.Add(OtherItem);
 					}
 					else
 					{
@@ -631,7 +640,7 @@ namespace UE::USDLayersTreeViewImpl::Private
 				// This means there is another item at the end of ExistingItems that is *not* something we dragged:
 				// That means we're dragging something that is not already at the end of ExistingITems, and so that
 				// this operation would actually do something
-				if ( UniqueExistingItemsTail.Num() < DraggedItems.Num() )
+				if (UniqueExistingItemsTail.Num() < DraggedItems.Num())
 				{
 					return false;
 				}
@@ -647,55 +656,52 @@ namespace UE::USDLayersTreeViewImpl::Private
 
 		return true;
 	}
-}
+}	 // namespace UE::USDLayersTreeViewImpl::Private
 
-TOptional<EItemDropZone> SUsdLayersTreeView::OnRowCanAcceptDrop( const FDragDropEvent& Event, EItemDropZone Zone, FUsdLayerViewModelRef Item )
+TOptional<EItemDropZone> SUsdLayersTreeView::OnRowCanAcceptDrop(const FDragDropEvent& Event, EItemDropZone Zone, FUsdLayerViewModelRef Item)
 {
 	using namespace UE::USDLayersTreeViewImpl::Private;
 
 	TOptional<EItemDropZone> Result;
-	if ( TSharedPtr<FUsdDragDropOp> Op = Event.GetOperationAs<FUsdDragDropOp>() )
+	if (TSharedPtr<FUsdDragDropOp> Op = Event.GetOperationAs<FUsdDragDropOp>())
 	{
 		FText NewHoverText;
 		const FSlateBrush* NewHoverIcon = nullptr;
 		const int32 NumItems = Op->DraggedItems.Num();
 		bool bResetToDefaultAndBlockOp = false;
 
-		if ( Op->OpType == EUsdDragDropOpType::Layers && NumItems > 0 )
+		if (Op->OpType == EUsdDragDropOpType::Layers && NumItems > 0)
 		{
 			// Don't accept dropping onto one of the currently dragged layers
-			for ( const TSharedRef<IUsdTreeViewItem>& DraggedItem : Op->DraggedItems )
+			for (const TSharedRef<IUsdTreeViewItem>& DraggedItem : Op->DraggedItems)
 			{
-				if ( StaticCastSharedRef<FUsdLayerViewModel>( DraggedItem ) == Item )
+				if (StaticCastSharedRef<FUsdLayerViewModel>(DraggedItem) == Item)
 				{
 					bResetToDefaultAndBlockOp = true;
 					break;
 				}
 			}
 
-			if ( !bResetToDefaultAndBlockOp )
+			if (!bResetToDefaultAndBlockOp)
 			{
 				// Special case when hovering the root/session layer. We can only add as sublayers here, as we don't
 				// support swapping root/session layers or having multiple of them
 				// (USD doesn't either... that would be more like reopening a stage?)
-				if ( Item->ParentItem == nullptr )
+				if (Item->ParentItem == nullptr)
 				{
 					Zone = EItemDropZone::OntoItem;
 				}
 
 				FUsdLayerViewModel* ParentItem = nullptr;
 				FText MessageEnd;
-				switch ( Zone )
+				switch (Zone)
 				{
 					case EItemDropZone::AboveItem:
 					{
 						ParentItem = Item->ParentItem;
 						MessageEnd = FText::Format(
-							LOCTEXT(
-								"DropAboveMessageEnd",
-								", before '{0}'"
-							),
-							FText::FromString( FPaths::GetCleanFilename( Item->LayerIdentifier ) )
+							LOCTEXT("DropAboveMessageEnd", ", before '{0}'"),
+							FText::FromString(FPaths::GetCleanFilename(Item->LayerIdentifier))
 						);
 						break;
 					}
@@ -703,11 +709,8 @@ TOptional<EItemDropZone> SUsdLayersTreeView::OnRowCanAcceptDrop( const FDragDrop
 					{
 						ParentItem = Item->ParentItem;
 						MessageEnd = FText::Format(
-							LOCTEXT(
-								"DropBelowMessageEnd",
-								", after '{0}'"
-							),
-							FText::FromString( FPaths::GetCleanFilename( Item->LayerIdentifier ) )
+							LOCTEXT("DropBelowMessageEnd", ", after '{0}'"),
+							FText::FromString(FPaths::GetCleanFilename(Item->LayerIdentifier))
 						);
 						break;
 					}
@@ -720,20 +723,17 @@ TOptional<EItemDropZone> SUsdLayersTreeView::OnRowCanAcceptDrop( const FDragDrop
 					}
 				}
 
-				const UsdUtils::ECanInsertSublayerResult CanAddResult = CanAddAsSubLayers(
-					ParentItem->GetLayer(),
-					Op->DraggedItems
-				);
+				const UsdUtils::ECanInsertSublayerResult CanAddResult = CanAddAsSubLayers(ParentItem->GetLayer(), Op->DraggedItems);
 
 				// Find the index of the hovered item within its parent, so we know where to insert the dragged layers
 				// into. We don't need this for the EItemDropZone::OntoItem case as we always append them at the end
 				int32 ItemIndexInParent = INDEX_NONE;
-				if ( CanAddResult == UsdUtils::ECanInsertSublayerResult::Success && Zone != EItemDropZone::OntoItem )
+				if (CanAddResult == UsdUtils::ECanInsertSublayerResult::Success && Zone != EItemDropZone::OntoItem)
 				{
-					for ( int32 Index = 0; Index < ParentItem->Children.Num(); ++Index )
+					for (int32 Index = 0; Index < ParentItem->Children.Num(); ++Index)
 					{
-						const FUsdLayerViewModelRef& SiblingItem = ParentItem->Children[ Index ];
-						if ( SiblingItem == Item )
+						const FUsdLayerViewModelRef& SiblingItem = ParentItem->Children[Index];
+						if (SiblingItem == Item)
 						{
 							ItemIndexInParent = Index;
 							break;
@@ -747,22 +747,21 @@ TOptional<EItemDropZone> SUsdLayersTreeView::OnRowCanAcceptDrop( const FDragDrop
 				// We have to check for this particular error here and not within CanAddAsSubLayers because
 				// we need the actual tree view items.
 				bool bHasDuplicateLayer = false;
-				if ( CanAddResult == UsdUtils::ECanInsertSublayerResult::Success )
+				if (CanAddResult == UsdUtils::ECanInsertSublayerResult::Success)
 				{
 					TSet<FString> DraggedItemIdentifiers;
-					DraggedItemIdentifiers.Reserve( Op->DraggedItems.Num() );
-					for ( const TSharedRef<IUsdTreeViewItem>& DraggedItem : Op->DraggedItems )
+					DraggedItemIdentifiers.Reserve(Op->DraggedItems.Num());
+					for (const TSharedRef<IUsdTreeViewItem>& DraggedItem : Op->DraggedItems)
 					{
-						const FUsdLayerViewModelRef& DraggedLayerItem =
-							StaticCastSharedRef<FUsdLayerViewModel>( DraggedItem );
+						const FUsdLayerViewModelRef& DraggedLayerItem = StaticCastSharedRef<FUsdLayerViewModel>(DraggedItem);
 
-						DraggedItemIdentifiers.Add( DraggedLayerItem->LayerIdentifier );
+						DraggedItemIdentifiers.Add(DraggedLayerItem->LayerIdentifier);
 					}
 
-					for ( const FUsdLayerViewModelRef& SiblingItem : ParentItem->Children )
+					for (const FUsdLayerViewModelRef& SiblingItem : ParentItem->Children)
 					{
-						if ( !Op->DraggedItems.Contains( StaticCastSharedRef<IUsdTreeViewItem>( SiblingItem ) )
-							&& DraggedItemIdentifiers.Contains( SiblingItem->LayerIdentifier ) )
+						if (!Op->DraggedItems.Contains(StaticCastSharedRef<IUsdTreeViewItem>(SiblingItem))
+							&& DraggedItemIdentifiers.Contains(SiblingItem->LayerIdentifier))
 						{
 							bHasDuplicateLayer = true;
 							break;
@@ -770,26 +769,19 @@ TOptional<EItemDropZone> SUsdLayersTreeView::OnRowCanAcceptDrop( const FDragDrop
 					}
 				}
 
-				const FText ErrorMessage = bHasDuplicateLayer
-					? LOCTEXT( "CanDropDuplicate_Text", "Cannot add duplicate SubLayer!" )
-					: UsdUtils::ToText( CanAddResult );
-
-				const bool bIsNoOp = !bHasDuplicateLayer
-					&& CanAddResult == UsdUtils::ECanInsertSublayerResult::Success
-					&& IsNoOp(
-						ParentItem->Children,
-						Op->DraggedItems,
-						ItemIndexInParent,
-						Zone
-					);
+				const FText ErrorMessage = bHasDuplicateLayer ? LOCTEXT("CanDropDuplicate_Text", "Cannot add duplicate SubLayer!")
+															  : UsdUtils::ToText(CanAddResult);
+
+				const bool bIsNoOp = !bHasDuplicateLayer && CanAddResult == UsdUtils::ECanInsertSublayerResult::Success
+									 && IsNoOp(ParentItem->Children, Op->DraggedItems, ItemIndexInParent, Zone);
 
 				// Drag and drop would do nothing
-				if ( bIsNoOp )
+				if (bIsNoOp)
 				{
 					bResetToDefaultAndBlockOp = true;
 				}
 				// Drag and drop cannot be performed for some reason
-				else if ( bHasDuplicateLayer || CanAddResult != UsdUtils::ECanInsertSublayerResult::Success )
+				else if (bHasDuplicateLayer || CanAddResult != UsdUtils::ECanInsertSublayerResult::Success)
 				{
 					NewHoverText = FText::Format(
 						LOCTEXT(
@@ -797,136 +789,122 @@ TOptional<EItemDropZone> SUsdLayersTreeView::OnRowCanAcceptDrop( const FDragDrop
 							"Cannot add dragged {0}|plural(one=layer,other=layers) as {0}|plural(one=sublayer,other=sublayers) of '{1}': {2}"
 						),
 						NumItems,
-						FText::FromString( FPaths::GetCleanFilename( ParentItem->LayerIdentifier ) ),
+						FText::FromString(FPaths::GetCleanFilename(ParentItem->LayerIdentifier)),
 						ErrorMessage
 					);
 				}
 				// Can drop the one dragged layer
-				else if ( NumItems == 1 )
+				else if (NumItems == 1)
 				{
 					FString FirstDraggedLayer = FPaths::GetCleanFilename(
-						StaticCastSharedRef<FUsdLayerViewModel>(
-							*Op->DraggedItems.CreateConstIterator()
-						)->LayerIdentifier
+						StaticCastSharedRef<FUsdLayerViewModel>(*Op->DraggedItems.CreateConstIterator())->LayerIdentifier
 					);
 
 					NewHoverText = FText::Format(
-						LOCTEXT(
-							"CanDropSingle_Text",
-							"Add layer '{0}' as a sublayer of '{1}'{2}"
-						),
-						FText::FromString( FirstDraggedLayer ),
-						FText::FromString( FPaths::GetCleanFilename( ParentItem->LayerIdentifier ) ),
+						LOCTEXT("CanDropSingle_Text", "Add layer '{0}' as a sublayer of '{1}'{2}"),
+						FText::FromString(FirstDraggedLayer),
+						FText::FromString(FPaths::GetCleanFilename(ParentItem->LayerIdentifier)),
 						MessageEnd
 					);
 
-					Result = { Zone };
+					Result = {Zone};
 				}
 				// Can drop multiple dragged layers
 				else
 				{
 					NewHoverText = FText::Format(
-						LOCTEXT(
-							"CanDropMultiple_Text",
-							"Add {0} layers as sublayers of '{1}'{2}"
-						),
+						LOCTEXT("CanDropMultiple_Text", "Add {0} layers as sublayers of '{1}'{2}"),
 						NumItems,
-						FText::FromString( FPaths::GetCleanFilename( ParentItem->LayerIdentifier ) ),
+						FText::FromString(FPaths::GetCleanFilename(ParentItem->LayerIdentifier)),
 						MessageEnd
 					);
 
-					Result = { Zone };
+					Result = {Zone};
 				}
 			}
 		}
 		else
 		{
-			NewHoverText = LOCTEXT(
-				"CanDropUnsupported_Text",
-				"Can only drag and drop layers for now"
-			);
+			NewHoverText = LOCTEXT("CanDropUnsupported_Text", "Can only drag and drop layers for now");
 		}
 
-		if ( bResetToDefaultAndBlockOp )
+		if (bResetToDefaultAndBlockOp)
 		{
 			Op->ResetToDefaultToolTip();
 			Result.Reset();
 		}
 		else
 		{
-			NewHoverIcon = FAppStyle::GetBrush( Result.IsSet()
-				? TEXT( "Graph.ConnectorFeedback.Ok" )
-				: TEXT( "Graph.ConnectorFeedback.Error" )
-			);
+			NewHoverIcon = FAppStyle::GetBrush(Result.IsSet() ? TEXT("Graph.ConnectorFeedback.Ok") : TEXT("Graph.ConnectorFeedback.Error"));
 
-			Op->SetToolTip( NewHoverText, NewHoverIcon );
+			Op->SetToolTip(NewHoverText, NewHoverIcon);
 		}
 	}
 
 	return Result;
 }
 
-FReply SUsdLayersTreeView::OnRowAcceptDrop( const FDragDropEvent& Event, EItemDropZone Zone, FUsdLayerViewModelRef Item )
+FReply SUsdLayersTreeView::OnRowAcceptDrop(const FDragDropEvent& Event, EItemDropZone Zone, FUsdLayerViewModelRef Item)
 {
 	// This doesn't work very well USD-wise yet, and I don't know how we can possibly undo/redo layer reassignments
 	// like this. We do need a transaction though, as this may trigger the creation or deletion of UObjects
-	FScopedTransaction Transaction( LOCTEXT( "OnAcceptDropTransaction", "Drop dragged USD layers" ) );
+	FScopedTransaction Transaction(LOCTEXT("OnAcceptDropTransaction", "Drop dragged USD layers"));
 
-	if ( TSharedPtr<FUsdDragDropOp> Op = Event.GetOperationAs<FUsdDragDropOp>() )
+	if (TSharedPtr<FUsdDragDropOp> Op = Event.GetOperationAs<FUsdDragDropOp>())
 	{
 		// Only accept layers for now
 		const int32 NumItems = Op->DraggedItems.Num();
-		if ( Op->OpType == EUsdDragDropOpType::Layers && NumItems > 0 )
+		if (Op->OpType == EUsdDragDropOpType::Layers && NumItems > 0)
 		{
 			UE::FSdfChangeBlock ChangeBlock;
 
 			// Make sure our layers aren't closed by USD between us removing them and readding
 			TArray<UE::FSdfLayer> LayerPins;
-			LayerPins.Reserve( Op->DraggedItems.Num() );
+			LayerPins.Reserve(Op->DraggedItems.Num());
 
 			// Remove the dragged items from their parents before we even fetch where to insert them into as this will
 			// affect the target indices
-			for ( const TSharedRef<IUsdTreeViewItem>& DraggedItem : Op->DraggedItems )
+			for (const TSharedRef<IUsdTreeViewItem>& DraggedItem : Op->DraggedItems)
 			{
-				const FUsdLayerViewModelRef& LayerItem = StaticCastSharedRef<FUsdLayerViewModel>( DraggedItem );
+				const FUsdLayerViewModelRef& LayerItem = StaticCastSharedRef<FUsdLayerViewModel>(DraggedItem);
 				UE::FSdfLayer LayerPin = LayerItem->GetLayer();
-				LayerPins.Add( LayerPin );
+				LayerPins.Add(LayerPin);
 
 				// We are only allowed to drag off things that have parents: There would never be anywhere to add the
 				// root/session layer into as that would create loops/strange situations
 				FUsdLayerViewModel* OldParent = LayerItem->ParentItem;
-				if ( ensure( OldParent ) )
+				if (ensure(OldParent))
 				{
-					if ( UE::FSdfLayer OldParentLayer = OldParent->GetLayer() )
+					if (UE::FSdfLayer OldParentLayer = OldParent->GetLayer())
 					{
-						int32 OldIndex = OldParent->Children.IndexOfByKey( LayerItem->AsShared() );
-						ensure( OldIndex != INDEX_NONE );
+						int32 OldIndex = OldParent->Children.IndexOfByKey(LayerItem->AsShared());
+						ensure(OldIndex != INDEX_NONE);
 
-						OldParentLayer.RemoveSubLayerPath( OldIndex );
+						OldParentLayer.RemoveSubLayerPath(OldIndex);
 
 						// If we have dragged multiple items, the upcoming insertion of this item may affect all the
 						// other removals indices too, so we need to make sure our widgets are up-to-date. They won't
 						// be updated on the next iteration of this loop even if we remove the outer FSdfChangeBlock,
 						// as the resync would end up creating a new batch of widgets anyway
-						OldParent->Children.RemoveAt( OldIndex );
+						OldParent->Children.RemoveAt(OldIndex);
 					}
 				}
 			}
 
 			UE::FSdfLayer ParentLayer;
 			int32 TargetSubLayerIndex = INDEX_NONE;
-			switch ( Zone )
+			switch (Zone)
 			{
 				case EItemDropZone::AboveItem:
 				{
 					// Here we'll add the dragged items as siblings of Item, before it
 
 					FUsdLayerViewModel* ParentItem = Item->ParentItem;
-					if ( ensure( ParentItem ) )
+					if (ensure(ParentItem))
 					{
 						ParentLayer = ParentItem->GetLayer();
-						TargetSubLayerIndex = ParentItem->Children.IndexOfByKey( Item );
-						ensure( TargetSubLayerIndex != INDEX_NONE );
+						TargetSubLayerIndex = ParentItem->Children.IndexOfByKey(Item);
+						ensure(TargetSubLayerIndex != INDEX_NONE);
 					}
 					break;
 				}
@@ -935,11 +913,11 @@ FReply SUsdLayersTreeView::OnRowAcceptDrop( const FDragDropEvent& Event, EItemDr
 					// Here we'll add the dragged items as siblings of Item, after it
 
 					FUsdLayerViewModel* ParentItem = Item->ParentItem;
-					if ( ensure( ParentItem ) )
+					if (ensure(ParentItem))
 					{
 						ParentLayer = ParentItem->GetLayer();
-						TargetSubLayerIndex = ParentItem->Children.IndexOfByKey( Item );
-						ensure( TargetSubLayerIndex != INDEX_NONE );
+						TargetSubLayerIndex = ParentItem->Children.IndexOfByKey(Item);
+						ensure(TargetSubLayerIndex != INDEX_NONE);
 					}
 					break;
 				}
@@ -953,30 +931,30 @@ FReply SUsdLayersTreeView::OnRowAcceptDrop( const FDragDropEvent& Event, EItemDr
 				}
 			}
 
-			for ( const TSharedRef<IUsdTreeViewItem>& DraggedItem : Op->DraggedItems )
+			for (const TSharedRef<IUsdTreeViewItem>& DraggedItem : Op->DraggedItems)
 			{
-				const FUsdLayerViewModelRef& LayerItem = StaticCastSharedRef<FUsdLayerViewModel>( DraggedItem );
+				const FUsdLayerViewModelRef& LayerItem = StaticCastSharedRef<FUsdLayerViewModel>(DraggedItem);
 				UE::FSdfLayer LayerPin = LayerItem->GetLayer();
 
 				const bool bInserted = UsdUtils::InsertSubLayer(
 					ParentLayer,
 					*LayerPin.GetIdentifier(),
-					Zone == EItemDropZone::OntoItem
-						? -1 // Always at the end
-						: Zone == EItemDropZone::AboveItem
-							? TargetSubLayerIndex++
-							: ++TargetSubLayerIndex
+					Zone == EItemDropZone::OntoItem	   ? -1	   // Always at the end
+					: Zone == EItemDropZone::AboveItem ? TargetSubLayerIndex++
+													   : ++TargetSubLayerIndex
 				);
 
-				if ( !bInserted )
+				if (!bInserted)
 				{
-					UE_LOG( LogUsd, Warning, TEXT( "Failed to insert layer '%s' as a sublayer of '%s' at index '%d'" ),
+					UE_LOG(
+						LogUsd,
+						Warning,
+						TEXT("Failed to insert layer '%s' as a sublayer of '%s' at index '%d'"),
 						*LayerPin.GetIdentifier(),
 						*ParentLayer.GetIdentifier(),
 						TargetSubLayerIndex
 					);
 				}
-
 			}
 
 			return FReply::Handled();
@@ -990,63 +968,65 @@ TArray<UE::FSdfLayer> SUsdLayersTreeView::GetSelectedLayers() const
 {
 	TArray<UE::FSdfLayer> SelectedLayers;
 	TArray<FUsdLayerViewModelRef> SelectedViewModels = GetSelectedItems();
-	SelectedLayers.Reserve( SelectedViewModels.Num() );
+	SelectedLayers.Reserve(SelectedViewModels.Num());
 
-	for ( const FUsdLayerViewModelRef& SelectedItem : SelectedViewModels )
+	for (const FUsdLayerViewModelRef& SelectedItem : SelectedViewModels)
 	{
-		SelectedLayers.AddUnique( SelectedItem->GetLayer() );
+		SelectedLayers.AddUnique(SelectedItem->GetLayer());
 	}
 
 	return SelectedLayers;
 }
 
-void SUsdLayersTreeView::SetSelectedLayers( const TArray<UE::FSdfLayer>& NewSelection )
+void SUsdLayersTreeView::SetSelectedLayers(const TArray<UE::FSdfLayer>& NewSelection)
 {
 	TSet<FString> IdentifiersToSelect;
-	IdentifiersToSelect.Reserve( NewSelection.Num() );
-	for ( const UE::FSdfLayer& SelectedLayer : NewSelection )
+	IdentifiersToSelect.Reserve(NewSelection.Num());
+	for (const UE::FSdfLayer& SelectedLayer : NewSelection)
 	{
-		IdentifiersToSelect.Add( SelectedLayer.GetIdentifier() );
+		IdentifiersToSelect.Add(SelectedLayer.GetIdentifier());
 	}
 
 	TArray<FUsdLayerViewModelRef> ItemsToSelect;
-	ItemsToSelect.Reserve( NewSelection.Num() );
+	ItemsToSelect.Reserve(NewSelection.Num());
 
-	TFunction<void( const TArray< FUsdLayerViewModelRef >& )> Traverse;
-	Traverse = [this, &IdentifiersToSelect, &ItemsToSelect, &Traverse]( const TArray< FUsdLayerViewModelRef >& Items )
+	TFunction<void(const TArray<FUsdLayerViewModelRef>&)> Traverse;
+	Traverse = [this, &IdentifiersToSelect, &ItemsToSelect, &Traverse](const TArray<FUsdLayerViewModelRef>& Items)
 	{
-		for ( const FUsdLayerViewModelRef& Item : Items )
+		for (const FUsdLayerViewModelRef& Item : Items)
 		{
-			if ( IdentifiersToSelect.Contains( Item->LayerIdentifier ) )
+			if (IdentifiersToSelect.Contains(Item->LayerIdentifier))
 			{
-				ItemsToSelect.Add( Item );
+				ItemsToSelect.Add(Item);
 			}
 
-			Traverse( Item->Children );
+			Traverse(Item->Children);
 		}
 	};
-	Traverse( RootItems );
+	Traverse(RootItems);
 
 	Private_ClearSelection();
 
 	const bool bSelected = true;
-	SetItemSelection( ItemsToSelect, bSelected );
+	SetItemSelection(ItemsToSelect, bSelected);
 }
 
-TSharedRef< ITableRow > SUsdLayersTreeView::OnGenerateRow( FUsdLayerViewModelRef InDisplayNode, const TSharedRef< STableViewBase >& OwnerTable )
+TSharedRef<ITableRow> SUsdLayersTreeView::OnGenerateRow(FUsdLayerViewModelRef InDisplayNode, const TSharedRef<STableViewBase>& OwnerTable)
 {
-	return SNew( SUsdTreeRow< FUsdLayerViewModelRef >, InDisplayNode, OwnerTable, SharedData )
-		.OnDragDetected( this, &SUsdLayersTreeView::OnRowDragDetected )
-		.OnDragLeave( this, &SUsdLayersTreeView::OnRowDragLeave )
-		.OnCanAcceptDrop( this, &SUsdLayersTreeView::OnRowCanAcceptDrop )
-		.OnAcceptDrop( this, &SUsdLayersTreeView::OnRowAcceptDrop );
+	// clang-format off
+	return SNew(SUsdTreeRow< FUsdLayerViewModelRef >, InDisplayNode, OwnerTable, SharedData)
+		.OnDragDetected(this, &SUsdLayersTreeView::OnRowDragDetected)
+		.OnDragLeave(this, &SUsdLayersTreeView::OnRowDragLeave)
+		.OnCanAcceptDrop(this, &SUsdLayersTreeView::OnRowCanAcceptDrop)
+		.OnAcceptDrop(this, &SUsdLayersTreeView::OnRowAcceptDrop);
+	// clang-format on
 }
 
-void SUsdLayersTreeView::OnGetChildren( FUsdLayerViewModelRef InParent, TArray< FUsdLayerViewModelRef >& OutChildren ) const
+void SUsdLayersTreeView::OnGetChildren(FUsdLayerViewModelRef InParent, TArray<FUsdLayerViewModelRef>& OutChildren) const
 {
-	for ( const FUsdLayerViewModelRef& Child : InParent->GetChildren() )
+	for (const FUsdLayerViewModelRef& Child : InParent->GetChildren())
 	{
-		OutChildren.Add( Child );
+		OutChildren.Add(Child);
 	}
 }
 
@@ -1054,10 +1034,10 @@ void SUsdLayersTreeView::BuildUsdLayersEntries()
 {
 	RootItems.Empty();
 
-	if ( UsdStage )
+	if (UsdStage)
 	{
-		RootItems.Add( MakeSharedUnreal< FUsdLayerViewModel >( nullptr, UsdStage, IsolatedStage, UsdStage.GetRootLayer().GetIdentifier() ) );
-		RootItems.Add( MakeSharedUnreal< FUsdLayerViewModel >( nullptr, UsdStage, IsolatedStage, UsdStage.GetSessionLayer().GetIdentifier() ) );
+		RootItems.Add(MakeSharedUnreal<FUsdLayerViewModel>(nullptr, UsdStage, IsolatedStage, UsdStage.GetRootLayer().GetIdentifier()));
+		RootItems.Add(MakeSharedUnreal<FUsdLayerViewModel>(nullptr, UsdStage, IsolatedStage, UsdStage.GetSessionLayer().GetIdentifier()));
 	}
 }
 
@@ -1066,71 +1046,81 @@ void SUsdLayersTreeView::SetupColumns()
 	HeaderRowWidget->ClearColumns();
 
 	SHeaderRow::FColumn::FArguments LayerMutedColumnArguments;
-	LayerMutedColumnArguments.FixedWidth( 24.f );
-	LayerMutedColumnArguments.HAlignHeader( HAlign_Center );
-	LayerMutedColumnArguments.HAlignCell( HAlign_Center );
-	TSharedRef< FUsdLayerMutedColumn > LayerMutedColumn = MakeShared< FUsdLayerMutedColumn >();
-	AddColumn( TEXT("Mute"), FText(), LayerMutedColumn, LayerMutedColumnArguments );
+	LayerMutedColumnArguments.FixedWidth(24.f);
+	LayerMutedColumnArguments.HAlignHeader(HAlign_Center);
+	LayerMutedColumnArguments.HAlignCell(HAlign_Center);
+	TSharedRef<FUsdLayerMutedColumn> LayerMutedColumn = MakeShared<FUsdLayerMutedColumn>();
+	AddColumn(TEXT("Mute"), FText(), LayerMutedColumn, LayerMutedColumnArguments);
 
-	TSharedRef< FUsdLayerNameColumn > LayerNameColumn = MakeShared< FUsdLayerNameColumn >();
+	TSharedRef<FUsdLayerNameColumn> LayerNameColumn = MakeShared<FUsdLayerNameColumn>();
 	LayerNameColumn->bIsMainColumn = true;
 
-	AddColumn( TEXT("Layers"), LOCTEXT( "Layers", "Layers" ), LayerNameColumn );
+	AddColumn(TEXT("Layers"), LOCTEXT("Layers", "Layers"), LayerNameColumn);
 
-	TSharedRef< FUsdLayerEditColumn > LayerEditColumn = MakeShared< FUsdLayerEditColumn >();
-	AddColumn( TEXT("Edit"), LOCTEXT( "Edit", "Edit" ), LayerEditColumn );
+	TSharedRef<FUsdLayerEditColumn> LayerEditColumn = MakeShared<FUsdLayerEditColumn>();
+	AddColumn(TEXT("Edit"), LOCTEXT("Edit", "Edit"), LayerEditColumn);
 }
 
-TSharedPtr< SWidget > SUsdLayersTreeView::ConstructLayerContextMenu()
+TSharedPtr<SWidget> SUsdLayersTreeView::ConstructLayerContextMenu()
 {
-	TSharedRef< SWidget > MenuWidget = SNullWidget::NullWidget;
+	TSharedRef<SWidget> MenuWidget = SNullWidget::NullWidget;
 
-	FMenuBuilder LayerOptions( true, nullptr );
-	LayerOptions.BeginSection( "Layer", LOCTEXT("Layer", "Layer") );
+	FMenuBuilder LayerOptions(true, nullptr);
+	LayerOptions.BeginSection("Layer", LOCTEXT("Layer", "Layer"));
 	{
 		LayerOptions.AddMenuEntry(
-			TAttribute<FText>::Create( TAttribute<FText>::FGetter::CreateLambda( [this]()
-			{
-				if ( IsolatedStage )
+			TAttribute<FText>::Create(TAttribute<FText>::FGetter::CreateLambda(
+				[this]()
 				{
-					TArray< FUsdLayerViewModelRef > MySelectedItems = GetSelectedItems();
-					if ( MySelectedItems.Num() == 1 && IsolatedStage )
+					if (IsolatedStage)
 					{
-						const UE::FSdfLayerWeak& Layer = MySelectedItems[ 0 ]->GetLayer();
-						if ( Layer == IsolatedStage.GetRootLayer() || Layer == UsdStage.GetRootLayer() )
+						TArray<FUsdLayerViewModelRef> MySelectedItems = GetSelectedItems();
+						if (MySelectedItems.Num() == 1 && IsolatedStage)
 						{
-							return LOCTEXT( "StopIsolatingStage_Text", "Stop isolating" );
+							const UE::FSdfLayerWeak& Layer = MySelectedItems[0]->GetLayer();
+							if (Layer == IsolatedStage.GetRootLayer() || Layer == UsdStage.GetRootLayer())
+							{
+								return LOCTEXT("StopIsolatingStage_Text", "Stop isolating");
+							}
 						}
 					}
-				}
 
-				return LOCTEXT( "IsolateStage_Text", "Isolate" );
-			})),
-			TAttribute<FText>::Create( TAttribute<FText>::FGetter::CreateLambda( [this]()
-			{
-				if ( IsolatedStage )
+					return LOCTEXT("IsolateStage_Text", "Isolate");
+				}
+			)),
+			TAttribute<FText>::Create(TAttribute<FText>::FGetter::CreateLambda(
+				[this]()
 				{
-					TArray< FUsdLayerViewModelRef > MySelectedItems = GetSelectedItems();
-					if ( MySelectedItems.Num() == 1 )
+					if (IsolatedStage)
 					{
-						const UE::FSdfLayerWeak& Layer = MySelectedItems[ 0 ]->GetLayer();
-						if ( Layer == IsolatedStage.GetRootLayer() )
-						{
-							return LOCTEXT( "StopIsolatingStage_ToolTip", "Stops isolating this layer and go back to showing the full composed stage" );
-						}
-						else if ( IsolatedStage && Layer == UsdStage.GetRootLayer() )
+						TArray<FUsdLayerViewModelRef> MySelectedItems = GetSelectedItems();
+						if (MySelectedItems.Num() == 1)
 						{
-							return LOCTEXT( "LeaveIsolatedMode_ToolTip", "Stops isolating on any layer and go back to showing the full composed stage" );
+							const UE::FSdfLayerWeak& Layer = MySelectedItems[0]->GetLayer();
+							if (Layer == IsolatedStage.GetRootLayer())
+							{
+								return LOCTEXT(
+									"StopIsolatingStage_ToolTip",
+									"Stops isolating this layer and go back to showing the full composed stage"
+								);
+							}
+							else if (IsolatedStage && Layer == UsdStage.GetRootLayer())
+							{
+								return LOCTEXT(
+									"LeaveIsolatedMode_ToolTip",
+									"Stops isolating on any layer and go back to showing the full composed stage"
+								);
+							}
 						}
 					}
-				}
 
-				return LOCTEXT( "IsolateStage_ToolTip", "Isolate a stage with this layer as its root layer" );
-			})),
+					return LOCTEXT("IsolateStage_ToolTip", "Isolate a stage with this layer as its root layer");
+				}
+			)),
 			FSlateIcon(),
 			FUIAction(
-				FExecuteAction::CreateSP( this, &SUsdLayersTreeView::OnIsolateSelectedLayer ),
-				FCanExecuteAction::CreateSP( this, &SUsdLayersTreeView::CanIsolateSelectedLayer )
+				FExecuteAction::CreateSP(this, &SUsdLayersTreeView::OnIsolateSelectedLayer),
+				FCanExecuteAction::CreateSP(this, &SUsdLayersTreeView::CanIsolateSelectedLayer)
 			),
 			NAME_None,
 			EUserInterfaceActionType::Button
@@ -1141,46 +1131,48 @@ TSharedPtr< SWidget > SUsdLayersTreeView::ConstructLayerContextMenu()
 			LOCTEXT("EditLayer_ToolTip", "Sets the layer as the edit target"),
 			FSlateIcon(),
 			FUIAction(
-				FExecuteAction::CreateSP( this, &SUsdLayersTreeView::OnEditSelectedLayer ),
-				FCanExecuteAction::CreateSP( this, &SUsdLayersTreeView::CanEditSelectedLayer )
+				FExecuteAction::CreateSP(this, &SUsdLayersTreeView::OnEditSelectedLayer),
+				FCanExecuteAction::CreateSP(this, &SUsdLayersTreeView::CanEditSelectedLayer)
 			),
 			NAME_None,
 			EUserInterfaceActionType::Button
 		);
 
 		LayerOptions.AddMenuEntry(
-			LOCTEXT( "ClearLayer", "Clear" ),
-			LOCTEXT( "ClearLayer_ToolTip", "Clears this layer of all data" ),
+			LOCTEXT("ClearLayer", "Clear"),
+			LOCTEXT("ClearLayer_ToolTip", "Clears this layer of all data"),
 			FSlateIcon(),
 			FUIAction(
-				FExecuteAction::CreateSP( this, &SUsdLayersTreeView::OnClearSelectedLayers ),
-				FCanExecuteAction::CreateSP( this, &SUsdLayersTreeView::CanClearSelectedLayers )
+				FExecuteAction::CreateSP(this, &SUsdLayersTreeView::OnClearSelectedLayers),
+				FCanExecuteAction::CreateSP(this, &SUsdLayersTreeView::CanClearSelectedLayers)
 			),
 			NAME_None,
 			EUserInterfaceActionType::Button
 		);
 
 		LayerOptions.AddMenuEntry(
-			LOCTEXT( "SaveLayer", "Save" ),
-			LOCTEXT( "SaveLayer_ToolTip", "Saves the layer modifications to disk" ),
+			LOCTEXT("SaveLayer", "Save"),
+			LOCTEXT("SaveLayer_ToolTip", "Saves the layer modifications to disk"),
 			FSlateIcon(),
 			FUIAction(
-				FExecuteAction::CreateSP( this, &SUsdLayersTreeView::OnSaveSelectedLayers ),
-				FCanExecuteAction::CreateSP( this, &SUsdLayersTreeView::CanSaveSelectedLayers )
+				FExecuteAction::CreateSP(this, &SUsdLayersTreeView::OnSaveSelectedLayers),
+				FCanExecuteAction::CreateSP(this, &SUsdLayersTreeView::CanSaveSelectedLayers)
 			),
 			NAME_None,
 			EUserInterfaceActionType::Button
 		);
 
 		LayerOptions.AddMenuEntry(
-			LOCTEXT( "ExportLayer", "Export" ),
-			LOCTEXT( "Export_ToolTip", "Export the selected layers, having the exported layers reference the original stage's layers" ),
+			LOCTEXT("ExportLayer", "Export"),
+			LOCTEXT("Export_ToolTip", "Export the selected layers, having the exported layers reference the original stage's layers"),
 			FSlateIcon(),
 			FUIAction(
-				FExecuteAction::CreateLambda([this]()
-				{
-					ExportSelectedLayers();
-				}),
+				FExecuteAction::CreateLambda(
+					[this]()
+					{
+						ExportSelectedLayers();
+					}
+				),
 				FCanExecuteAction()
 			),
 			NAME_None,
@@ -1189,15 +1181,15 @@ TSharedPtr< SWidget > SUsdLayersTreeView::ConstructLayerContextMenu()
 	}
 	LayerOptions.EndSection();
 
-	LayerOptions.BeginSection( "SubLayers", LOCTEXT("SubLayers", "SubLayers") );
+	LayerOptions.BeginSection("SubLayers", LOCTEXT("SubLayers", "SubLayers"));
 	{
 		LayerOptions.AddMenuEntry(
 			LOCTEXT("AddExistingSubLayer", "Add Existing"),
 			LOCTEXT("AddExistingSubLayer_ToolTip", "Adds a sublayer from an existing file to this layer"),
 			FSlateIcon(),
 			FUIAction(
-				FExecuteAction::CreateSP( this, &SUsdLayersTreeView::OnAddSubLayer ),
-				FCanExecuteAction::CreateSP( this, &SUsdLayersTreeView::CanInsertSubLayer )
+				FExecuteAction::CreateSP(this, &SUsdLayersTreeView::OnAddSubLayer),
+				FCanExecuteAction::CreateSP(this, &SUsdLayersTreeView::CanInsertSubLayer)
 			),
 			NAME_None,
 			EUserInterfaceActionType::Button
@@ -1208,8 +1200,8 @@ TSharedPtr< SWidget > SUsdLayersTreeView::ConstructLayerContextMenu()
 			LOCTEXT("AddNewSubLayer_ToolTip", "Adds a sublayer using a new file to this layer"),
 			FSlateIcon(),
 			FUIAction(
-				FExecuteAction::CreateSP( this, &SUsdLayersTreeView::OnNewSubLayer ),
-				FCanExecuteAction::CreateSP( this, &SUsdLayersTreeView::CanInsertSubLayer )
+				FExecuteAction::CreateSP(this, &SUsdLayersTreeView::OnNewSubLayer),
+				FCanExecuteAction::CreateSP(this, &SUsdLayersTreeView::CanInsertSubLayer)
 			),
 			NAME_None,
 			EUserInterfaceActionType::Button
@@ -1220,8 +1212,8 @@ TSharedPtr< SWidget > SUsdLayersTreeView::ConstructLayerContextMenu()
 			LOCTEXT("RemoveSubLayer_ToolTip", "Removes the sublayer from its owner"),
 			FSlateIcon(),
 			FUIAction(
-				FExecuteAction::CreateSP( this, &SUsdLayersTreeView::OnRemoveSelectedLayers ),
-				FCanExecuteAction::CreateSP( this, &SUsdLayersTreeView::CanRemoveSelectedLayers )
+				FExecuteAction::CreateSP(this, &SUsdLayersTreeView::OnRemoveSelectedLayers),
+				FCanExecuteAction::CreateSP(this, &SUsdLayersTreeView::CanRemoveSelectedLayers)
 			),
 			NAME_None,
 			EUserInterfaceActionType::Button
@@ -1236,24 +1228,25 @@ TSharedPtr< SWidget > SUsdLayersTreeView::ConstructLayerContextMenu()
 
 bool SUsdLayersTreeView::CanIsolateSelectedLayer() const
 {
-	TArray< FUsdLayerViewModelRef > MySelectedItems = GetSelectedItems();
+	TArray<FUsdLayerViewModelRef> MySelectedItems = GetSelectedItems();
 
-	if ( MySelectedItems.Num() != 1 )
+	if (MySelectedItems.Num() != 1)
 	{
 		return false;
 	}
 
-	for ( FUsdLayerViewModelRef SelectedItem : MySelectedItems )
+	for (FUsdLayerViewModelRef SelectedItem : MySelectedItems)
 	{
 		const UE::FSdfLayerWeak& Layer = SelectedItem->GetLayer();
 		const bool bLayerIsStageRoot = Layer == UsdStage.GetRootLayer();
 
-		return SelectedItem->IsValid() && (
-			// If we're right clicking a sublayer its always OK
-			!bLayerIsStageRoot
-			// If we're isolating and selecting either a sublayer or the root layer (i.e. not the session layer)
-			|| ( IsolatedStage && ( SelectedItem->ParentItem || bLayerIsStageRoot ) )
-		);
+		return SelectedItem->IsValid()
+			   && (
+				   // If we're right clicking a sublayer its always OK
+				   !bLayerIsStageRoot
+				   // If we're isolating and selecting either a sublayer or the root layer (i.e. not the session layer)
+				   || (IsolatedStage && (SelectedItem->ParentItem || bLayerIsStageRoot))
+			   );
 	}
 
 	return false;
@@ -1261,31 +1254,31 @@ bool SUsdLayersTreeView::CanIsolateSelectedLayer() const
 
 void SUsdLayersTreeView::OnIsolateSelectedLayer()
 {
-	if ( !LayerIsolatedDelegate.IsBound() )
+	if (!LayerIsolatedDelegate.IsBound())
 	{
 		return;
 	}
 
-	TArray< FUsdLayerViewModelRef > MySelectedItems = GetSelectedItems();
+	TArray<FUsdLayerViewModelRef> MySelectedItems = GetSelectedItems();
 
-	if ( MySelectedItems.Num() != 1 )
+	if (MySelectedItems.Num() != 1)
 	{
 		return;
 	}
 
-	for ( FUsdLayerViewModelRef SelectedItem : MySelectedItems )
+	for (FUsdLayerViewModelRef SelectedItem : MySelectedItems)
 	{
 		UE::FSdfLayer Layer = SelectedItem->GetLayer();
 
-		if ( IsolatedStage && IsolatedStage.GetRootLayer() == Layer )
+		if (IsolatedStage && IsolatedStage.GetRootLayer() == Layer)
 		{
 			// We're already isolating this one -> Toggle isolate mode off
-			LayerIsolatedDelegate.Execute( UE::FSdfLayer{} );
+			LayerIsolatedDelegate.Execute(UE::FSdfLayer{});
 		}
 		else
 		{
 			// Isolate the provided layer
-			LayerIsolatedDelegate.Execute( Layer );
+			LayerIsolatedDelegate.Execute(Layer);
 		}
 	}
 }
@@ -1294,11 +1287,11 @@ bool SUsdLayersTreeView::CanEditSelectedLayer() const
 {
 	bool bHasEditableLayer = false;
 
-	TArray< FUsdLayerViewModelRef > MySelectedItems = GetSelectedItems();
+	TArray<FUsdLayerViewModelRef> MySelectedItems = GetSelectedItems();
 
-	for ( FUsdLayerViewModelRef SelectedItem : MySelectedItems )
+	for (FUsdLayerViewModelRef SelectedItem : MySelectedItems)
 	{
-		if ( SelectedItem->CanEditLayer() )
+		if (SelectedItem->CanEditLayer())
 		{
 			bHasEditableLayer = true;
 			break;
@@ -1310,11 +1303,11 @@ bool SUsdLayersTreeView::CanEditSelectedLayer() const
 
 void SUsdLayersTreeView::OnEditSelectedLayer()
 {
-	TArray< FUsdLayerViewModelRef > MySelectedItems = GetSelectedItems();
+	TArray<FUsdLayerViewModelRef> MySelectedItems = GetSelectedItems();
 
-	for ( FUsdLayerViewModelRef SelectedItem : MySelectedItems )
+	for (FUsdLayerViewModelRef SelectedItem : MySelectedItems)
 	{
-		if ( SelectedItem->EditLayer() )
+		if (SelectedItem->EditLayer())
 		{
 			break;
 		}
@@ -1323,35 +1316,35 @@ void SUsdLayersTreeView::OnEditSelectedLayer()
 
 void SUsdLayersTreeView::OnClearSelectedLayers()
 {
-	TArray< FUsdLayerViewModelRef > MySelectedItems = GetSelectedItems();
+	TArray<FUsdLayerViewModelRef> MySelectedItems = GetSelectedItems();
 
 	// We'll show a confirmation toast, which is non-modal. So keep track of the original selected items so that
 	// if anything changes by the time we accept the dialog we'll still know which layers to clear
 	TArray<UE::FSdfLayerWeak> SelectedLayers;
-	SelectedLayers.Reserve( MySelectedItems.Num() );
+	SelectedLayers.Reserve(MySelectedItems.Num());
 
 	FString LayerNames;
 	const FString Separator = TEXT(", ");
-	for ( FUsdLayerViewModelRef SelectedItem : MySelectedItems )
+	for (FUsdLayerViewModelRef SelectedItem : MySelectedItems)
 	{
-		if ( UE::FSdfLayer Layer = SelectedItem->GetLayer() )
+		if (UE::FSdfLayer Layer = SelectedItem->GetLayer())
 		{
-			SelectedLayers.Add( Layer );
+			SelectedLayers.Add(Layer);
 			LayerNames += Layer.GetDisplayName() + Separator;
 		}
 	}
-	LayerNames.RemoveFromEnd( Separator );
+	LayerNames.RemoveFromEnd(Separator);
 
-	TFunction<void()> ClearSelectedLayersInner = [this, SelectedLayers]()
+	TFunction<void()> ClearSelectedLayersInner = [SelectedLayers]()
 	{
 		// This doesn't work very well USD-wise yet, and I don't know how we can possibly undo/redo
 		// clearing a layer like this. We do need a transaction though, as this may trigger the creation
 		// or deletion of UObjects
-		FScopedTransaction Transaction( LOCTEXT( "ClearTransaction", "Clear selected layers" ) );
+		FScopedTransaction Transaction(LOCTEXT("ClearTransaction", "Clear selected layers"));
 
-		for ( UE::FSdfLayerWeak Layer : SelectedLayers )
+		for (UE::FSdfLayerWeak Layer : SelectedLayers)
 		{
-			if ( Layer )
+			if (Layer)
 			{
 				Layer.Clear();
 			}
@@ -1359,17 +1352,20 @@ void SUsdLayersTreeView::OnClearSelectedLayers()
 	};
 
 	const UUsdProjectSettings* Settings = GetDefault<UUsdProjectSettings>();
-	if ( Settings && Settings->bShowConfirmationWhenClearingLayers )
+	if (Settings && Settings->bShowConfirmationWhenClearingLayers)
 	{
 		static TWeakPtr<SNotificationItem> Notification;
 
-		FNotificationInfo Toast( LOCTEXT( "ConfirmClearingLayer", "Clearing cannot be undone" ) );
+		FNotificationInfo Toast(LOCTEXT("ConfirmClearingLayer", "Clearing cannot be undone"));
 		Toast.SubText = FText::Format(
-			LOCTEXT( "ConfirmClearingLayer_Subtext", "Clearing USD layers cannot be undone.\n\nDo you wish to proceed clearing {0}|plural(one=layer,other=layers) '{1}' ?" ),
+			LOCTEXT(
+				"ConfirmClearingLayer_Subtext",
+				"Clearing USD layers cannot be undone.\n\nDo you wish to proceed clearing {0}|plural(one=layer,other=layers) '{1}' ?"
+			),
 			SelectedLayers.Num(),
 			FText::FromString(LayerNames)
 		);
-		Toast.Image = FCoreStyle::Get().GetBrush( TEXT( "MessageLog.Warning" ) );
+		Toast.Image = FCoreStyle::Get().GetBrush(TEXT("MessageLog.Warning"));
 		Toast.bUseLargeFont = false;
 		Toast.bFireAndForget = false;
 		Toast.FadeOutDuration = 0.0f;
@@ -1378,55 +1374,67 @@ void SUsdLayersTreeView::OnClearSelectedLayers()
 		Toast.bUseSuccessFailIcons = false;
 
 		Toast.ButtonDetails.Emplace(
-			LOCTEXT( "ConfirmClearingLayerOkAll", "Always proceed" ),
+			LOCTEXT("ConfirmClearingLayerOkAll", "Always proceed"),
 			FText::GetEmpty(),
-			FSimpleDelegate::CreateLambda( [ClearSelectedLayersInner]() {
-			if ( TSharedPtr<SNotificationItem> PinnedNotification = Notification.Pin() )
-			{
-				PinnedNotification->SetCompletionState( SNotificationItem::CS_Success );
-				PinnedNotification->ExpireAndFadeout();
+			FSimpleDelegate::CreateLambda(
+				[ClearSelectedLayersInner]()
+				{
+					if (TSharedPtr<SNotificationItem> PinnedNotification = Notification.Pin())
+					{
+						PinnedNotification->SetCompletionState(SNotificationItem::CS_Success);
+						PinnedNotification->ExpireAndFadeout();
 
-				UUsdProjectSettings* Settings = GetMutableDefault<UUsdProjectSettings>();
-				Settings->bShowConfirmationWhenClearingLayers = false;
-				Settings->SaveConfig();
+						UUsdProjectSettings* Settings = GetMutableDefault<UUsdProjectSettings>();
+						Settings->bShowConfirmationWhenClearingLayers = false;
+						Settings->SaveConfig();
 
-				ClearSelectedLayersInner();
-			}
-		}));
+						ClearSelectedLayersInner();
+					}
+				}
+			)
+		);
 
 		Toast.ButtonDetails.Emplace(
-			LOCTEXT( "ConfirmClearingLayerOk", "Proceed" ),
+			LOCTEXT("ConfirmClearingLayerOk", "Proceed"),
 			FText::GetEmpty(),
-			FSimpleDelegate::CreateLambda( [ClearSelectedLayersInner]() {
-			if ( TSharedPtr<SNotificationItem> PinnedNotification = Notification.Pin() )
-			{
-				PinnedNotification->SetCompletionState( SNotificationItem::CS_Success );
-				PinnedNotification->ExpireAndFadeout();
+			FSimpleDelegate::CreateLambda(
+				[ClearSelectedLayersInner]()
+				{
+					if (TSharedPtr<SNotificationItem> PinnedNotification = Notification.Pin())
+					{
+						PinnedNotification->SetCompletionState(SNotificationItem::CS_Success);
+						PinnedNotification->ExpireAndFadeout();
 
-				ClearSelectedLayersInner();
-			}
-		}));
+						ClearSelectedLayersInner();
+					}
+				}
+			)
+		);
 
 		Toast.ButtonDetails.Emplace(
-			LOCTEXT( "ConfirmClearingLayerCancel", "Cancel" ),
+			LOCTEXT("ConfirmClearingLayerCancel", "Cancel"),
 			FText::GetEmpty(),
-			FSimpleDelegate::CreateLambda( []() {
-			if ( TSharedPtr<SNotificationItem> PinnedNotification = Notification.Pin() )
-			{
-				PinnedNotification->SetCompletionState( SNotificationItem::CS_Fail );
-				PinnedNotification->ExpireAndFadeout();
-			}
-		}));
+			FSimpleDelegate::CreateLambda(
+				[]()
+				{
+					if (TSharedPtr<SNotificationItem> PinnedNotification = Notification.Pin())
+					{
+						PinnedNotification->SetCompletionState(SNotificationItem::CS_Fail);
+						PinnedNotification->ExpireAndFadeout();
+					}
+				}
+			)
+		);
 
 		// Only show one at a time
-		if ( !Notification.IsValid() )
+		if (!Notification.IsValid())
 		{
-			Notification = FSlateNotificationManager::Get().AddNotification( Toast );
+			Notification = FSlateNotificationManager::Get().AddNotification(Toast);
 		}
 
-		if ( TSharedPtr<SNotificationItem> PinnedNotification = Notification.Pin() )
+		if (TSharedPtr<SNotificationItem> PinnedNotification = Notification.Pin())
 		{
-			PinnedNotification->SetCompletionState( SNotificationItem::CS_Pending );
+			PinnedNotification->SetCompletionState(SNotificationItem::CS_Pending);
 		}
 	}
 	else
@@ -1438,11 +1446,11 @@ void SUsdLayersTreeView::OnClearSelectedLayers()
 
 bool SUsdLayersTreeView::CanClearSelectedLayers() const
 {
-	TArray< FUsdLayerViewModelRef > MySelectedItems = GetSelectedItems();
+	TArray<FUsdLayerViewModelRef> MySelectedItems = GetSelectedItems();
 
-	for ( FUsdLayerViewModelRef SelectedItem : MySelectedItems )
+	for (FUsdLayerViewModelRef SelectedItem : MySelectedItems)
 	{
-		if ( !SelectedItem->GetLayer().IsEmpty() )
+		if (!SelectedItem->GetLayer().IsEmpty())
 		{
 			return true;
 		}
@@ -1453,25 +1461,25 @@ bool SUsdLayersTreeView::CanClearSelectedLayers() const
 
 void SUsdLayersTreeView::OnSaveSelectedLayers()
 {
-	TArray< FUsdLayerViewModelRef > MySelectedItems = GetSelectedItems();
+	TArray<FUsdLayerViewModelRef> MySelectedItems = GetSelectedItems();
 
-	for ( FUsdLayerViewModelRef SelectedItem : MySelectedItems )
+	for (FUsdLayerViewModelRef SelectedItem : MySelectedItems)
 	{
-		if ( SelectedItem->IsLayerDirty() )
+		if (SelectedItem->IsLayerDirty())
 		{
 			const bool bForce = true;
-			SelectedItem->GetLayer().Save( bForce );
+			SelectedItem->GetLayer().Save(bForce);
 		}
 	}
 }
 
 bool SUsdLayersTreeView::CanSaveSelectedLayers() const
 {
-	TArray< FUsdLayerViewModelRef > MySelectedItems = GetSelectedItems();
+	TArray<FUsdLayerViewModelRef> MySelectedItems = GetSelectedItems();
 
-	for ( FUsdLayerViewModelRef SelectedItem : MySelectedItems )
+	for (FUsdLayerViewModelRef SelectedItem : MySelectedItems)
 	{
-		if ( SelectedItem->IsLayerDirty() )
+		if (SelectedItem->IsLayerDirty())
 		{
 			return true;
 		}
@@ -1482,20 +1490,20 @@ bool SUsdLayersTreeView::CanSaveSelectedLayers() const
 
 void SUsdLayersTreeView::ExportSelectedLayers(const FString& OutputLayerOrDirectory) const
 {
-	TArray< FUsdLayerViewModelRef > MySelectedItems = GetSelectedItems();
+	TArray<FUsdLayerViewModelRef> MySelectedItems = GetSelectedItems();
 
 	TArray<UE::FSdfLayer> LayersToExport;
-	LayersToExport.Reserve( MySelectedItems.Num() );
+	LayersToExport.Reserve(MySelectedItems.Num());
 
-	for ( FUsdLayerViewModelRef SelectedItem : MySelectedItems )
+	for (FUsdLayerViewModelRef SelectedItem : MySelectedItems)
 	{
 		UE::FSdfLayer SelectedLayer = SelectedItem->GetLayer();
-		if ( !SelectedLayer )
+		if (!SelectedLayer)
 		{
 			continue;
 		}
 
-		LayersToExport.Add( SelectedLayer );
+		LayersToExport.Add(SelectedLayer);
 	}
 
 	double StartTime = FPlatformTime::Cycles64();
@@ -1503,12 +1511,12 @@ void SUsdLayersTreeView::ExportSelectedLayers(const FString& OutputLayerOrDirect
 	FString OutputLayerOrDirectoryCopy = OutputLayerOrDirectory;
 
 	// Single layer -> Allow picking the target layer filename
-	if ( LayersToExport.Num() == 1 )
+	if (LayersToExport.Num() == 1)
 	{
-		if (OutputLayerOrDirectoryCopy.IsEmpty() )
+		if (OutputLayerOrDirectoryCopy.IsEmpty())
 		{
-			TOptional< FString > UsdFilePath = UsdUtils::BrowseUsdFile( UsdUtils::EBrowseFileMode::Save );
-			if ( !UsdFilePath.IsSet() )
+			TOptional<FString> UsdFilePath = UsdUtils::BrowseUsdFile(UsdUtils::EBrowseFileMode::Save);
+			if (!UsdFilePath.IsSet())
 			{
 				return;
 			}
@@ -1519,30 +1527,30 @@ void SUsdLayersTreeView::ExportSelectedLayers(const FString& OutputLayerOrDirect
 		StartTime = FPlatformTime::Cycles64();
 		Extension = FPaths::GetExtension(OutputLayerOrDirectoryCopy);
 
-		UE::USDLayersTreeViewImpl::Private::ExportLayerToPath( LayersToExport[ 0 ], OutputLayerOrDirectoryCopy);
+		UE::USDLayersTreeViewImpl::Private::ExportLayerToPath(LayersToExport[0], OutputLayerOrDirectoryCopy);
 	}
 	// Multiple layers -> Pick folder and export them with the same name
-	else if ( LayersToExport.Num() > 1 )
+	else if (LayersToExport.Num() > 1)
 	{
 		IDesktopPlatform* DesktopPlatform = FDesktopPlatformModule::Get();
-		if ( !DesktopPlatform )
+		if (!DesktopPlatform)
 		{
 			return;
 		}
 
-		TSharedPtr< SWindow > ParentWindow = FSlateApplication::Get().FindWidgetWindow( AsShared() );
-		void* ParentWindowHandle = ( ParentWindow.IsValid() && ParentWindow->GetNativeWindow().IsValid() )
-			? ParentWindow->GetNativeWindow()->GetOSWindowHandle()
-			: nullptr;
+		TSharedPtr<SWindow> ParentWindow = FSlateApplication::Get().FindWidgetWindow(AsShared());
+		void* ParentWindowHandle = (ParentWindow.IsValid() && ParentWindow->GetNativeWindow().IsValid())
+									   ? ParentWindow->GetNativeWindow()->GetOSWindowHandle()
+									   : nullptr;
 
-		if (OutputLayerOrDirectoryCopy.IsEmpty() )
+		if (OutputLayerOrDirectoryCopy.IsEmpty())
 		{
-			if ( !DesktopPlatform->OpenDirectoryDialog(
-				ParentWindowHandle,
-				LOCTEXT( "ChooseFolder", "Choose output folder" ).ToString(),
-				TEXT( "" ),
-				OutputLayerOrDirectoryCopy
-			))
+			if (!DesktopPlatform->OpenDirectoryDialog(
+					ParentWindowHandle,
+					LOCTEXT("ChooseFolder", "Choose output folder").ToString(),
+					TEXT(""),
+					OutputLayerOrDirectoryCopy
+				))
 			{
 				return;
 			}
@@ -1550,39 +1558,39 @@ void SUsdLayersTreeView::ExportSelectedLayers(const FString& OutputLayerOrDirect
 		}
 
 		StartTime = FPlatformTime::Cycles64();
-		Extension = FPaths::GetExtension( LayersToExport[ 0 ].GetDisplayName() );
+		Extension = FPaths::GetExtension(LayersToExport[0].GetDisplayName());
 
-		for ( const UE::FSdfLayer& LayerToExport : LayersToExport )
+		for (const UE::FSdfLayer& LayerToExport : LayersToExport)
 		{
-			FString TargetFileName = FPaths::GetCleanFilename( LayerToExport.GetRealPath() );
-			FString FullPath = FPaths::Combine(OutputLayerOrDirectoryCopy, TargetFileName );
+			FString TargetFileName = FPaths::GetCleanFilename(LayerToExport.GetRealPath());
+			FString FullPath = FPaths::Combine(OutputLayerOrDirectoryCopy, TargetFileName);
 			FString FinalFullPath = FullPath;
 
 			uint32 Suffix = 0;
-			while ( FPaths::FileExists( FinalFullPath ) )
+			while (FPaths::FileExists(FinalFullPath))
 			{
-				FinalFullPath = FString::Printf( TEXT( "%s_%u" ), *FullPath, Suffix++ );
+				FinalFullPath = FString::Printf(TEXT("%s_%u"), *FullPath, Suffix++);
 			}
 
-			UE::USDLayersTreeViewImpl::Private::ExportLayerToPath( LayerToExport, FinalFullPath );
+			UE::USDLayersTreeViewImpl::Private::ExportLayerToPath(LayerToExport, FinalFullPath);
 		}
 	}
 
 	// Send analytics
-	if ( FEngineAnalytics::IsAvailable() )
+	if (FEngineAnalytics::IsAvailable())
 	{
 		TArray<FAnalyticsEventAttribute> EventAttributes;
 
-		EventAttributes.Emplace( TEXT( "NumLayersExported" ), LayersToExport.Num() );
+		EventAttributes.Emplace(TEXT("NumLayersExported"), LayersToExport.Num());
 
 		bool bAutomated = false;
-		double ElapsedSeconds = FPlatformTime::ToSeconds64( FPlatformTime::Cycles64() - StartTime );
+		double ElapsedSeconds = FPlatformTime::ToSeconds64(FPlatformTime::Cycles64() - StartTime);
 		IUsdClassesModule::SendAnalytics(
-			MoveTemp( EventAttributes ),
-			TEXT( "ExportSelectedLayers" ),
+			MoveTemp(EventAttributes),
+			TEXT("ExportSelectedLayers"),
 			bAutomated,
 			ElapsedSeconds,
-			LayersToExport.Num() > 0 ? UsdUtils::GetSdfLayerNumFrames( LayersToExport[ 0 ] ) : 0,
+			LayersToExport.Num() > 0 ? UsdUtils::GetSdfLayerNumFrames(LayersToExport[0]) : 0,
 			Extension
 		);
 	}
@@ -1595,18 +1603,18 @@ bool SUsdLayersTreeView::CanInsertSubLayer() const
 
 void SUsdLayersTreeView::OnAddSubLayer()
 {
-	TOptional< FString > SubLayerFile = UsdUtils::BrowseUsdFile( UsdUtils::EBrowseFileMode::Composition );
+	TOptional<FString> SubLayerFile = UsdUtils::BrowseUsdFile(UsdUtils::EBrowseFileMode::Composition);
 
-	if ( !SubLayerFile )
+	if (!SubLayerFile)
 	{
 		return;
 	}
 
-	TArray< FUsdLayerViewModelRef > MySelectedItems = GetSelectedItems();
+	TArray<FUsdLayerViewModelRef> MySelectedItems = GetSelectedItems();
 
-	for ( FUsdLayerViewModelRef SelectedItem : MySelectedItems )
+	for (FUsdLayerViewModelRef SelectedItem : MySelectedItems)
 	{
-		SelectedItem->AddSubLayer( *SubLayerFile.GetValue() );
+		SelectedItem->AddSubLayer(*SubLayerFile.GetValue());
 		break;
 	}
 
@@ -1615,39 +1623,39 @@ void SUsdLayersTreeView::OnAddSubLayer()
 
 void SUsdLayersTreeView::OnNewSubLayer()
 {
-	TOptional< FString > SubLayerFile = UsdUtils::BrowseUsdFile( UsdUtils::EBrowseFileMode::Save );
+	TOptional<FString> SubLayerFile = UsdUtils::BrowseUsdFile(UsdUtils::EBrowseFileMode::Save);
 
-	if ( !SubLayerFile )
+	if (!SubLayerFile)
 	{
 		return;
 	}
 
-	TArray< FUsdLayerViewModelRef > MySelectedItems = GetSelectedItems();
-	for ( FUsdLayerViewModelRef SelectedItem : MySelectedItems )
+	TArray<FUsdLayerViewModelRef> MySelectedItems = GetSelectedItems();
+	for (FUsdLayerViewModelRef SelectedItem : MySelectedItems)
 	{
-		SelectedItem->NewSubLayer( *SubLayerFile.GetValue() );
+		SelectedItem->NewSubLayer(*SubLayerFile.GetValue());
 		break;
 	}
 
 	RequestTreeRefresh();
 }
 
-bool SUsdLayersTreeView::CanRemoveLayer( FUsdLayerViewModelRef LayerItem ) const
+bool SUsdLayersTreeView::CanRemoveLayer(FUsdLayerViewModelRef LayerItem) const
 {
 	// We can't remove root layers
-	return ( LayerItem->IsValid() && LayerItem->ParentItem && LayerItem->ParentItem->IsValid() );
+	return (LayerItem->IsValid() && LayerItem->ParentItem && LayerItem->ParentItem->IsValid());
 }
 
 bool SUsdLayersTreeView::CanRemoveSelectedLayers() const
 {
 	bool bHasRemovableLayer = false;
 
-	TArray< FUsdLayerViewModelRef > SelectedLayers = GetSelectedItems();
+	TArray<FUsdLayerViewModelRef> SelectedLayers = GetSelectedItems();
 
-	for ( FUsdLayerViewModelRef SelectedLayer : SelectedLayers )
+	for (FUsdLayerViewModelRef SelectedLayer : SelectedLayers)
 	{
 		// We can't remove root layers
-		if ( CanRemoveLayer( SelectedLayer ) )
+		if (CanRemoveLayer(SelectedLayer))
 		{
 			bHasRemovableLayer = true;
 			break;
@@ -1661,23 +1669,23 @@ void SUsdLayersTreeView::OnRemoveSelectedLayers()
 {
 	bool bLayerRemoved = false;
 
-	TArray< FUsdLayerViewModelRef > SelectedLayers = GetSelectedItems();
+	TArray<FUsdLayerViewModelRef> SelectedLayers = GetSelectedItems();
 
-	for ( FUsdLayerViewModelRef SelectedLayer : SelectedLayers )
+	for (FUsdLayerViewModelRef SelectedLayer : SelectedLayers)
 	{
-		if ( !CanRemoveLayer( SelectedLayer ) )
+		if (!CanRemoveLayer(SelectedLayer))
 		{
 			continue;
 		}
 
 		int32 SubLayerIndex = 0;
-		for ( FUsdLayerViewModelRef Child : SelectedLayer->ParentItem->Children )
+		for (FUsdLayerViewModelRef Child : SelectedLayer->ParentItem->Children)
 		{
-			if ( Child->LayerIdentifier == SelectedLayer->LayerIdentifier )
+			if (Child->LayerIdentifier == SelectedLayer->LayerIdentifier)
 			{
-				if ( SelectedLayer->ParentItem )
+				if (SelectedLayer->ParentItem)
 				{
-					bLayerRemoved = SelectedLayer->ParentItem->RemoveSubLayer( SubLayerIndex );
+					bLayerRemoved = SelectedLayer->ParentItem->RemoveSubLayer(SubLayerIndex);
 				}
 				break;
 			}
@@ -1686,7 +1694,7 @@ void SUsdLayersTreeView::OnRemoveSelectedLayers()
 		}
 	}
 
-	if ( bLayerRemoved )
+	if (bLayerRemoved)
 	{
 		RequestTreeRefresh();
 	}
@@ -1694,34 +1702,34 @@ void SUsdLayersTreeView::OnRemoveSelectedLayers()
 
 void SUsdLayersTreeView::RestoreExpansionStates()
 {
-	TFunction< void( const FUsdLayerViewModelRef& ) > SetExpansionRecursive = [&]( const FUsdLayerViewModelRef& Item )
+	TFunction<void(const FUsdLayerViewModelRef&)> SetExpansionRecursive = [&](const FUsdLayerViewModelRef& Item)
 	{
-		if ( Item->IsValid() )
+		if (Item->IsValid())
 		{
-			if ( bool* bFoundExpansionState = TreeItemExpansionStates.Find( Item->LayerIdentifier ) )
+			if (bool* bFoundExpansionState = TreeItemExpansionStates.Find(Item->LayerIdentifier))
 			{
-				SetItemExpansion( Item, *bFoundExpansionState );
+				SetItemExpansion(Item, *bFoundExpansionState);
 			}
 			// Default to showing everything expanded
 			else
 			{
 				const bool bShouldExpand = true;
-				SetItemExpansion( Item, bShouldExpand );
+				SetItemExpansion(Item, bShouldExpand);
 			}
 		}
 
-		for ( const FUsdLayerViewModelRef& Child : Item->Children )
+		for (const FUsdLayerViewModelRef& Child : Item->Children)
 		{
-			SetExpansionRecursive( Child );
+			SetExpansionRecursive(Child);
 		}
 	};
 
-	for ( const FUsdLayerViewModelRef& RootItem : RootItems )
+	for (const FUsdLayerViewModelRef& RootItem : RootItems)
 	{
-		SetExpansionRecursive( RootItem );
+		SetExpansionRecursive(RootItem);
 	}
 }
 
 #undef LOCTEXT_NAMESPACE
 
-#endif // #if USE_USD_SDK
+#endif	  // #if USE_USD_SDK
diff --git a/Engine/Plugins/Importers/USDImporter/Source/USDStageEditor/Private/SUSDPrimInfo.cpp b/Engine/Plugins/Importers/USDImporter/Source/USDStageEditor/Private/SUSDPrimInfo.cpp
--- a/Engine/Plugins/Importers/USDImporter/Source/USDStageEditor/Private/SUSDPrimInfo.cpp
+++ b/Engine/Plugins/Importers/USDImporter/Source/USDStageEditor/Private/SUSDPrimInfo.cpp
@@ -2,37 +2,22 @@
 
 #include "SUSDPrimInfo.h"
 
-#include "USDIntegrationUtils.h"
-#include "USDStageActor.h"
-#include "USDStageModule.h"
-#include "USDTypesConversion.h"
 #include "Widgets/SUSDIntegrationsPanel.h"
 #include "Widgets/SUSDObjectFieldList.h"
 #include "Widgets/SUSDReferencesList.h"
 #include "Widgets/SUSDVariantSetsList.h"
 
 #include "UsdWrappers/UsdStage.h"
-#include "UsdWrappers/UsdPrim.h"
-#include "UsdWrappers/SdfPath.h"
 
-#include "Styling/AppStyle.h"
-
-#include "Widgets/Views/SListView.h"
-#include "Widgets/Text/STextBlock.h"
 #include "Widgets/SBoxPanel.h"
-#include "Widgets/Layout/SSplitter.h"
-#include "Widgets/Layout/SSpacer.h"
-#include "Widgets/Layout/SSeparator.h"
-#include "Modules/ModuleManager.h"
-#include "Engine/World.h"
-#include "Algo/Find.h"
 
 #if USE_USD_SDK
 
 #define LOCTEXT_NAMESPACE "SUSDPrimInfo"
 
-void SUsdPrimInfo::Construct( const FArguments& InArgs )
+void SUsdPrimInfo::Construct(const FArguments& InArgs)
 {
+	// clang-format off
 	ChildSlot
 	[
 		SNew(SVerticalBox)
@@ -71,7 +56,7 @@ void SUsdPrimInfo::Construct( const FArguments& InArgs )
 		+SVerticalBox::Slot()
 		.AutoHeight()
 		[
-			SNew( SBox )
+			SNew(SBox)
 			.Content()
 			[
 				SAssignNew(PropertyMetadataPanel, SUsdObjectFieldList)
@@ -87,38 +72,39 @@ void SUsdPrimInfo::Construct( const FArguments& InArgs )
 		+SVerticalBox::Slot()
 		.AutoHeight()
 		[
-			SNew( SBox )
+			SNew(SBox)
 			.Content()
 			[
-				SAssignNew( IntegrationsPanel, SUsdIntegrationsPanel )
+				SAssignNew(IntegrationsPanel, SUsdIntegrationsPanel)
 			]
 		]
 
 		+SVerticalBox::Slot()
 		.AutoHeight()
 		[
-			SNew( SBox )
+			SNew(SBox)
 			.Content()
 			[
-				SAssignNew( VariantsList, SVariantsList )
+				SAssignNew(VariantsList, SVariantsList)
 			]
 		]
 
 		+SVerticalBox::Slot()
 		.AutoHeight()
 		[
-			SNew( SBox )
+			SNew(SBox)
 			.Content()
 			[
-				SAssignNew( ReferencesList, SUsdReferencesList )
+				SAssignNew(ReferencesList, SUsdReferencesList)
 			]
 		]
 	];
+	// clang-format on
 }
 
 void SUsdPrimInfo::SetPrimPath(const UE::FUsdStageWeak& UsdStage, const TCHAR* PrimPath)
 {
-	if ( PropertiesList )
+	if (PropertiesList)
 	{
 		TArray<FString> OldSelectedProperties = PropertiesList->GetSelectedFieldNames();
 		TArray<FString> OldSelectedPropertyMetadata = PropertyMetadataPanel->GetSelectedFieldNames();
@@ -137,7 +123,7 @@ void SUsdPrimInfo::SetPrimPath(const UE::FUsdStageWeak& UsdStage, const TCHAR* P
 		// We need this because apparently the list view doesn't generate a selection changed event when it's items are
 		// fully rebuilt, which will happen on SetObjectPath. If SetSelectedFieldNames can't select anything, nothing will
 		// update the metadata panel
-		if(OldSelectedProperties.Num() != NewSelectedProperties.Num() || NewSelectedProperties.Num() != 1)
+		if (OldSelectedProperties.Num() != NewSelectedProperties.Num() || NewSelectedProperties.Num() != 1)
 		{
 			PropertyMetadataPanel->SetObjectPath({}, TEXT(""));
 			PropertyMetadataPanel->SetVisibility(EVisibility::Collapsed);
@@ -149,22 +135,22 @@ void SUsdPrimInfo::SetPrimPath(const UE::FUsdStageWeak& UsdStage, const TCHAR* P
 		}
 	}
 
-	if ( IntegrationsPanel )
+	if (IntegrationsPanel)
 	{
-		IntegrationsPanel->SetPrimPath( UsdStage, PrimPath );
+		IntegrationsPanel->SetPrimPath(UsdStage, PrimPath);
 	}
 
-	if ( VariantsList )
+	if (VariantsList)
 	{
-		VariantsList->SetPrimPath( UsdStage, PrimPath );
+		VariantsList->SetPrimPath(UsdStage, PrimPath);
 	}
 
-	if ( ReferencesList )
+	if (ReferencesList)
 	{
-		ReferencesList->SetPrimPath( UsdStage, PrimPath );
+		ReferencesList->SetPrimPath(UsdStage, PrimPath);
 	}
 }
 
 #undef LOCTEXT_NAMESPACE
 
-#endif // #if USE_USD_SDK
+#endif	  // #if USE_USD_SDK
