USD: Reformat USDStageEditor module.

diff --git a/Engine/Plugins/Importers/USDImporter/Source/USDStageEditor/Private/SUSDStage.cpp b/Engine/Plugins/Importers/USDImporter/Source/USDStageEditor/Private/SUSDStage.cpp
--- a/Engine/Plugins/Importers/USDImporter/Source/USDStageEditor/Private/SUSDStage.cpp
+++ b/Engine/Plugins/Importers/USDImporter/Source/USDStageEditor/Private/SUSDStage.cpp
@@ -31,7 +31,6 @@
 #include "Framework/MultiBox/MultiBoxBuilder.h"
 #include "ISceneOutliner.h"
 #include "ISceneOutlinerColumn.h"
-#include "ISinglePropertyView.h"
 #include "Modules/ModuleManager.h"
 #include "PropertyEditorModule.h"
 #include "SceneOutlinerModule.h"
@@ -39,7 +38,6 @@
 #include "Selection.h"
 #include "Styling/AppStyle.h"
 #include "Widgets/Input/SButton.h"
-#include "Widgets/Input/SButton.h"
 #include "Widgets/Input/SComboBox.h"
 #include "Widgets/Input/SComboButton.h"
 #include "Widgets/Input/SEditableTextBox.h"
@@ -52,13 +50,15 @@ static bool GDiscardUndoBufferOnStageOpenClose = false;
 static FAutoConsoleVariableRef CVarDiscardUndoBufferOnStageOpenClose(
 	TEXT("USD.DiscardUndoBufferOnStageOpenClose"),
 	GDiscardUndoBufferOnStageOpenClose,
-	TEXT("Enabling this will prevent the recording of open/close stage transactions, but also discard the undo buffer after they happen. Use this when memory-constrained, as sometimes recording all created assets and actors in the undo buffer can be expensive."));
+	TEXT("Enabling this will prevent the recording of open/close stage transactions, but also discard the undo buffer after they happen. Use this "
+		 "when memory-constrained, as sometimes recording all created assets and actors in the undo buffer can be expensive.")
+);
 
 #if USE_USD_SDK
 
 namespace SUSDStageConstants
 {
-	static const FMargin SectionPadding( 1.f, 1.f, 1.f, 1.f );
+	static const FMargin SectionPadding(1.f, 1.f, 1.f, 1.f);
 }
 
 namespace SUSDStageImpl
@@ -66,26 +66,27 @@ namespace SUSDStageImpl
 #if PLATFORM_LINUX
 	struct FCaseSensitiveStringSetFuncs : BaseKeyFuncs<FString, FString>
 	{
-		static FORCEINLINE const FString& GetSetKey( const FString& Element )
+		static FORCEINLINE const FString& GetSetKey(const FString& Element)
 		{
 			return Element;
 		}
-		static FORCEINLINE bool Matches( const FString& A, const FString& B )
+		static FORCEINLINE bool Matches(const FString& A, const FString& B)
 		{
-			return A.Equals( B, ESearchCase::CaseSensitive );
+			return A.Equals(B, ESearchCase::CaseSensitive);
 		}
-		static FORCEINLINE uint32 GetKeyHash( const FString& Key )
+		static FORCEINLINE uint32 GetKeyHash(const FString& Key)
 		{
-			return FCrc::StrCrc32<TCHAR>( *Key );
+			return FCrc::StrCrc32<TCHAR>(*Key);
 		}
 	};
 #endif
 
 	/**
 	 * Makes sure that AllLayers includes all the external references of any of its members.
-	 * Example: Receives AllLayers [a.usda], that references b.usda and c.usda, while c.usda also references d.usda -> Modifies AllLayers to be [a.usda, b.usda, c.usda, d.usda]
+	 * Example: Receives AllLayers [a.usda], that references b.usda and c.usda, while c.usda also references d.usda -> Modifies AllLayers to be
+	 * [a.usda, b.usda, c.usda, d.usda]
 	 */
-	void AppendAllExternalReferences( TArray<UE::FSdfLayer>& AllLayers )
+	void AppendAllExternalReferences(TArray<UE::FSdfLayer>& AllLayers)
 	{
 		// Consider paths in a case-sensitive way for linux
 #if PLATFORM_LINUX
@@ -94,63 +95,63 @@ namespace SUSDStageImpl
 		TSet<FString> LoadedLayers;
 #endif
 
-		LoadedLayers.Reserve( AllLayers.Num() );
-		for ( const UE::FSdfLayer& Layer : AllLayers )
+		LoadedLayers.Reserve(AllLayers.Num());
+		for (const UE::FSdfLayer& Layer : AllLayers)
 		{
 			FString LayerPath = Layer.GetRealPath();
-			FPaths::NormalizeFilename( LayerPath );
-			LoadedLayers.Add( LayerPath );
+			FPaths::NormalizeFilename(LayerPath);
+			LoadedLayers.Add(LayerPath);
 		}
 
-		for ( int32 Index = 0; Index < AllLayers.Num(); ++Index )
+		for (int32 Index = 0; Index < AllLayers.Num(); ++Index)
 		{
-			UE::FSdfLayer& Layer = AllLayers[ Index ];
+			UE::FSdfLayer& Layer = AllLayers[Index];
 
-			TArray< UE::FSdfLayer > NewLayers;
-			for ( const FString& AssetDependency : Layer.GetCompositionAssetDependencies() )
+			TArray<UE::FSdfLayer> NewLayers;
+			for (const FString& AssetDependency : Layer.GetCompositionAssetDependencies())
 			{
-				FString AbsoluteReference = Layer.ComputeAbsolutePath( AssetDependency );
-				FPaths::NormalizeFilename( AbsoluteReference );
+				FString AbsoluteReference = Layer.ComputeAbsolutePath(AssetDependency);
+				FPaths::NormalizeFilename(AbsoluteReference);
 
-				if ( !LoadedLayers.Contains( AbsoluteReference ) )
+				if (!LoadedLayers.Contains(AbsoluteReference))
 				{
-					NewLayers.Add( UE::FSdfLayer::FindOrOpen( *AbsoluteReference ) );
-					LoadedLayers.Add( AbsoluteReference );
+					NewLayers.Add(UE::FSdfLayer::FindOrOpen(*AbsoluteReference));
+					LoadedLayers.Add(AbsoluteReference);
 				}
 			}
-			AllLayers.Append( NewLayers );
+			AllLayers.Append(NewLayers);
 		}
 	}
 
-	void SelectGeneratedComponentsAndActors( AUsdStageActor* StageActor, const TArray<FString>& PrimPaths )
+	void SelectGeneratedComponentsAndActors(AUsdStageActor* StageActor, const TArray<FString>& PrimPaths)
 	{
-		if ( !StageActor )
+		if (!StageActor)
 		{
 			return;
 		}
 
 		TSet<USceneComponent*> ComponentsToSelect;
-		for ( const FString& PrimPath : PrimPaths )
+		for (const FString& PrimPath : PrimPaths)
 		{
-			if ( USceneComponent* GeneratedComponent = StageActor->GetGeneratedComponent( PrimPath ) )
+			if (USceneComponent* GeneratedComponent = StageActor->GetGeneratedComponent(PrimPath))
 			{
-				ComponentsToSelect.Add( GeneratedComponent );
+				ComponentsToSelect.Add(GeneratedComponent);
 			}
 		}
 
 		TSet<AActor*> ActorsToSelect;
-		for ( TSet<USceneComponent*>::TIterator It( ComponentsToSelect ); It; ++It )
+		for (TSet<USceneComponent*>::TIterator It(ComponentsToSelect); It; ++It)
 		{
-			if ( AActor* Owner = ( *It )->GetOwner() )
+			if (AActor* Owner = (*It)->GetOwner())
 			{
 				// We always need the parent actor selected to select a component
-				ActorsToSelect.Add( Owner );
+				ActorsToSelect.Add(Owner);
 
 				// If we're going to select a root component, select the actor instead. This is useful
 				// because if we later press "F" to focus on the prim, having the actor selected will use the entire
 				// actor's bounding box and focus on something. If all we have selected is an empty scene component
 				// however, the camera won't focus on anything
-				if ( *It == Owner->GetRootComponent() )
+				if (*It == Owner->GetRootComponent())
 				{
 					It.RemoveCurrent();
 				}
@@ -158,57 +159,60 @@ namespace SUSDStageImpl
 		}
 
 		// Don't deselect anything if we're not going to select anything
-		if ( ActorsToSelect.Num() == 0 && ComponentsToSelect.Num() == 0 )
+		if (ActorsToSelect.Num() == 0 && ComponentsToSelect.Num() == 0)
 		{
 			return;
 		}
 
 		// Don't use GEditor->SelectNone() as that will affect *every type of selection in the editor*,
 		// including even some UI menus, brushes, etc.
-		if ( USelection* SelectedActors = GEditor->GetSelectedActors() )
+		if (USelection* SelectedActors = GEditor->GetSelectedActors())
 		{
-			SelectedActors->DeselectAll( AActor::StaticClass() );
+			SelectedActors->DeselectAll(AActor::StaticClass());
 		}
-		if ( USelection* SelectedComponents = GEditor->GetSelectedComponents() )
+		if (USelection* SelectedComponents = GEditor->GetSelectedComponents())
 		{
-			SelectedComponents->DeselectAll( UActorComponent::StaticClass() );
+			SelectedComponents->DeselectAll(UActorComponent::StaticClass());
 		}
 
 		const bool bSelected = true;
 		const bool bNotifySelectionChanged = true;
-		for ( AActor* Actor : ActorsToSelect )
+		for (AActor* Actor : ActorsToSelect)
 		{
-			GEditor->SelectActor( Actor, bSelected, bNotifySelectionChanged );
+			GEditor->SelectActor(Actor, bSelected, bNotifySelectionChanged);
 		}
 
-		for ( USceneComponent* Component : ComponentsToSelect )
+		for (USceneComponent* Component : ComponentsToSelect)
 		{
-			GEditor->SelectComponent( Component, bSelected, bNotifySelectionChanged );
+			GEditor->SelectComponent(Component, bSelected, bNotifySelectionChanged);
 		}
 	}
 
 	class FStageNameOutlinerColumn : public ISceneOutlinerColumn
 	{
 	public:
-		FStageNameOutlinerColumn( ISceneOutliner& SceneOutliner )
-			: WeakSceneOutliner( StaticCastSharedRef<ISceneOutliner>( SceneOutliner.AsShared() ) )
-		{}
+		FStageNameOutlinerColumn(ISceneOutliner& SceneOutliner)
+			: WeakSceneOutliner(StaticCastSharedRef<ISceneOutliner>(SceneOutliner.AsShared()))
+		{
+		}
 
-		virtual ~FStageNameOutlinerColumn() {}
+		virtual ~FStageNameOutlinerColumn()
+		{
+		}
 		static FName GetID()
 		{
-			static FName ID( "Stage" );
+			static FName ID("Stage");
 			return ID;
 		}
-		static FString GetStageForItem( const ISceneOutlinerTreeItem* Item )
+		static FString GetStageForItem(const ISceneOutlinerTreeItem* Item)
 		{
-			if ( Item )
+			if (Item)
 			{
-				if ( const FActorTreeItem* ActorTreeItem = Item->CastTo<FActorTreeItem>() )
+				if (const FActorTreeItem* ActorTreeItem = Item->CastTo<FActorTreeItem>())
 				{
-					if ( const AUsdStageActor* StageActor = Cast<const AUsdStageActor>( ActorTreeItem->Actor.Get() ) )
+					if (const AUsdStageActor* StageActor = Cast<const AUsdStageActor>(ActorTreeItem->Actor.Get()))
 					{
-						if ( const UE::FUsdStage& Stage = StageActor->GetUsdStage() )
+						if (const UE::FUsdStage& Stage = StageActor->GetUsdStage())
 						{
 							return Stage.GetRootLayer().GetDisplayName();
 						}
@@ -226,41 +230,39 @@ namespace SUSDStageImpl
 		}
 		virtual SHeaderRow::FColumn::FArguments ConstructHeaderRowColumn() override
 		{
-			return SHeaderRow::Column( FName( TEXT( "Stage" ) ) );
+			return SHeaderRow::Column(FName(TEXT("Stage")));
 		}
-		virtual const TSharedRef<SWidget> ConstructRowWidget(
-			FSceneOutlinerTreeItemRef TreeItem,
-			const STableRow<FSceneOutlinerTreeItemPtr>& Row
-		) override
+		virtual const TSharedRef<SWidget> ConstructRowWidget(FSceneOutlinerTreeItemRef TreeItem, const STableRow<FSceneOutlinerTreeItemPtr>& Row)
+			override
 		{
-			return SNew( STextBlock ).Text( FText::FromString( GetStageForItem( &TreeItem.Get() ) ) );
+			return SNew(STextBlock).Text(FText::FromString(GetStageForItem(&TreeItem.Get())));
 		}
 		virtual bool SupportsSorting() const override
 		{
 			return true;
 		}
-		virtual void SortItems( TArray<FSceneOutlinerTreeItemPtr>& OutItems, const EColumnSortMode::Type SortMode ) const override
+		virtual void SortItems(TArray<FSceneOutlinerTreeItemPtr>& OutItems, const EColumnSortMode::Type SortMode) const override
 		{
-			if ( SortMode == EColumnSortMode::Ascending )
+			if (SortMode == EColumnSortMode::Ascending)
 			{
 				OutItems.Sort(
-					[]( const FSceneOutlinerTreeItemPtr& First, const FSceneOutlinerTreeItemPtr& Second )
+					[](const FSceneOutlinerTreeItemPtr& First, const FSceneOutlinerTreeItemPtr& Second)
 					{
 						const ISceneOutlinerTreeItem* FirstPtr = First.Get();
 						const ISceneOutlinerTreeItem* SecondPtr = Second.Get();
-						return GetStageForItem( FirstPtr ) < GetStageForItem( SecondPtr );
+						return GetStageForItem(FirstPtr) < GetStageForItem(SecondPtr);
 					}
 				);
 			}
 
-			if ( SortMode == EColumnSortMode::Descending )
+			if (SortMode == EColumnSortMode::Descending)
 			{
 				OutItems.Sort(
-					[]( const FSceneOutlinerTreeItemPtr& First, const FSceneOutlinerTreeItemPtr& Second )
+					[](const FSceneOutlinerTreeItemPtr& First, const FSceneOutlinerTreeItemPtr& Second)
 					{
 						const ISceneOutlinerTreeItem* FirstPtr = First.Get();
 						const ISceneOutlinerTreeItem* SecondPtr = Second.Get();
-						return GetStageForItem( FirstPtr ) > GetStageForItem( SecondPtr );
+						return GetStageForItem(FirstPtr) > GetStageForItem(SecondPtr);
 					}
 				);
 			}
@@ -270,20 +272,20 @@ namespace SUSDStageImpl
 	private:
 		TWeakPtr<ISceneOutliner> WeakSceneOutliner;
 	};
-}
+}	 // namespace SUSDStageImpl
 
-void SUsdStage::Construct( const FArguments& InArgs )
+void SUsdStage::Construct(const FArguments& InArgs)
 {
 	const bool bFlashButton = true;
-	OnActorLoadedHandle = AUsdStageActor::OnActorLoaded.AddSP( SharedThis( this ), &SUsdStage::AttachToStageActor, bFlashButton );
+	OnActorLoadedHandle = AUsdStageActor::OnActorLoaded.AddSP(SharedThis(this), &SUsdStage::AttachToStageActor, bFlashButton);
 
-	OnViewportSelectionChangedHandle = USelection::SelectionChangedEvent.AddSP( SharedThis(this), &SUsdStage::OnViewportSelectionChanged );
+	OnViewportSelectionChangedHandle = USelection::SelectionChangedEvent.AddSP(SharedThis(this), &SUsdStage::OnViewportSelectionChanged);
 
-	PostPIEStartedHandle = FEditorDelegates::PostPIEStarted.AddSP( SharedThis(this), &SUsdStage::OnPostPIEStarted );
-	EndPIEHandle = FEditorDelegates::EndPIE.AddSP( SharedThis(this), &SUsdStage::OnEndPIE );
+	PostPIEStartedHandle = FEditorDelegates::PostPIEStarted.AddSP(SharedThis(this), &SUsdStage::OnPostPIEStarted);
+	EndPIEHandle = FEditorDelegates::EndPIE.AddSP(SharedThis(this), &SUsdStage::OnEndPIE);
 
-	IUsdStageModule& UsdStageModule = FModuleManager::Get().LoadModuleChecked< IUsdStageModule >( "UsdStage" );
-	ViewModel.UsdStageActor = UsdStageModule.FindUsdStageActor( IUsdClassesModule::GetCurrentWorld() );
+	IUsdStageModule& UsdStageModule = FModuleManager::Get().LoadModuleChecked<IUsdStageModule>("UsdStage");
+	ViewModel.UsdStageActor = UsdStageModule.FindUsdStageActor(IUsdClassesModule::GetCurrentWorld());
 
 	bUpdatingViewportSelection = false;
 	bUpdatingPrimSelection = false;
@@ -298,25 +300,26 @@ void SUsdStage::Construct( const FArguments& InArgs )
 	);
 
 	const AUsdStageActor* StageActor = ViewModel.UsdStageActor.Get();
-	if ( StageActor )
+	if (StageActor)
 	{
-		UsdStage = static_cast< const AUsdStageActor* >( StageActor )->GetUsdStage();
+		UsdStage = static_cast<const AUsdStageActor*>(StageActor)->GetUsdStage();
 	}
 
+	// clang-format off
 	ChildSlot
 	[
-		SNew( SBorder )
-		.BorderImage( FAppStyle::GetBrush("Docking.Tab.ContentAreaBrush") )
+		SNew(SBorder)
+		.BorderImage(FAppStyle::GetBrush("Docking.Tab.ContentAreaBrush"))
 		[
-			SNew( SVerticalBox )
+			SNew(SVerticalBox)
 
 			// Menu
 			+ SVerticalBox::Slot()
 			.AutoHeight()
 			[
-				SNew( SHorizontalBox )
+				SNew(SHorizontalBox)
 				+ SHorizontalBox::Slot()
-				.FillWidth( 1 )
+				.FillWidth(1)
 				[
 					MakeMainMenu()
 				]
@@ -324,10 +327,10 @@ void SUsdStage::Construct( const FArguments& InArgs )
 				+ SHorizontalBox::Slot()
 				.AutoWidth()
 				[
-					SNew( SBox )
-					.HAlign( HAlign_Right )
-					.VAlign( VAlign_Fill )
-					.Padding( FMargin( 0.0f, 0.0f, 10.0f, 0.0f ) )
+					SNew(SBox)
+					.HAlign(HAlign_Right)
+					.VAlign(VAlign_Fill)
+					.Padding(FMargin(0.0f, 0.0f, 10.0f, 0.0f))
 					[
 						MakeIsolateWarningButton()
 					]
@@ -336,9 +339,9 @@ void SUsdStage::Construct( const FArguments& InArgs )
 				+ SHorizontalBox::Slot()
 				.AutoWidth()
 				[
-					SNew( SBox )
-					.HAlign( HAlign_Right )
-					.VAlign( VAlign_Fill )
+					SNew(SBox)
+					.HAlign(HAlign_Right)
+					.VAlign(VAlign_Fill)
 					[
 						MakeActorPickerMenu()
 					]
@@ -347,53 +350,54 @@ void SUsdStage::Construct( const FArguments& InArgs )
 
 			+SVerticalBox::Slot()
 			.FillHeight(1.f)
-			.Padding( SUSDStageConstants::SectionPadding )
+			.Padding(SUSDStageConstants::SectionPadding)
 			[
-				SNew( SSplitter )
-				.Orientation( Orient_Vertical )
+				SNew(SSplitter)
+				.Orientation(Orient_Vertical)
 
 				+SSplitter::Slot()
-				.Value( 0.7f )
+				.Value(0.7f)
 				[
-					SNew( SSplitter )
-					.Orientation( Orient_Horizontal )
+					SNew(SSplitter)
+					.Orientation(Orient_Horizontal)
 
 					// Stage tree
 					+SSplitter::Slot()
 					[
-						SNew( SBorder )
-						.BorderImage( FAppStyle::GetBrush(TEXT("ToolPanel.GroupBorder")) )
+						SNew(SBorder)
+						.BorderImage(FAppStyle::GetBrush(TEXT("ToolPanel.GroupBorder")))
 						[
-							SAssignNew( UsdStageTreeView, SUsdStageTreeView )
-							.OnPrimSelectionChanged( this, &SUsdStage::OnPrimSelectionChanged )
+							SAssignNew(UsdStageTreeView, SUsdStageTreeView)
+							.OnPrimSelectionChanged(this, &SUsdStage::OnPrimSelectionChanged)
 						]
 					]
 
 					// Prim details
 					+SSplitter::Slot()
 					[
-						SNew( SBorder )
-						.BorderImage( FAppStyle::GetBrush("ToolPanel.GroupBorder") )
+						SNew(SBorder)
+						.BorderImage(FAppStyle::GetBrush("ToolPanel.GroupBorder"))
 						[
-							SAssignNew( UsdPrimInfoWidget, SUsdPrimInfo )
+							SAssignNew(UsdPrimInfoWidget, SUsdPrimInfo)
 						]
 					]
 				]
 
 				// Layers tree
 				+SSplitter::Slot()
-				.Value( 0.3f )
+				.Value(0.3f)
 				[
 					SNew(SBorder)
-					.BorderImage( FAppStyle::GetBrush(TEXT("ToolPanel.GroupBorder")) )
+					.BorderImage(FAppStyle::GetBrush(TEXT("ToolPanel.GroupBorder")))
 					[
-						SAssignNew( UsdLayersTreeView, SUsdLayersTreeView )
-						.OnLayerIsolated( this, &SUsdStage::OnLayerIsolated )
+						SAssignNew(UsdLayersTreeView, SUsdLayersTreeView)
+						.OnLayerIsolated(this, &SUsdStage::OnLayerIsolated)
 					]
 				]
 			]
 		]
 	];
+	// clang-format on
 
 	SetupStageActorDelegates();
 
@@ -401,7 +405,7 @@ void SUsdStage::Construct( const FArguments& InArgs )
 	FSlateApplication::Get().OnPreTick().AddSP(this, &SUsdStage::OnSlateTick);
 
 	// We're opening the USD Stage editor for the first time and we already have a stage: Display it immediately
-	if ( UsdStage )
+	if (UsdStage)
 	{
 		RequestFullRefresh();
 	}
@@ -414,7 +418,7 @@ AUsdStageActor* SUsdStage::GetAttachedStageActor() const
 
 TArray<UE::FSdfLayer> SUsdStage::GetSelectedLayers() const
 {
-	if ( UsdLayersTreeView )
+	if (UsdLayersTreeView)
 	{
 		return UsdLayersTreeView->GetSelectedLayers();
 	}
@@ -422,17 +426,17 @@ TArray<UE::FSdfLayer> SUsdStage::GetSelectedLayers() const
 	return {};
 }
 
-void SUsdStage::SetSelectedLayers( const TArray<UE::FSdfLayer>& NewSelection ) const
+void SUsdStage::SetSelectedLayers(const TArray<UE::FSdfLayer>& NewSelection) const
 {
-	if ( UsdLayersTreeView )
+	if (UsdLayersTreeView)
 	{
-		UsdLayersTreeView->SetSelectedLayers( NewSelection );
+		UsdLayersTreeView->SetSelectedLayers(NewSelection);
 	}
 }
 
 TArray<UE::FUsdPrim> SUsdStage::GetSelectedPrims() const
 {
-	if ( UsdStageTreeView )
+	if (UsdStageTreeView)
 	{
 		return UsdStageTreeView->GetSelectedPrims();
 	}
@@ -440,17 +444,17 @@ TArray<UE::FUsdPrim> SUsdStage::GetSelectedPrims() const
 	return {};
 }
 
-void SUsdStage::SetSelectedPrims( const TArray<UE::FUsdPrim>& NewSelection ) const
+void SUsdStage::SetSelectedPrims(const TArray<UE::FUsdPrim>& NewSelection) const
 {
-	if (UsdStageTreeView )
+	if (UsdStageTreeView)
 	{
-		UsdStageTreeView->SetSelectedPrims( NewSelection );
+		UsdStageTreeView->SetSelectedPrims(NewSelection);
 	}
 }
 
 TArray<FString> SUsdStage::GetSelectedPropertyNames() const
 {
-	if( UsdPrimInfoWidget && UsdPrimInfoWidget->PropertiesList )
+	if (UsdPrimInfoWidget && UsdPrimInfoWidget->PropertiesList)
 	{
 		return UsdPrimInfoWidget->PropertiesList->GetSelectedFieldNames();
 	}
@@ -458,9 +462,9 @@ TArray<FString> SUsdStage::GetSelectedPropertyNames() const
 	return {};
 }
 
-void SUsdStage::SetSelectedPropertyNames( const TArray<FString>& NewSelection )
+void SUsdStage::SetSelectedPropertyNames(const TArray<FString>& NewSelection)
 {
-	if ( UsdPrimInfoWidget && UsdPrimInfoWidget->PropertiesList )
+	if (UsdPrimInfoWidget && UsdPrimInfoWidget->PropertiesList)
 	{
 		UsdPrimInfoWidget->PropertiesList->SetSelectedFieldNames(NewSelection);
 	}
@@ -484,34 +488,34 @@ void SUsdStage::SetSelectedPropertyMetadataNames(const TArray<FString>& NewSelec
 	}
 }
 
-void SUsdStage::AttachToStageActor( AUsdStageActor* InUsdStageActor, bool bFlashButton )
+void SUsdStage::AttachToStageActor(AUsdStageActor* InUsdStageActor, bool bFlashButton)
 {
-	if ( ViewModel.UsdStageActor == InUsdStageActor )
+	if (ViewModel.UsdStageActor == InUsdStageActor)
 	{
 		return;
 	}
 
-	if ( bFlashButton )
+	if (bFlashButton)
 	{
 		const float StartTimeSeconds = 0.0f;
 		const float FlashDuration = 0.5f;
 		const ECurveEaseFunction EaseFunction = ECurveEaseFunction::Linear;
-		FlashActorPickerCurve = FCurveSequence{ StartTimeSeconds, FlashDuration, EaseFunction };
-		FlashActorPickerCurve.Play( this->AsShared() );
+		FlashActorPickerCurve = FCurveSequence{StartTimeSeconds, FlashDuration, EaseFunction};
+		FlashActorPickerCurve.Play(this->AsShared());
 	}
 
 	ClearStageActorDelegates();
 	ViewModel.UsdStageActor = InUsdStageActor;
 	SetupStageActorDelegates();
 
-	if ( InUsdStageActor )
+	if (InUsdStageActor)
 	{
-		UE::FUsdStage UsdStage = static_cast< const AUsdStageActor* >( InUsdStageActor )->GetUsdStage();
+		UE::FUsdStage UsdStage = static_cast<const AUsdStageActor*>(InUsdStageActor)->GetUsdStage();
 
-		if ( this->UsdPrimInfoWidget && InUsdStageActor )
+		if (this->UsdPrimInfoWidget && InUsdStageActor)
 		{
 			// Just reset to the pseudoroot for now
-			this->UsdPrimInfoWidget->SetPrimPath( UsdStage, TEXT( "/" ) );
+			this->UsdPrimInfoWidget->SetPrimPath(UsdStage, TEXT("/"));
 		}
 	}
 
@@ -524,7 +528,7 @@ void SUsdStage::SetupStageActorDelegates()
 {
 	ClearStageActorDelegates();
 
-	if ( ViewModel.UsdStageActor.IsValid() )
+	if (ViewModel.UsdStageActor.IsValid())
 	{
 		OnPrimChangedHandle = ViewModel.UsdStageActor->OnPrimChanged.AddLambda(
 			[this](const FString& PrimPath, bool bResync)
@@ -555,7 +559,7 @@ void SUsdStage::SetupStageActorDelegates()
 		);
 
 		OnActorDestroyedHandle = ViewModel.UsdStageActor->OnActorDestroyed.AddLambda(
-			[ this ]()
+			[this]()
 			{
 				// Refresh widgets on game thread, but close the stage right away. In some contexts this is important, for example when
 				// running a Python script: If our USD Stage Editor is open and our script deletes an actor, it will trigger OnActorDestroyed.
@@ -590,15 +594,15 @@ void SUsdStage::SetupStageActorDelegates()
 void SUsdStage::ClearStageActorDelegates()
 {
 	const bool bEvenIfPendingKill = true;
-	if ( AUsdStageActor* StageActor = ViewModel.UsdStageActor.Get( bEvenIfPendingKill ) )
+	if (AUsdStageActor* StageActor = ViewModel.UsdStageActor.Get(bEvenIfPendingKill))
 	{
-		StageActor->OnStageChanged.Remove( OnStageChangedHandle );
-		StageActor->OnPrimChanged.Remove( OnPrimChangedHandle );
-		StageActor->OnActorDestroyed.Remove ( OnActorDestroyedHandle );
+		StageActor->OnStageChanged.Remove(OnStageChangedHandle);
+		StageActor->OnPrimChanged.Remove(OnPrimChangedHandle);
+		StageActor->OnActorDestroyed.Remove(OnActorDestroyedHandle);
 
-		StageActor->GetUsdListener().GetOnStageEditTargetChanged().Remove( OnStageEditTargetChangedHandle );
-		StageActor->GetUsdListener().GetOnSdfLayersChanged().Remove( OnSdfLayersChangedHandle );
-		StageActor->GetUsdListener().GetOnSdfLayerDirtinessChanged().Remove( OnSdfLayerDirtinessChangedHandle );
+		StageActor->GetUsdListener().GetOnStageEditTargetChanged().Remove(OnStageEditTargetChangedHandle);
+		StageActor->GetUsdListener().GetOnSdfLayersChanged().Remove(OnSdfLayersChangedHandle);
+		StageActor->GetUsdListener().GetOnSdfLayerDirtinessChanged().Remove(OnSdfLayerDirtinessChangedHandle);
 	}
 }
 
@@ -606,72 +610,76 @@ SUsdStage::~SUsdStage()
 {
 	FSlateApplication::Get().OnPreTick().RemoveAll(this);
 
-	FEditorDelegates::PostPIEStarted.Remove( PostPIEStartedHandle );
-	FEditorDelegates::EndPIE.Remove( EndPIEHandle );
+	FEditorDelegates::PostPIEStarted.Remove(PostPIEStartedHandle);
+	FEditorDelegates::EndPIE.Remove(EndPIEHandle);
 
 	GEditor->OnEditorClose().Remove(OnEditorCloseHandle);
 
-	FCoreUObjectDelegates::OnObjectPropertyChanged.Remove( OnStageActorPropertyChangedHandle );
-	AUsdStageActor::OnActorLoaded.Remove( OnActorLoadedHandle );
-	USelection::SelectionChangedEvent.Remove( OnViewportSelectionChangedHandle );
+	FCoreUObjectDelegates::OnObjectPropertyChanged.Remove(OnStageActorPropertyChangedHandle);
+	AUsdStageActor::OnActorLoaded.Remove(OnActorLoadedHandle);
+	USelection::SelectionChangedEvent.Remove(OnViewportSelectionChangedHandle);
 
 	ClearStageActorDelegates();
 
 	ActorPickerMenu.Reset();
 }
 
-TSharedRef< SWidget > SUsdStage::MakeMainMenu()
+TSharedRef<SWidget> SUsdStage::MakeMainMenu()
 {
-	FMenuBarBuilder MenuBuilder( nullptr );
+	FMenuBarBuilder MenuBuilder(nullptr);
 	{
 		// File
 		MenuBuilder.AddPullDownMenu(
-			LOCTEXT( "FileMenu", "File" ),
-			LOCTEXT( "FileMenu_ToolTip", "Opens the file menu" ),
-			FNewMenuDelegate::CreateSP( this, &SUsdStage::FillFileMenu ) );
+			LOCTEXT("FileMenu", "File"),
+			LOCTEXT("FileMenu_ToolTip", "Opens the file menu"),
+			FNewMenuDelegate::CreateSP(this, &SUsdStage::FillFileMenu)
+		);
 
 		// Actions
 		MenuBuilder.AddPullDownMenu(
-			LOCTEXT( "ActionsMenu", "Actions" ),
-			LOCTEXT( "ActionsMenu_ToolTip", "Opens the actions menu" ),
-			FNewMenuDelegate::CreateSP( this, &SUsdStage::FillActionsMenu ) );
+			LOCTEXT("ActionsMenu", "Actions"),
+			LOCTEXT("ActionsMenu_ToolTip", "Opens the actions menu"),
+			FNewMenuDelegate::CreateSP(this, &SUsdStage::FillActionsMenu)
+		);
 
 		// Options
 		MenuBuilder.AddPullDownMenu(
-			LOCTEXT( "OptionsMenu", "Options" ),
-			LOCTEXT( "OptionsMenu_ToolTip", "Opens the options menu" ),
-			FNewMenuDelegate::CreateSP( this, &SUsdStage::FillOptionsMenu ) );
+			LOCTEXT("OptionsMenu", "Options"),
+			LOCTEXT("OptionsMenu_ToolTip", "Opens the options menu"),
+			FNewMenuDelegate::CreateSP(this, &SUsdStage::FillOptionsMenu)
+		);
 	}
 
 	// Create the menu bar
-	TSharedRef< SWidget > MenuBarWidget = MenuBuilder.MakeWidget();
-	MenuBarWidget->SetVisibility( EVisibility::Visible ); // Work around for menu bar not showing on Mac
+	TSharedRef<SWidget> MenuBarWidget = MenuBuilder.MakeWidget();
+	MenuBarWidget->SetVisibility(EVisibility::Visible);	   // Work around for menu bar not showing on Mac
 
 	return MenuBarWidget;
 }
 
-TSharedRef< SWidget > SUsdStage::MakeActorPickerMenu()
+TSharedRef<SWidget> SUsdStage::MakeActorPickerMenu()
 {
+	// clang-format off
 	return
 		SNew(SOverlay)
 		+ SOverlay::Slot()
 		[
-			SNew( SComboButton )
-			.ComboButtonStyle( FAppStyle::Get(), "SimpleComboButton" )
-			.OnGetMenuContent( this, &SUsdStage::MakeActorPickerMenuContent )
-			.MenuPlacement( EMenuPlacement::MenuPlacement_BelowRightAnchor )
-			.ToolTipText( LOCTEXT( "ActorPicker_ToolTip", "Switch the active stage actor" ) )
+			SNew(SComboButton)
+			.ComboButtonStyle(FAppStyle::Get(), "SimpleComboButton")
+			.OnGetMenuContent(this, &SUsdStage::MakeActorPickerMenuContent)
+			.MenuPlacement(EMenuPlacement::MenuPlacement_BelowRightAnchor)
+			.ToolTipText(LOCTEXT("ActorPicker_ToolTip", "Switch the active stage actor"))
 			.ButtonContent()
 			[
-				SNew( STextBlock )
-				.Text_Lambda( [this]()
+				SNew(STextBlock)
+				.Text_Lambda([this]()
 				{
-					if ( AUsdStageActor* StageActor = ViewModel.UsdStageActor.Get() )
+					if (AUsdStageActor* StageActor = ViewModel.UsdStageActor.Get())
 					{
-						return FText::FromString( StageActor->GetActorLabel() );
+						return FText::FromString(StageActor->GetActorLabel());
 					}
 
-					return FText::FromString( "None" );
+					return FText::FromString("None");
 				})
 			]
 		]
@@ -680,13 +688,13 @@ TSharedRef< SWidget > SUsdStage::MakeActorPickerMenu()
 		+ SOverlay::Slot()
 		[
 			SNew(SBorder)
-			.Visibility( EVisibility::HitTestInvisible )
-			.BorderImage( FUsdStageEditorStyle::Get()->GetBrush( "UsdStageEditor.ActorPickerFlashBrush" ) )
-			.BorderBackgroundColor_Lambda( [this]() -> FSlateColor
+			.Visibility(EVisibility::HitTestInvisible)
+			.BorderImage(FUsdStageEditorStyle::Get()->GetBrush("UsdStageEditor.ActorPickerFlashBrush"))
+			.BorderBackgroundColor_Lambda([this]() -> FSlateColor
 			{
 				FLinearColor Color = FLinearColor::White;
 
-				if ( FlashActorPickerCurve.IsPlaying() )
+				if (FlashActorPickerCurve.IsPlaying())
 				{
 					Color.A = 0.6f * (1.0f - FlashActorPickerCurve.GetLerp());
 				}
@@ -695,70 +703,62 @@ TSharedRef< SWidget > SUsdStage::MakeActorPickerMenu()
 					Color.A = 0.0f;
 				}
 
-				return FSlateColor( Color );
+				return FSlateColor(Color);
 			})
 		];
+	// clang-format on
 }
 
-TSharedRef< SWidget > SUsdStage::MakeActorPickerMenuContent()
+TSharedRef<SWidget> SUsdStage::MakeActorPickerMenuContent()
 {
-	if ( !ActorPickerMenu.IsValid() )
+	if (!ActorPickerMenu.IsValid())
 	{
 		const int32 NamePriority = 0;
 		const int32 StagePriority = 1;
 		const bool bCanBeHidden = true;
 		const float FillSize = 1.0f;
 
-		FSceneOutlinerColumnInfo ActorColumnInfo{ ESceneOutlinerColumnVisibility::Visible, NamePriority };
-		ActorColumnInfo.ColumnLabel = LOCTEXT( "ActorColumnText", "Actor" );
+		FSceneOutlinerColumnInfo ActorColumnInfo{ESceneOutlinerColumnVisibility::Visible, NamePriority};
+		ActorColumnInfo.ColumnLabel = LOCTEXT("ActorColumnText", "Actor");
 		ActorColumnInfo.FillSize = FillSize;
 
 		FSceneOutlinerColumnInfo StageColumnInfo{
 			ESceneOutlinerColumnVisibility::Visible,
 			StagePriority,
 			FCreateSceneOutlinerColumn::CreateLambda(
-				[]( ISceneOutliner& InSceneOutliner )
+				[](ISceneOutliner& InSceneOutliner)
 				{
-					return MakeShareable( new SUSDStageImpl::FStageNameOutlinerColumn( InSceneOutliner ) );
+					return MakeShareable(new SUSDStageImpl::FStageNameOutlinerColumn(InSceneOutliner));
 				}
 			),
 			bCanBeHidden,
-			FillSize
-		};
+			FillSize};
 
 		FSceneOutlinerInitializationOptions InitOptions;
-		InitOptions.bShowHeaderRow = true;  // The header lets us resize/hide the columns if we want
+		InitOptions.bShowHeaderRow = true;	  // The header lets us resize/hide the columns if we want
 		InitOptions.bShowSearchBox = true;
 		InitOptions.bShowCreateNewFolder = false;
 		InitOptions.bFocusSearchBoxWhenOpened = true;
-		InitOptions.ColumnMap.Add(
-			FSceneOutlinerBuiltInColumnTypes::Label(),
-			ActorColumnInfo
-		);
-		InitOptions.ColumnMap.Add(
-			SUSDStageImpl::FStageNameOutlinerColumn::GetID(),
-			StageColumnInfo
-		);
-		InitOptions.Filters->AddFilterPredicate<FActorTreeItem>(
-			FActorTreeItem::FFilterPredicate::CreateLambda(
-				[]( const AActor* Actor )
-				{
-					return Actor && Actor->IsA<AUsdStageActor>();
-				}
-			)
-		);
+		InitOptions.ColumnMap.Add(FSceneOutlinerBuiltInColumnTypes::Label(), ActorColumnInfo);
+		InitOptions.ColumnMap.Add(SUSDStageImpl::FStageNameOutlinerColumn::GetID(), StageColumnInfo);
+		InitOptions.Filters->AddFilterPredicate<FActorTreeItem>(FActorTreeItem::FFilterPredicate::CreateLambda(
+			[](const AActor* Actor)
+			{
+				return Actor && Actor->IsA<AUsdStageActor>();
+			}
+		));
 
-		FSceneOutlinerModule& SceneOutlinerModule = FModuleManager::LoadModuleChecked<FSceneOutlinerModule>( "SceneOutliner" );
+		FSceneOutlinerModule& SceneOutlinerModule = FModuleManager::LoadModuleChecked<FSceneOutlinerModule>("SceneOutliner");
 		ActorPickerMenu = SceneOutlinerModule.CreateActorPicker(
 			InitOptions,
 			FOnActorPicked::CreateLambda(
-				[this]( AActor* Actor )
+				[this](AActor* Actor)
 				{
-					if ( Actor && Actor->IsA<AUsdStageActor>() )
+					if (Actor && Actor->IsA<AUsdStageActor>())
 					{
 						// Don't need to flash here as it was an intentional act to attach to a new actor
 						const bool bFlashButton = false;
-						AttachToStageActor( Cast<AUsdStageActor>( Actor ), bFlashButton );
+						AttachToStageActor(Cast<AUsdStageActor>(Actor), bFlashButton);
 
 						FSlateApplication::Get().DismissAllMenus();
 					}
@@ -767,53 +767,56 @@ TSharedRef< SWidget > SUsdStage::MakeActorPickerMenuContent()
 		);
 	}
 
-	if ( ActorPickerMenu.IsValid() )
+	if (ActorPickerMenu.IsValid())
 	{
 		ActorPickerMenu->FullRefresh();
 
+		// clang-format off
 		return SNew(SBox)
-			.Padding( FMargin( 1 ) )    // Add a small margin or else we'll get dark gray on dark gray which can look a bit confusing
-			.MinDesiredWidth( 600.0f )  // Force a min width or else the tree view item text will run up right to the very edge pixel of the menu
-			.HAlign( HAlign_Fill )
+			.Padding(FMargin(1))    // Add a small margin or else we'll get dark gray on dark gray which can look a bit confusing
+			.MinDesiredWidth(600.0f)  // Force a min width or else the tree view item text will run up right to the very edge pixel of the menu
+			.HAlign(HAlign_Fill)
 			[
 				ActorPickerMenu.ToSharedRef()
 			];
+		// clang-format on
 	}
 
 	return SNullWidget::NullWidget;
 }
 
-TSharedRef< SWidget > SUsdStage::MakeIsolateWarningButton()
+TSharedRef<SWidget> SUsdStage::MakeIsolateWarningButton()
 {
-	return SNew( SButton )
-		.ButtonStyle( &FAppStyle::Get().GetWidgetStyle<FButtonStyle>( "EditorViewportToolBar.WarningButton" ) )
-		.ButtonColorAndOpacity( FLinearColor{ 1.0f, 0.4f, 0.0f, 1.0f } )
-		.OnClicked_Lambda( [this]() -> FReply
+	// clang-format off
+	return SNew(SButton)
+		.ButtonStyle(&FAppStyle::Get().GetWidgetStyle<FButtonStyle>("EditorViewportToolBar.WarningButton"))
+		.ButtonColorAndOpacity(FLinearColor{1.0f, 0.4f, 0.0f, 1.0f})
+		.OnClicked_Lambda([this]() -> FReply
 		{
 			// Isolating nothing is how we un-isolate
-			OnLayerIsolated( UE::FSdfLayerWeak{} );
+			OnLayerIsolated(UE::FSdfLayerWeak{});
 			return FReply::Handled();
 		})
-		.Visibility_Lambda( [this]() -> EVisibility
+		.Visibility_Lambda([this]() -> EVisibility
 		{
-			if ( const AUsdStageActor* StageActor = ViewModel.UsdStageActor.Get() )
+			if (const AUsdStageActor* StageActor = ViewModel.UsdStageActor.Get())
 			{
-				if ( StageActor->GetIsolatedUsdStage() )
+				if (StageActor->GetIsolatedUsdStage())
 				{
 					return EVisibility::Visible;
 				}
 			}
 			return EVisibility::Hidden;
 		})
-		.ToolTipText_Lambda( [this]() -> FText
+		.ToolTipText_Lambda([this]() -> FText
 		{
-			if ( const AUsdStageActor* StageActor = ViewModel.UsdStageActor.Get() )
+			if (const AUsdStageActor* StageActor = ViewModel.UsdStageActor.Get())
 			{
-				if ( const UE::FUsdStage& CurrentStage = StageActor->GetUsdStage() )
+				if (const UE::FUsdStage& CurrentStage = StageActor->GetUsdStage())
 				{
 					return FText::Format(
-						LOCTEXT( "IsolateWarningButtonToolTip", "The USD Stage editor and the Unreal level are displaying an isolated stage with '{0}' as its root layer.\nPress this button to revert to displaying the entire composed stage." ),
-						FText::FromString( CurrentStage.GetRootLayer().GetDisplayName() )
+						LOCTEXT("IsolateWarningButtonToolTip", "The USD Stage editor and the Unreal level are displaying an isolated stage with '{0}' as its root layer.\nPress this button to revert to displaying the entire composed stage."),
+						FText::FromString(CurrentStage.GetRootLayer().GetDisplayName())
 					);
 				}
 			}
@@ -822,23 +825,21 @@ TSharedRef< SWidget > SUsdStage::MakeIsolateWarningButton()
 		})
 		.Content()
 		[
-			SNew( STextBlock )
-			.Text( LOCTEXT( "IsolateWarningButtonText", "Isolated Mode" ) )
+			SNew(STextBlock)
+			.Text(LOCTEXT("IsolateWarningButtonText", "Isolated Mode"))
 		];
+	// clang-format on
 }
 
-void SUsdStage::FillFileMenu( FMenuBuilder& MenuBuilder )
+void SUsdStage::FillFileMenu(FMenuBuilder& MenuBuilder)
 {
-	MenuBuilder.BeginSection( "File", LOCTEXT("File", "File") );
+	MenuBuilder.BeginSection("File", LOCTEXT("File", "File"));
 	{
 		MenuBuilder.AddMenuEntry(
 			LOCTEXT("New", "New"),
 			LOCTEXT("New_ToolTip", "Creates a new layer and opens the stage with it at its root"),
 			FSlateIcon(),
-			FUIAction(
-				FExecuteAction::CreateSP( this, &SUsdStage::FileNew ),
-				FCanExecuteAction()
-			),
+			FUIAction(FExecuteAction::CreateSP(this, &SUsdStage::FileNew), FCanExecuteAction()),
 			NAME_None,
 			EUserInterfaceActionType::Button
 		);
@@ -848,10 +849,12 @@ void SUsdStage::FillFileMenu( FMenuBuilder& MenuBuilder )
 			LOCTEXT("Open_ToolTip", "Pick an existing USD file"),
 			FSlateIcon(),
 			FUIAction(
-				FExecuteAction::CreateLambda([this]()
-				{
-					FileOpen();
-				}),
+				FExecuteAction::CreateLambda(
+					[this]()
+					{
+						FileOpen();
+					}
+				),
 				FCanExecuteAction()
 			),
 			NAME_None,
@@ -863,81 +866,85 @@ void SUsdStage::FillFileMenu( FMenuBuilder& MenuBuilder )
 			LOCTEXT("Save_ToolTip", "Saves the stage"),
 			FSlateIcon(),
 			FUIAction(
-				FExecuteAction::CreateLambda([this]()
-				{
-					FileSave();
-				}),
-				FCanExecuteAction::CreateLambda( [this]()
-				{
-					if ( const AUsdStageActor* StageActor = ViewModel.UsdStageActor.Get() )
+				FExecuteAction::CreateLambda(
+					[this]()
+					{
+						FileSave();
+					}
+				),
+				FCanExecuteAction::CreateLambda(
+					[this]()
 					{
-						if ( UE::FUsdStage Stage = StageActor->GetUsdStage() )
+						if (const AUsdStageActor* StageActor = ViewModel.UsdStageActor.Get())
 						{
-							return true;
+							if (UE::FUsdStage Stage = StageActor->GetUsdStage())
+							{
+								return true;
+							}
 						}
-					}
 
-					return false;
-				})
+						return false;
+					}
+				)
 			),
 			NAME_None,
 			EUserInterfaceActionType::Button
 		);
 
-		MenuBuilder.AddSubMenu(
-			LOCTEXT( "Export", "Export" ),
-			FText::GetEmpty(),
-			FNewMenuDelegate::CreateSP( this, &SUsdStage::FillExportSubMenu )
-		);
+		MenuBuilder.AddSubMenu(LOCTEXT("Export", "Export"), FText::GetEmpty(), FNewMenuDelegate::CreateSP(this, &SUsdStage::FillExportSubMenu));
 	}
 	MenuBuilder.EndSection();
 
-	MenuBuilder.BeginSection( "Reload", LOCTEXT( "Reload", "Reload" ) );
+	MenuBuilder.BeginSection("Reload", LOCTEXT("Reload", "Reload"));
 	{
 		MenuBuilder.AddMenuEntry(
 			LOCTEXT("Reload", "Reload"),
 			LOCTEXT("Reload_ToolTip", "Reloads the stage from disk, keeping aspects of the session intact"),
 			FSlateIcon(),
 			FUIAction(
-				FExecuteAction::CreateSP( this, &SUsdStage::FileReload ),
-				FCanExecuteAction::CreateLambda( [this]()
-				{
-					if ( const AUsdStageActor* StageActor = ViewModel.UsdStageActor.Get() )
+				FExecuteAction::CreateSP(this, &SUsdStage::FileReload),
+				FCanExecuteAction::CreateLambda(
+					[this]()
 					{
-						if ( UE::FUsdStage Stage = StageActor->GetUsdStage() )
+						if (const AUsdStageActor* StageActor = ViewModel.UsdStageActor.Get())
 						{
-							if ( !Stage.GetRootLayer().IsAnonymous() )
+							if (UE::FUsdStage Stage = StageActor->GetUsdStage())
 							{
-								return true;
+								if (!Stage.GetRootLayer().IsAnonymous())
+								{
+									return true;
+								}
 							}
 						}
-					}
 
-					return false;
-				})
+						return false;
+					}
+				)
 			),
 			NAME_None,
 			EUserInterfaceActionType::Button
 		);
 
 		MenuBuilder.AddMenuEntry(
-			LOCTEXT( "ResetState", "Reset state" ),
-			LOCTEXT( "ResetState_ToolTip", "Resets the session layer and other options like edit target and muted layers" ),
+			LOCTEXT("ResetState", "Reset state"),
+			LOCTEXT("ResetState_ToolTip", "Resets the session layer and other options like edit target and muted layers"),
 			FSlateIcon(),
 			FUIAction(
-				FExecuteAction::CreateSP( this, &SUsdStage::FileReset ),
-				FCanExecuteAction::CreateLambda( [this]()
-				{
-					if ( const AUsdStageActor* StageActor = ViewModel.UsdStageActor.Get() )
+				FExecuteAction::CreateSP(this, &SUsdStage::FileReset),
+				FCanExecuteAction::CreateLambda(
+					[this]()
 					{
-						if ( UE::FUsdStage Stage = StageActor->GetUsdStage() )
+						if (const AUsdStageActor* StageActor = ViewModel.UsdStageActor.Get())
 						{
-							return true;
+							if (UE::FUsdStage Stage = StageActor->GetUsdStage())
+							{
+								return true;
+							}
 						}
-					}
 
-					return false;
-				})
+						return false;
+					}
+				)
 			),
 			NAME_None,
 			EUserInterfaceActionType::Button
@@ -945,7 +952,7 @@ void SUsdStage::FillFileMenu( FMenuBuilder& MenuBuilder )
 	}
 	MenuBuilder.EndSection();
 
-	MenuBuilder.BeginSection( "Close", LOCTEXT( "Close", "Close" ) );
+	MenuBuilder.BeginSection("Close", LOCTEXT("Close", "Close"));
 	{
 		MenuBuilder.AddMenuEntry(
 			LOCTEXT("Close", "Close"),
@@ -975,28 +982,30 @@ void SUsdStage::FillFileMenu( FMenuBuilder& MenuBuilder )
 	MenuBuilder.EndSection();
 }
 
-void SUsdStage::FillActionsMenu( FMenuBuilder& MenuBuilder )
+void SUsdStage::FillActionsMenu(FMenuBuilder& MenuBuilder)
 {
-	MenuBuilder.BeginSection( "Actions", LOCTEXT("Actions", "Actions") );
+	MenuBuilder.BeginSection("Actions", LOCTEXT("Actions", "Actions"));
 	{
 		MenuBuilder.AddMenuEntry(
 			LOCTEXT("Import", "Import..."),
 			LOCTEXT("Import_ToolTip", "Imports the stage as Unreal assets"),
 			FSlateIcon(),
 			FUIAction(
-				FExecuteAction::CreateSP( this, &SUsdStage::ActionsImportWithDialog ),
-				FCanExecuteAction::CreateLambda( [this]()
-				{
-					if ( const AUsdStageActor* StageActor = ViewModel.UsdStageActor.Get() )
+				FExecuteAction::CreateSP(this, &SUsdStage::ActionsImportWithDialog),
+				FCanExecuteAction::CreateLambda(
+					[this]()
 					{
-						if ( UE::FUsdStage Stage = StageActor->GetUsdStage() )
+						if (const AUsdStageActor* StageActor = ViewModel.UsdStageActor.Get())
 						{
-							return true;
+							if (UE::FUsdStage Stage = StageActor->GetUsdStage())
+							{
+								return true;
+							}
 						}
-					}
 
-					return false;
-				})
+						return false;
+					}
+				)
 			),
 			NAME_None,
 			EUserInterfaceActionType::Button
@@ -1007,134 +1016,159 @@ void SUsdStage::FillActionsMenu( FMenuBuilder& MenuBuilder )
 
 void SUsdStage::FillOptionsMenu(FMenuBuilder& MenuBuilder)
 {
-	MenuBuilder.BeginSection( TEXT( "StageOptions" ), LOCTEXT( "StageOptions", "Stage options" ) );
+	MenuBuilder.BeginSection(TEXT("StageOptions"), LOCTEXT("StageOptions", "Stage options"));
 	{
 		MenuBuilder.AddSubMenu(
 			LOCTEXT("State", "Stage state"),
 			LOCTEXT("State_ToolTip", "Whether we should open the USD Stage or spawn assets, actors and components"),
-			FNewMenuDelegate::CreateSP(this, &SUsdStage::FillStageStateSubMenu));
+			FNewMenuDelegate::CreateSP(this, &SUsdStage::FillStageStateSubMenu)
+		);
 
 		MenuBuilder.AddSubMenu(
 			LOCTEXT("Payloads", "Payloads"),
 			LOCTEXT("Payloads_ToolTip", "What to do with payloads when initially opening the stage"),
-			FNewMenuDelegate::CreateSP(this, &SUsdStage::FillPayloadsSubMenu));
+			FNewMenuDelegate::CreateSP(this, &SUsdStage::FillPayloadsSubMenu)
+		);
 
 		MenuBuilder.AddSubMenu(
 			LOCTEXT("PurposesToLoad", "Purposes to load"),
 			LOCTEXT("PurposesToLoad_ToolTip", "Only load prims with these specific purposes from the USD stage"),
-			FNewMenuDelegate::CreateSP(this, &SUsdStage::FillPurposesToLoadSubMenu));
+			FNewMenuDelegate::CreateSP(this, &SUsdStage::FillPurposesToLoadSubMenu)
+		);
 
 		MenuBuilder.AddSubMenu(
 			LOCTEXT("RenderContext", "Render context"),
 			LOCTEXT("RenderContext_ToolTip", "Choose which render context to use when parsing materials"),
-			FNewMenuDelegate::CreateSP(this, &SUsdStage::FillRenderContextSubMenu));
+			FNewMenuDelegate::CreateSP(this, &SUsdStage::FillRenderContextSubMenu)
+		);
 
 		MenuBuilder.AddSubMenu(
-			LOCTEXT( "MaterialPurpose", "Material purpose" ),
-			LOCTEXT( "MaterialPurpose_ToolTip", "Material purpose to use when parsing material bindings in addition to the \"allPurpose\" fallback" ),
-			FNewMenuDelegate::CreateSP( this, &SUsdStage::FillMaterialPurposeSubMenu ) );
+			LOCTEXT("MaterialPurpose", "Material purpose"),
+			LOCTEXT("MaterialPurpose_ToolTip", "Material purpose to use when parsing material bindings in addition to the \"allPurpose\" fallback"),
+			FNewMenuDelegate::CreateSP(this, &SUsdStage::FillMaterialPurposeSubMenu)
+		);
 
 		MenuBuilder.AddSubMenu(
-			LOCTEXT( "RootMotionHandling", "Root motion handling" ),
-			LOCTEXT( "RootMotionHandling_ToolTip", "Choose how to handle root motion for generated AnimSequences and LevelSequences" ),
-			FNewMenuDelegate::CreateSP( this, &SUsdStage::FillRootMotionSubMenu ) );
+			LOCTEXT("RootMotionHandling", "Root motion handling"),
+			LOCTEXT("RootMotionHandling_ToolTip", "Choose how to handle root motion for generated AnimSequences and LevelSequences"),
+			FNewMenuDelegate::CreateSP(this, &SUsdStage::FillRootMotionSubMenu)
+		);
 
 		MenuBuilder.AddSubMenu(
 			LOCTEXT("SubdivLevel", "Subdivision level"),
 			LOCTEXT("SubdivLevel_ToolTip", "Subdivision level to use for all subdivision meshes on the opened stage. 0 means 'don't subdivide'"),
-			FNewMenuDelegate::CreateSP(this, &SUsdStage::FillSubdivisionLevelSubMenu));
+			FNewMenuDelegate::CreateSP(this, &SUsdStage::FillSubdivisionLevelSubMenu)
+		);
 
 		MenuBuilder.AddSubMenu(
 			LOCTEXT("Metadata", "Metadata"),
 			LOCTEXT("Metadata_ToolTip", "How to collect and handle metadata from USD prims"),
-			FNewMenuDelegate::CreateSP(this, &SUsdStage::FillMetadataSubMenu));
+			FNewMenuDelegate::CreateSP(this, &SUsdStage::FillMetadataSubMenu)
+		);
 
 		MenuBuilder.AddSubMenu(
-			LOCTEXT( "Collapsing", "Collapsing" ),
-			LOCTEXT( "Collapsing_ToolTip", "Whether to try to combine individual assets and components of the same type on a Kind-per-Kind basis, like multiple Mesh prims into a single Static Mesh" ),
-			FNewMenuDelegate::CreateSP( this, &SUsdStage::FillCollapsingSubMenu ) );
+			LOCTEXT("Collapsing", "Collapsing"),
+			LOCTEXT(
+				"Collapsing_ToolTip",
+				"Whether to try to combine individual assets and components of the same type on a Kind-per-Kind basis, like multiple Mesh prims into "
+				"a single Static Mesh"
+			),
+			FNewMenuDelegate::CreateSP(this, &SUsdStage::FillCollapsingSubMenu)
+		);
 
 		MenuBuilder.AddSubMenu(
 			LOCTEXT("AssetReuse", "Asset reuse"),
 			LOCTEXT("AssetReuse_ToolTip", "How to behave when generating identical assets from different prims"),
-			FNewMenuDelegate::CreateSP(this, &SUsdStage::FillAssetReuseSubMenu));
+			FNewMenuDelegate::CreateSP(this, &SUsdStage::FillAssetReuseSubMenu)
+		);
 
 		MenuBuilder.AddSubMenu(
-			LOCTEXT( "InterpolationType", "Interpolation type" ),
-			LOCTEXT( "InterpolationType_ToolTip", "Whether to interpolate between time samples linearly or with 'held' (i.e. constant) interpolation" ),
-			FNewMenuDelegate::CreateSP( this, &SUsdStage::FillInterpolationTypeSubMenu ) );
+			LOCTEXT("InterpolationType", "Interpolation type"),
+			LOCTEXT("InterpolationType_ToolTip", "Whether to interpolate between time samples linearly or with 'held' (i.e. constant) interpolation"),
+			FNewMenuDelegate::CreateSP(this, &SUsdStage::FillInterpolationTypeSubMenu)
+		);
 	}
 	MenuBuilder.EndSection();
 
-	MenuBuilder.BeginSection( TEXT( "EditorOptions" ), LOCTEXT( "EditorOptions", "Editor options" ) );
+	MenuBuilder.BeginSection(TEXT("EditorOptions"), LOCTEXT("EditorOptions", "Editor options"));
 	{
 		MenuBuilder.AddSubMenu(
-			LOCTEXT( "SelectionText", "Selection" ),
-			LOCTEXT( "SelectionText_ToolTip", "How the selection of prims, actors and components should behave" ),
-			FNewMenuDelegate::CreateSP( this, &SUsdStage::FillSelectionSubMenu ) );
+			LOCTEXT("SelectionText", "Selection"),
+			LOCTEXT("SelectionText_ToolTip", "How the selection of prims, actors and components should behave"),
+			FNewMenuDelegate::CreateSP(this, &SUsdStage::FillSelectionSubMenu)
+		);
 
 		MenuBuilder.AddSubMenu(
-			LOCTEXT( "NaniteSettings", "Nanite" ),
-			LOCTEXT( "NaniteSettings_ToolTip", "Configure how to use Nanite for generated assets" ),
-			FNewMenuDelegate::CreateSP( this, &SUsdStage::FillNaniteThresholdSubMenu ),
+			LOCTEXT("NaniteSettings", "Nanite"),
+			LOCTEXT("NaniteSettings_ToolTip", "Configure how to use Nanite for generated assets"),
+			FNewMenuDelegate::CreateSP(this, &SUsdStage::FillNaniteThresholdSubMenu),
 			false,
 			FSlateIcon(),
-			false );
+			false
+		);
 	}
 	MenuBuilder.EndSection();
 }
 
-void SUsdStage::FillExportSubMenu( FMenuBuilder& MenuBuilder )
+void SUsdStage::FillExportSubMenu(FMenuBuilder& MenuBuilder)
 {
 	const FString FilePath;
 
 	MenuBuilder.AddMenuEntry(
-		LOCTEXT( "ExportAll", "Export all layers..." ),
-		LOCTEXT( "ExportAll_ToolTip", "Exports copies of all file-based layers in the stage's layer stack to a new folder" ),
+		LOCTEXT("ExportAll", "Export all layers..."),
+		LOCTEXT("ExportAll_ToolTip", "Exports copies of all file-based layers in the stage's layer stack to a new folder"),
 		FSlateIcon(),
 		FUIAction(
-			FExecuteAction::CreateLambda([this]()
-			{
-				FileExportAllLayers();
-			}),
-			FCanExecuteAction::CreateLambda( [this]()
-			{
-				if ( const AUsdStageActor* StageActor = ViewModel.UsdStageActor.Get() )
+			FExecuteAction::CreateLambda(
+				[this]()
+				{
+					FileExportAllLayers();
+				}
+			),
+			FCanExecuteAction::CreateLambda(
+				[this]()
 				{
-					if ( UE::FUsdStage Stage = StageActor->GetUsdStage() )
+					if (const AUsdStageActor* StageActor = ViewModel.UsdStageActor.Get())
 					{
-						return true;
+						if (UE::FUsdStage Stage = StageActor->GetUsdStage())
+						{
+							return true;
+						}
 					}
-				}
 
-				return false;
-			})
+					return false;
+				}
+			)
 		),
 		NAME_None,
 		EUserInterfaceActionType::Button
 	);
 
 	MenuBuilder.AddMenuEntry(
-		LOCTEXT( "ExportFlattened", "Export flattened stage..." ),
-		LOCTEXT( "ExportFlattened_ToolTip", "Flattens the current stage to a single USD layer and exports it as a new USD file" ),
+		LOCTEXT("ExportFlattened", "Export flattened stage..."),
+		LOCTEXT("ExportFlattened_ToolTip", "Flattens the current stage to a single USD layer and exports it as a new USD file"),
 		FSlateIcon(),
 		FUIAction(
-			FExecuteAction::CreateLambda([this]()
-			{
-				FileExportFlattenedStage();
-			}),
-			FCanExecuteAction::CreateLambda( [this]()
-			{
-				if ( const AUsdStageActor* StageActor = ViewModel.UsdStageActor.Get() )
+			FExecuteAction::CreateLambda(
+				[this]()
+				{
+					FileExportFlattenedStage();
+				}
+			),
+			FCanExecuteAction::CreateLambda(
+				[this]()
 				{
-					if ( UE::FUsdStage Stage = StageActor->GetUsdStage() )
+					if (const AUsdStageActor* StageActor = ViewModel.UsdStageActor.Get())
 					{
-						return true;
+						if (UE::FUsdStage Stage = StageActor->GetUsdStage())
+						{
+							return true;
+						}
 					}
-				}
 
-				return false;
-			})
+					return false;
+				}
+			)
 		),
 		NAME_None,
 		EUserInterfaceActionType::Button
@@ -1214,34 +1248,38 @@ void SUsdStage::FillPayloadsSubMenu(FMenuBuilder& MenuBuilder)
 		LOCTEXT("LoadAll_ToolTip", "Loads all payloads initially"),
 		FSlateIcon(),
 		FUIAction(
-			FExecuteAction::CreateLambda([this]()
-			{
-				if ( AUsdStageActor* StageActor = GetStageActorOrCDO() )
+			FExecuteAction::CreateLambda(
+				[this]()
 				{
-					FScopedTransaction Transaction(FText::Format(
-						LOCTEXT("SetLoadAllTransaction", "Set USD stage actor '{0}' actor to load all payloads initially"),
-						FText::FromString(StageActor->GetActorLabel())
-					));
+					if (AUsdStageActor* StageActor = GetStageActorOrCDO())
+					{
+						FScopedTransaction Transaction(FText::Format(
+							LOCTEXT("SetLoadAllTransaction", "Set USD stage actor '{0}' actor to load all payloads initially"),
+							FText::FromString(StageActor->GetActorLabel())
+						));
 
-					// c.f. comment in SUsdStage::FillCollapsingSubMenu
-					TGuardValue<bool> MaintainSelectionGuard(bUpdatingViewportSelection, true);
+						// c.f. comment in SUsdStage::FillCollapsingSubMenu
+						TGuardValue<bool> MaintainSelectionGuard(bUpdatingViewportSelection, true);
 
-					StageActor->SetInitialLoadSet(EUsdInitialLoadSet::LoadAll);
-					if (StageActor->IsTemplate())
-					{
-						StageActor->SaveConfig();
+						StageActor->SetInitialLoadSet(EUsdInitialLoadSet::LoadAll);
+						if (StageActor->IsTemplate())
+						{
+							StageActor->SaveConfig();
+						}
 					}
 				}
-			}),
+			),
 			FCanExecuteAction{},
-			FIsActionChecked::CreateLambda([this]()
-			{
-				if ( AUsdStageActor* StageActor = GetStageActorOrCDO() )
+			FIsActionChecked::CreateLambda(
+				[this]()
 				{
-					return StageActor->InitialLoadSet == EUsdInitialLoadSet::LoadAll;
+					if (AUsdStageActor* StageActor = GetStageActorOrCDO())
+					{
+						return StageActor->InitialLoadSet == EUsdInitialLoadSet::LoadAll;
+					}
+					return false;
 				}
-				return false;
-			})
+			)
 		),
 		NAME_None,
 		EUserInterfaceActionType::RadioButton
@@ -1252,34 +1290,38 @@ void SUsdStage::FillPayloadsSubMenu(FMenuBuilder& MenuBuilder)
 		LOCTEXT("LoadNone_ToolTip", "Don't load any payload initially"),
 		FSlateIcon(),
 		FUIAction(
-			FExecuteAction::CreateLambda([this]()
-			{
-				if ( AUsdStageActor* StageActor = GetStageActorOrCDO() )
+			FExecuteAction::CreateLambda(
+				[this]()
 				{
-					FScopedTransaction Transaction(FText::Format(
-						LOCTEXT("SetLoadNoneTransaction", "Set USD stage actor '{0}' actor to load no payloads initially"),
-						FText::FromString(StageActor->GetActorLabel())
-					));
+					if (AUsdStageActor* StageActor = GetStageActorOrCDO())
+					{
+						FScopedTransaction Transaction(FText::Format(
+							LOCTEXT("SetLoadNoneTransaction", "Set USD stage actor '{0}' actor to load no payloads initially"),
+							FText::FromString(StageActor->GetActorLabel())
+						));
 
-					// c.f. comment in SUsdStage::FillCollapsingSubMenu
-					TGuardValue<bool> MaintainSelectionGuard(bUpdatingViewportSelection, true);
+						// c.f. comment in SUsdStage::FillCollapsingSubMenu
+						TGuardValue<bool> MaintainSelectionGuard(bUpdatingViewportSelection, true);
 
-					StageActor->SetInitialLoadSet(EUsdInitialLoadSet::LoadNone);
-					if (StageActor->IsTemplate())
-					{
-						StageActor->SaveConfig();
+						StageActor->SetInitialLoadSet(EUsdInitialLoadSet::LoadNone);
+						if (StageActor->IsTemplate())
+						{
+							StageActor->SaveConfig();
+						}
 					}
 				}
-			}),
+			),
 			FCanExecuteAction{},
-			FIsActionChecked::CreateLambda([this]()
-			{
-				if ( AUsdStageActor* StageActor = GetStageActorOrCDO() )
+			FIsActionChecked::CreateLambda(
+				[this]()
 				{
-					return StageActor->InitialLoadSet == EUsdInitialLoadSet::LoadNone;
+					if (AUsdStageActor* StageActor = GetStageActorOrCDO())
+					{
+						return StageActor->InitialLoadSet == EUsdInitialLoadSet::LoadNone;
+					}
+					return false;
 				}
-				return false;
-			})
+			)
 		),
 		NAME_None,
 		EUserInterfaceActionType::RadioButton
@@ -1295,51 +1337,55 @@ void SUsdStage::FillPurposesToLoadSubMenu(FMenuBuilder& MenuBuilder)
 			FText::GetEmpty(),
 			FSlateIcon(),
 			FUIAction(
-				FExecuteAction::CreateLambda([this, Purpose]()
-				{
-					if ( AUsdStageActor* StageActor = GetStageActorOrCDO() )
+				FExecuteAction::CreateLambda(
+					[this, Purpose]()
 					{
-						FScopedTransaction Transaction(FText::Format(
-							LOCTEXT("PurposesToLoadTransaction", "Change purposes to load for USD stage actor '{0}'"),
-							FText::FromString(StageActor->GetActorLabel())
-						));
+						if (AUsdStageActor* StageActor = GetStageActorOrCDO())
+						{
+							FScopedTransaction Transaction(FText::Format(
+								LOCTEXT("PurposesToLoadTransaction", "Change purposes to load for USD stage actor '{0}'"),
+								FText::FromString(StageActor->GetActorLabel())
+							));
 
-						// c.f. comment in SUsdStage::FillCollapsingSubMenu
-						TGuardValue<bool> MaintainSelectionGuard(bUpdatingViewportSelection, true);
+							// c.f. comment in SUsdStage::FillCollapsingSubMenu
+							TGuardValue<bool> MaintainSelectionGuard(bUpdatingViewportSelection, true);
 
-						StageActor->SetPurposesToLoad((int32)((EUsdPurpose)StageActor->PurposesToLoad ^ Purpose));
-						if (StageActor->IsTemplate())
-						{
-							StageActor->SaveConfig();
+							StageActor->SetPurposesToLoad((int32)((EUsdPurpose)StageActor->PurposesToLoad ^ Purpose));
+							if (StageActor->IsTemplate())
+							{
+								StageActor->SaveConfig();
+							}
 						}
 					}
-				}),
+				),
 				FCanExecuteAction{},
-				FIsActionChecked::CreateLambda([this, Purpose]()
-				{
-					if ( AUsdStageActor* StageActor = GetStageActorOrCDO() )
+				FIsActionChecked::CreateLambda(
+					[this, Purpose]()
 					{
-						return EnumHasAllFlags((EUsdPurpose)StageActor->PurposesToLoad, Purpose);
+						if (AUsdStageActor* StageActor = GetStageActorOrCDO())
+						{
+							return EnumHasAllFlags((EUsdPurpose)StageActor->PurposesToLoad, Purpose);
+						}
+						return false;
 					}
-					return false;
-				})
+				)
 			),
 			NAME_None,
 			EUserInterfaceActionType::ToggleButton
 		);
 	};
 
-	AddPurposeEntry(EUsdPurpose::Proxy,  LOCTEXT("ProxyPurpose",  "Proxy"));
+	AddPurposeEntry(EUsdPurpose::Proxy, LOCTEXT("ProxyPurpose", "Proxy"));
 	AddPurposeEntry(EUsdPurpose::Render, LOCTEXT("RenderPurpose", "Render"));
-	AddPurposeEntry(EUsdPurpose::Guide,  LOCTEXT("GuidePurpose",  "Guide"));
+	AddPurposeEntry(EUsdPurpose::Guide, LOCTEXT("GuidePurpose", "Guide"));
 }
 
-void SUsdStage::FillRenderContextSubMenu( FMenuBuilder& MenuBuilder )
+void SUsdStage::FillRenderContextSubMenu(FMenuBuilder& MenuBuilder)
 {
 	auto AddRenderContextEntry = [&](const FName& RenderContext)
 	{
-		FText RenderContextName = FText::FromName( RenderContext );
-		if ( RenderContext.IsNone() )
+		FText RenderContextName = FText::FromName(RenderContext);
+		if (RenderContext.IsNone())
 		{
 			RenderContextName = LOCTEXT("UniversalRenderContext", "universal");
 		}
@@ -1349,108 +1395,111 @@ void SUsdStage::FillRenderContextSubMenu( FMenuBuilder& MenuBuilder )
 			FText::GetEmpty(),
 			FSlateIcon(),
 			FUIAction(
-				FExecuteAction::CreateLambda([this, RenderContext]()
-				{
-					if ( AUsdStageActor* StageActor = GetStageActorOrCDO() )
+				FExecuteAction::CreateLambda(
+					[this, RenderContext]()
 					{
-						FScopedTransaction Transaction(FText::Format(
-							LOCTEXT("RenderContextToLoadTransaction", "Change render context to load for USD stage actor '{0}'"),
-							FText::FromString(StageActor->GetActorLabel())
-						));
+						if (AUsdStageActor* StageActor = GetStageActorOrCDO())
+						{
+							FScopedTransaction Transaction(FText::Format(
+								LOCTEXT("RenderContextToLoadTransaction", "Change render context to load for USD stage actor '{0}'"),
+								FText::FromString(StageActor->GetActorLabel())
+							));
 
-						// c.f. comment in SUsdStage::FillCollapsingSubMenu
-						TGuardValue<bool> MaintainSelectionGuard(bUpdatingViewportSelection, true);
+							// c.f. comment in SUsdStage::FillCollapsingSubMenu
+							TGuardValue<bool> MaintainSelectionGuard(bUpdatingViewportSelection, true);
 
-						StageActor->SetRenderContext(RenderContext);
-						if (StageActor->IsTemplate())
-						{
-							StageActor->SaveConfig();
+							StageActor->SetRenderContext(RenderContext);
+							if (StageActor->IsTemplate())
+							{
+								StageActor->SaveConfig();
+							}
 						}
 					}
-				}),
+				),
 				FCanExecuteAction{},
-				FIsActionChecked::CreateLambda([this, RenderContext]()
-				{
-					if ( AUsdStageActor* StageActor = GetStageActorOrCDO() )
+				FIsActionChecked::CreateLambda(
+					[this, RenderContext]()
 					{
-						return StageActor->RenderContext == RenderContext;
+						if (AUsdStageActor* StageActor = GetStageActorOrCDO())
+						{
+							return StageActor->RenderContext == RenderContext;
+						}
+						return false;
 					}
-					return false;
-				})
+				)
 			),
 			NAME_None,
 			EUserInterfaceActionType::RadioButton
 		);
 	};
 
-	IUsdSchemasModule& UsdSchemasModule = FModuleManager::Get().LoadModuleChecked< IUsdSchemasModule >( TEXT("USDSchemas") );
+	IUsdSchemasModule& UsdSchemasModule = FModuleManager::Get().LoadModuleChecked<IUsdSchemasModule>(TEXT("USDSchemas"));
 
-	for ( const FName& RenderContext : UsdSchemasModule.GetRenderContextRegistry().GetRenderContexts() )
+	for (const FName& RenderContext : UsdSchemasModule.GetRenderContextRegistry().GetRenderContexts())
 	{
-		AddRenderContextEntry( RenderContext );
+		AddRenderContextEntry(RenderContext);
 	}
 }
 
-void SUsdStage::FillMaterialPurposeSubMenu( FMenuBuilder& MenuBuilder )
+void SUsdStage::FillMaterialPurposeSubMenu(FMenuBuilder& MenuBuilder)
 {
 	MaterialPurposes = {
-		MakeShared<FString>( UnrealIdentifiers::MaterialAllPurpose ),
-		MakeShared<FString>( UnrealIdentifiers::MaterialPreviewPurpose ),
-		MakeShared<FString>( UnrealIdentifiers::MaterialFullPurpose )
-	};
+		MakeShared<FString>(UnrealIdentifiers::MaterialAllPurpose),
+		MakeShared<FString>(UnrealIdentifiers::MaterialPreviewPurpose),
+		MakeShared<FString>(UnrealIdentifiers::MaterialFullPurpose)};
 
 	// Add additional purposes from project settings
-	if ( const UUsdProjectSettings* ProjectSettings = GetDefault<UUsdProjectSettings>() )
+	if (const UUsdProjectSettings* ProjectSettings = GetDefault<UUsdProjectSettings>())
 	{
-		MaterialPurposes.Reserve( MaterialPurposes.Num() + ProjectSettings->AdditionalMaterialPurposes.Num() );
+		MaterialPurposes.Reserve(MaterialPurposes.Num() + ProjectSettings->AdditionalMaterialPurposes.Num());
 
 		TSet<FString> ExistingEntries = {
 			UnrealIdentifiers::MaterialAllPurpose,
 			UnrealIdentifiers::MaterialPreviewPurpose,
-			UnrealIdentifiers::MaterialFullPurpose
-		};
+			UnrealIdentifiers::MaterialFullPurpose};
 
-		for (const FName& AdditionalPurpose : ProjectSettings->AdditionalMaterialPurposes )
+		for (const FName& AdditionalPurpose : ProjectSettings->AdditionalMaterialPurposes)
 		{
 			FString AdditionalPurposeStr = AdditionalPurpose.ToString();
 
-			if ( !ExistingEntries.Contains( AdditionalPurposeStr ) )
+			if (!ExistingEntries.Contains(AdditionalPurposeStr))
 			{
-				ExistingEntries.Add( AdditionalPurposeStr );
-				MaterialPurposes.AddUnique( MakeShared<FString>( AdditionalPurposeStr ) );
+				ExistingEntries.Add(AdditionalPurposeStr);
+				MaterialPurposes.AddUnique(MakeShared<FString>(AdditionalPurposeStr));
 			}
 		}
 	}
 
-	TSharedRef<SBox> Box = SNew( SBox )
-	.Padding( FMargin( 8.0f, 0.0f ) )
-	.VAlign( VAlign_Center )
+	// clang-format off
+	TSharedRef<SBox> Box = SNew(SBox)
+	.Padding(FMargin(8.0f, 0.0f))
+	.VAlign(VAlign_Center)
 	[
 		// We have to use TSharedPtr<FString> here as the combobox actually contains a list view.
 		// Also, a regular SComboBox<FName> doesn't even call our OnGenerateWidget function in case the item is NAME_None,
 		// and we do need that case for allPurpose.
-		SNew( SComboBox< TSharedPtr<FString> > )
-		.OptionsSource( &MaterialPurposes )
-		.OnGenerateWidget_Lambda( [ & ]( TSharedPtr<FString> Option )
+		SNew(SComboBox<TSharedPtr<FString>>)
+		.OptionsSource(&MaterialPurposes)
+		.OnGenerateWidget_Lambda([&](TSharedPtr<FString> Option)
 		{
 			TSharedPtr<SWidget> Widget = SNullWidget::NullWidget;
-			if ( Option )
+			if (Option)
 			{
-				Widget = SNew( STextBlock )
-					.Text( FText::FromString( (*Option) == UnrealIdentifiers::MaterialAllPurpose
+				Widget = SNew(STextBlock)
+					.Text(FText::FromString((*Option) == UnrealIdentifiers::MaterialAllPurpose
 						? UnrealIdentifiers::MaterialAllPurposeText
 						: *Option
 					))
-					.Font( FAppStyle::GetFontStyle( "PropertyWindow.NormalFont" ) );
+					.Font(FAppStyle::GetFontStyle("PropertyWindow.NormalFont"));
 			}
 
 			return Widget.ToSharedRef();
 		})
-		.OnSelectionChanged_Lambda([this]( TSharedPtr<FString> ChosenOption, ESelectInfo::Type SelectInfo )
+		.OnSelectionChanged_Lambda([this](TSharedPtr<FString> ChosenOption, ESelectInfo::Type SelectInfo)
 		{
-			if ( ChosenOption )
+			if (ChosenOption)
 			{
-				if ( AUsdStageActor* StageActor = GetStageActorOrCDO() )
+				if (AUsdStageActor* StageActor = GetStageActorOrCDO())
 				{
 					FScopedTransaction Transaction(FText::Format(
 						LOCTEXT("ChangeMaterialPurpose", "Change selected material purpose for USD stage actor '{0}' to '{1}'"),
@@ -1470,12 +1519,12 @@ void SUsdStage::FillMaterialPurposeSubMenu( FMenuBuilder& MenuBuilder )
 			}
 		})
 		[
-			SNew( SEditableTextBox )
+			SNew(SEditableTextBox)
 			.Text_Lambda([this]() -> FText
 			{
-				if ( AUsdStageActor* StageActor = GetStageActorOrCDO() )
+				if (AUsdStageActor* StageActor = GetStageActorOrCDO())
 				{
-					return FText::FromString( StageActor->MaterialPurpose == *UnrealIdentifiers::MaterialAllPurpose
+					return FText::FromString(StageActor->MaterialPurpose == *UnrealIdentifiers::MaterialAllPurpose
 						? UnrealIdentifiers::MaterialAllPurposeText
 						: StageActor->MaterialPurpose.ToString()
 					);
@@ -1483,10 +1532,10 @@ void SUsdStage::FillMaterialPurposeSubMenu( FMenuBuilder& MenuBuilder )
 
 				return FText::GetEmpty();
 			})
-			.Font( FAppStyle::GetFontStyle( "PropertyWindow.NormalFont" ) )
-			.OnTextCommitted_Lambda( [this]( const FText& NewText, ETextCommit::Type CommitType )
+			.Font(FAppStyle::GetFontStyle("PropertyWindow.NormalFont"))
+			.OnTextCommitted_Lambda([this](const FText& NewText, ETextCommit::Type CommitType)
 			{
-				if ( CommitType != ETextCommit::OnEnter )
+				if (CommitType != ETextCommit::OnEnter)
 				{
 					return;
 				}
@@ -1495,25 +1544,25 @@ void SUsdStage::FillMaterialPurposeSubMenu( FMenuBuilder& MenuBuilder )
 				FName NewPurpose = *NewPurposeString;
 
 				bool bIsNew = true;
-				for ( const TSharedPtr<FString>& Purpose : MaterialPurposes )
+				for (const TSharedPtr<FString>& Purpose : MaterialPurposes)
 				{
-					if ( Purpose && *Purpose == NewPurposeString )
+					if (Purpose && *Purpose == NewPurposeString)
 					{
 						bIsNew = false;
 						break;
 					}
 				}
 
-				if ( bIsNew )
+				if (bIsNew)
 				{
-					if ( UUsdProjectSettings* ProjectSettings = GetMutableDefault<UUsdProjectSettings>() )
+					if (UUsdProjectSettings* ProjectSettings = GetMutableDefault<UUsdProjectSettings>())
 					{
-						ProjectSettings->AdditionalMaterialPurposes.AddUnique( NewPurpose );
+						ProjectSettings->AdditionalMaterialPurposes.AddUnique(NewPurpose);
 						ProjectSettings->SaveConfig();
 					}
 				}
 
-				if ( AUsdStageActor* StageActor = GetStageActorOrCDO() )
+				if (AUsdStageActor* StageActor = GetStageActorOrCDO())
 				{
 					FScopedTransaction Transaction(FText::Format(
 						LOCTEXT("TypeMaterialPurpose", "Change selected material purpose for USD stage actor '{0}' to '{1}'"),
@@ -1530,60 +1579,65 @@ void SUsdStage::FillMaterialPurposeSubMenu( FMenuBuilder& MenuBuilder )
 			})
 		]
 	];
+	// clang-format on
 
 	const bool bNoIndent = true;
-	MenuBuilder.AddWidget( Box, FText::GetEmpty(), bNoIndent );
+	MenuBuilder.AddWidget(Box, FText::GetEmpty(), bNoIndent);
 }
 
-void SUsdStage::FillRootMotionSubMenu( FMenuBuilder& MenuBuilder )
+void SUsdStage::FillRootMotionSubMenu(FMenuBuilder& MenuBuilder)
 {
 	auto AddRootMotionEntry = [&](const EUsdRootMotionHandling& HandlingStrategy, const FText& Text)
 	{
-		UEnum* Enum = StaticEnum< EUsdRootMotionHandling >();
-		FText ToolTip = Enum->GetToolTipTextByIndex( Enum->GetIndexByValue( static_cast< uint8 >( HandlingStrategy ) ) );
+		UEnum* Enum = StaticEnum<EUsdRootMotionHandling>();
+		FText ToolTip = Enum->GetToolTipTextByIndex(Enum->GetIndexByValue(static_cast<uint8>(HandlingStrategy)));
 
 		MenuBuilder.AddMenuEntry(
 			Text,
 			ToolTip,
 			FSlateIcon(),
 			FUIAction(
-				FExecuteAction::CreateLambda([this, HandlingStrategy]()
-				{
-					if ( AUsdStageActor* StageActor = GetStageActorOrCDO() )
+				FExecuteAction::CreateLambda(
+					[this, HandlingStrategy]()
 					{
-						FScopedTransaction Transaction(FText::Format(
-							LOCTEXT("RootMotionHandlingTransaction", "Change root motion handling strategy for USD stage actor '{0}'"),
-							FText::FromString(StageActor->GetActorLabel())
-						));
+						if (AUsdStageActor* StageActor = GetStageActorOrCDO())
+						{
+							FScopedTransaction Transaction(FText::Format(
+								LOCTEXT("RootMotionHandlingTransaction", "Change root motion handling strategy for USD stage actor '{0}'"),
+								FText::FromString(StageActor->GetActorLabel())
+							));
 
-						// c.f. comment in SUsdStage::FillCollapsingSubMenu
-						TGuardValue<bool> MaintainSelectionGuard(bUpdatingViewportSelection, true);
+							// c.f. comment in SUsdStage::FillCollapsingSubMenu
+							TGuardValue<bool> MaintainSelectionGuard(bUpdatingViewportSelection, true);
 
-						StageActor->SetRootMotionHandling(HandlingStrategy);
-						if (StageActor->IsTemplate())
-						{
-							StageActor->SaveConfig();
+							StageActor->SetRootMotionHandling(HandlingStrategy);
+							if (StageActor->IsTemplate())
+							{
+								StageActor->SaveConfig();
+							}
 						}
 					}
-				}),
+				),
 				FCanExecuteAction{},
-				FIsActionChecked::CreateLambda([this, HandlingStrategy]()
-				{
-					if ( AUsdStageActor* StageActor = GetStageActorOrCDO() )
+				FIsActionChecked::CreateLambda(
+					[this, HandlingStrategy]()
 					{
-						return StageActor->RootMotionHandling == HandlingStrategy;
+						if (AUsdStageActor* StageActor = GetStageActorOrCDO())
+						{
+							return StageActor->RootMotionHandling == HandlingStrategy;
+						}
+						return false;
 					}
-					return false;
-				})
+				)
 			),
 			NAME_None,
 			EUserInterfaceActionType::RadioButton
 		);
 	};
 
-	AddRootMotionEntry( EUsdRootMotionHandling::NoAdditionalRootMotion, LOCTEXT("NoAdditionalRootMotionText", "No additional root motion"));
-	AddRootMotionEntry( EUsdRootMotionHandling::UseMotionFromSkelRoot, LOCTEXT("UseMotionFromSkelRootText", "Use motion from SkelRoot"));
-	AddRootMotionEntry( EUsdRootMotionHandling::UseMotionFromSkeleton, LOCTEXT("UseMotionFromSkeletonText", "Use motion from Skeleton"));
+	AddRootMotionEntry(EUsdRootMotionHandling::NoAdditionalRootMotion, LOCTEXT("NoAdditionalRootMotionText", "No additional root motion"));
+	AddRootMotionEntry(EUsdRootMotionHandling::UseMotionFromSkelRoot, LOCTEXT("UseMotionFromSkelRootText", "Use motion from SkelRoot"));
+	AddRootMotionEntry(EUsdRootMotionHandling::UseMotionFromSkeleton, LOCTEXT("UseMotionFromSkeletonText", "Use motion from Skeleton"));
 }
 
 void SUsdStage::FillSubdivisionLevelSubMenu(FMenuBuilder& MenuBuilder)
@@ -1596,6 +1650,7 @@ void SUsdStage::FillSubdivisionLevelSubMenu(FMenuBuilder& MenuBuilder)
 	static IConsoleVariable* MaxSubdivCvar = IConsoleManager::Get().FindConsoleVariable(TEXT("USD.Subdiv.MaxSubdivLevel"));
 	const int32 MaxSubdivLevel = MaxSubdivCvar ? MaxSubdivCvar->GetInt() : 0;
 
+	// clang-format off
 	TSharedRef<SSpinBox<int32>> Slider = SNew(SSpinBox<int32>)
 		.MinValue(0)
 		.MaxValue(MaxSubdivLevel)
@@ -1607,6 +1662,7 @@ void SUsdStage::FillSubdivisionLevelSubMenu(FMenuBuilder& MenuBuilder)
 		.OnValueChanged(this, &SUsdStage::OnSubdivisionLevelValueChanged)
 		.SupportDynamicSliderMaxValue(true)
 		.OnValueCommitted(this, &SUsdStage::OnSubdivisionLevelValueCommitted);
+	// clang-format on
 
 	const bool bNoIndent = true;
 	MenuBuilder.AddWidget(Slider, FText::FromString(TEXT("Subdivision level: ")), bNoIndent);
@@ -1639,10 +1695,10 @@ void SUsdStage::FillMetadataSubMenu(FMenuBuilder& MenuBuilder)
 				return true;
 			}
 
-			for(const FProperty* ParentProp : PropertyAndParent.ParentProperties)
+			for (const FProperty* ParentProp : PropertyAndParent.ParentProperties)
 			{
 				const FName ParentPropertyName = ParentProp->GetFName();
-				if(ParentPropertyName == MetadataOptionsName)
+				if (ParentPropertyName == MetadataOptionsName)
 				{
 					return true;
 				}
@@ -1675,6 +1731,7 @@ void SUsdStage::FillMetadataSubMenu(FMenuBuilder& MenuBuilder)
 
 	// Override the widget with some extra width because the prefix filters need some space to be used
 	// comfortably
+	// clang-format off
 	TSharedRef<SBox> Box = SNew(SBox)
 	.Padding(FMargin(0.0f, 0.0f))
 	.VAlign(VAlign_Fill)
@@ -1683,168 +1740,183 @@ void SUsdStage::FillMetadataSubMenu(FMenuBuilder& MenuBuilder)
 	[
 		DetailView
 	];
+	// clang-format on
 
 	const bool bNoIndent = true;
 	MenuBuilder.AddWidget(Box, FText::GetEmpty(), bNoIndent);
 }
 
-void SUsdStage::FillCollapsingSubMenu( FMenuBuilder& MenuBuilder )
+void SUsdStage::FillCollapsingSubMenu(FMenuBuilder& MenuBuilder)
 {
 	MenuBuilder.AddMenuEntry(
-		LOCTEXT( "MergeIdenticalMaterialSlots", "Merge identical material slots" ),
-		LOCTEXT( "MergeIdenticalMaterialSlots_ToolTip", "If enabled, when multiple mesh prims are collapsed into a single static mesh, identical material slots are merged into one slot.\nOtherwise, material slots are simply appended to the list." ),
+		LOCTEXT("MergeIdenticalMaterialSlots", "Merge identical material slots"),
+		LOCTEXT(
+			"MergeIdenticalMaterialSlots_ToolTip",
+			"If enabled, when multiple mesh prims are collapsed into a single static mesh, identical material slots are merged into one "
+			"slot.\nOtherwise, material slots are simply appended to the list."
+		),
 		FSlateIcon(),
 		FUIAction(
-			FExecuteAction::CreateLambda( [this]()
-			{
-				if ( AUsdStageActor* StageActor = GetStageActorOrCDO() )
+			FExecuteAction::CreateLambda(
+				[this]()
 				{
-					FScopedTransaction Transaction( FText::Format(
-						LOCTEXT( "MergeIdenticalMaterialSlotsTransaction", "Toggle bMergeIdenticalMaterialSlots on USD stage actor '{0}'" ),
-						FText::FromString( StageActor->GetActorLabel() )
-					) );
+					if (AUsdStageActor* StageActor = GetStageActorOrCDO())
+					{
+						FScopedTransaction Transaction(FText::Format(
+							LOCTEXT("MergeIdenticalMaterialSlotsTransaction", "Toggle bMergeIdenticalMaterialSlots on USD stage actor '{0}'"),
+							FText::FromString(StageActor->GetActorLabel())
+						));
 
-					// c.f. comment within AddKindToCollapseEntry just below
-					TGuardValue<bool> MaintainSelectionGuard(bUpdatingViewportSelection, true);
+						// c.f. comment within AddKindToCollapseEntry just below
+						TGuardValue<bool> MaintainSelectionGuard(bUpdatingViewportSelection, true);
 
-					StageActor->SetMergeIdenticalMaterialSlots(!StageActor->bMergeIdenticalMaterialSlots);
-					if (StageActor->IsTemplate())
-					{
-						StageActor->SaveConfig();
+						StageActor->SetMergeIdenticalMaterialSlots(!StageActor->bMergeIdenticalMaterialSlots);
+						if (StageActor->IsTemplate())
+						{
+							StageActor->SaveConfig();
+						}
 					}
 				}
-			}),
+			),
 			FCanExecuteAction{},
-			FIsActionChecked::CreateLambda( [this]()
-			{
-				if ( AUsdStageActor* StageActor = GetStageActorOrCDO() )
+			FIsActionChecked::CreateLambda(
+				[this]()
 				{
-					return StageActor->bMergeIdenticalMaterialSlots;
+					if (AUsdStageActor* StageActor = GetStageActorOrCDO())
+					{
+						return StageActor->bMergeIdenticalMaterialSlots;
+					}
+					return false;
 				}
-				return false;
-			})
+			)
 		),
 		NAME_None,
 		EUserInterfaceActionType::ToggleButton
 	);
 
-	auto AddKindToCollapseEntry = [&]( const EUsdDefaultKind Kind, const FText& Text, FCanExecuteAction CanExecuteAction )
+	auto AddKindToCollapseEntry = [&](const EUsdDefaultKind Kind, const FText& Text, FCanExecuteAction CanExecuteAction)
 	{
 		MenuBuilder.AddMenuEntry(
 			Text,
 			FText::GetEmpty(),
 			FSlateIcon(),
 			FUIAction(
-				FExecuteAction::CreateLambda( [this, Text, Kind]()
-				{
-					if ( AUsdStageActor* StageActor = GetStageActorOrCDO() )
+				FExecuteAction::CreateLambda(
+					[this, Text, Kind]()
 					{
-						FScopedTransaction Transaction( FText::Format(
-							LOCTEXT( "ToggleCollapsingTransaction", "Toggle asset and component collapsing for kind '{0}' on USD stage actor '{1}'" ),
-							Text,
-							FText::FromString( StageActor->GetActorLabel() )
-						) );
-
-						// When we change our kinds to collapse, we'll end up loading the stage again and refreshing our selection.
-						// We'll take care to try to re-select the same prims within SUsdStage::Refresh after the refresh is complete,
-						// but if selection sync is on we'll also later on try updating our prim selection to match our component selection.
-						// Not only is this "selection spam" is very visible on the USD Stage Editor, if our originally selected component
-						// was just collapsed away, we'd end up updating our prim selection to point to the parent prim instead
-						// (the collapsing root), which is not ideal.
-						TGuardValue<bool> MaintainSelectionGuard(bUpdatingViewportSelection, true);
-
-						StageActor->SetKindsToCollapse((int32)((EUsdDefaultKind)StageActor->KindsToCollapse ^ Kind));
-						if (StageActor->IsTemplate())
+						if (AUsdStageActor* StageActor = GetStageActorOrCDO())
 						{
-							StageActor->SaveConfig();
+							FScopedTransaction Transaction(FText::Format(
+								LOCTEXT(
+									"ToggleCollapsingTransaction",
+									"Toggle asset and component collapsing for kind '{0}' on USD stage actor '{1}'"
+								),
+								Text,
+								FText::FromString(StageActor->GetActorLabel())
+							));
+
+							// When we change our kinds to collapse, we'll end up loading the stage again and refreshing our selection.
+							// We'll take care to try to re-select the same prims within SUsdStage::Refresh after the refresh is complete,
+							// but if selection sync is on we'll also later on try updating our prim selection to match our component selection.
+							// Not only is this "selection spam" is very visible on the USD Stage Editor, if our originally selected component
+							// was just collapsed away, we'd end up updating our prim selection to point to the parent prim instead
+							// (the collapsing root), which is not ideal.
+							TGuardValue<bool> MaintainSelectionGuard(bUpdatingViewportSelection, true);
+
+							StageActor->SetKindsToCollapse((int32)((EUsdDefaultKind)StageActor->KindsToCollapse ^ Kind));
+							if (StageActor->IsTemplate())
+							{
+								StageActor->SaveConfig();
+							}
 						}
 					}
-				}),
+				),
 				CanExecuteAction,
-				FIsActionChecked::CreateLambda( [this, Kind]()
-				{
-					if ( AUsdStageActor* StageActor = GetStageActorOrCDO() )
+				FIsActionChecked::CreateLambda(
+					[this, Kind]()
 					{
-						return EnumHasAllFlags( ( EUsdDefaultKind ) StageActor->KindsToCollapse, Kind );
+						if (AUsdStageActor* StageActor = GetStageActorOrCDO())
+						{
+							return EnumHasAllFlags((EUsdDefaultKind)StageActor->KindsToCollapse, Kind);
+						}
+						return false;
 					}
-					return false;
-				})
+				)
 			),
 			NAME_None,
 			EUserInterfaceActionType::ToggleButton
 		);
 	};
 
-	MenuBuilder.BeginSection( "Kinds to collapse", LOCTEXT( "KindsToCollapse", "Kinds to collapse" ) );
+	MenuBuilder.BeginSection("Kinds to collapse", LOCTEXT("KindsToCollapse", "Kinds to collapse"));
 
-	AddKindToCollapseEntry(
-		EUsdDefaultKind::Model,
-		LOCTEXT( "ModelKind", "Model" ),
-		FCanExecuteAction{}
-	);
+	AddKindToCollapseEntry(EUsdDefaultKind::Model, LOCTEXT("ModelKind", "Model"), FCanExecuteAction{});
 
 	AddKindToCollapseEntry(
 		EUsdDefaultKind::Component,
-		LOCTEXT( "ModelComponent", "   Component" ),
-		FCanExecuteAction::CreateLambda( [this]()
-		{
-			if ( AUsdStageActor* StageActor = GetStageActorOrCDO() )
+		LOCTEXT("ModelComponent", "   Component"),
+		FCanExecuteAction::CreateLambda(
+			[this]()
 			{
-				// If we're collapsing all "model" kinds, the all "component"s should be collapsed anyway
-				if ( !EnumHasAnyFlags( ( EUsdDefaultKind ) StageActor->KindsToCollapse, EUsdDefaultKind::Model ) )
+				if (AUsdStageActor* StageActor = GetStageActorOrCDO())
 				{
-					return true;
+					// If we're collapsing all "model" kinds, the all "component"s should be collapsed anyway
+					if (!EnumHasAnyFlags((EUsdDefaultKind)StageActor->KindsToCollapse, EUsdDefaultKind::Model))
+					{
+						return true;
+					}
 				}
-			}
 
-			return false;
-		})
+				return false;
+			}
+		)
 	);
 
 	AddKindToCollapseEntry(
 		EUsdDefaultKind::Group,
-		LOCTEXT( "ModelGroup", "   Group" ),
-		FCanExecuteAction::CreateLambda( [this]()
-		{
-			if ( AUsdStageActor* StageActor = GetStageActorOrCDO() )
+		LOCTEXT("ModelGroup", "   Group"),
+		FCanExecuteAction::CreateLambda(
+			[this]()
 			{
-				if ( !EnumHasAnyFlags( ( EUsdDefaultKind ) StageActor->KindsToCollapse, EUsdDefaultKind::Model ) )
+				if (AUsdStageActor* StageActor = GetStageActorOrCDO())
 				{
-					return true;
+					if (!EnumHasAnyFlags((EUsdDefaultKind)StageActor->KindsToCollapse, EUsdDefaultKind::Model))
+					{
+						return true;
+					}
 				}
-			}
 
-			return false;
-		})
+				return false;
+			}
+		)
 	);
 
 	AddKindToCollapseEntry(
 		EUsdDefaultKind::Assembly,
-		LOCTEXT( "ModelAssembly", "      Assembly" ),
-		FCanExecuteAction::CreateLambda( [this]()
-		{
-			if ( AUsdStageActor* StageActor = GetStageActorOrCDO() )
+		LOCTEXT("ModelAssembly", "      Assembly"),
+		FCanExecuteAction::CreateLambda(
+			[this]()
 			{
-				if ( !EnumHasAnyFlags( ( EUsdDefaultKind ) StageActor->KindsToCollapse, EUsdDefaultKind::Model ) && !EnumHasAnyFlags( ( EUsdDefaultKind ) StageActor->KindsToCollapse, EUsdDefaultKind::Group ) )
+				if (AUsdStageActor* StageActor = GetStageActorOrCDO())
 				{
-					return true;
+					if (!EnumHasAnyFlags((EUsdDefaultKind)StageActor->KindsToCollapse, EUsdDefaultKind::Model)
+						&& !EnumHasAnyFlags((EUsdDefaultKind)StageActor->KindsToCollapse, EUsdDefaultKind::Group))
+					{
+						return true;
+					}
 				}
-			}
 
-			return false;
-		})
+				return false;
+			}
+		)
 	);
 
-	AddKindToCollapseEntry(
-		EUsdDefaultKind::Subcomponent,
-		LOCTEXT( "ModelSubcomponent", "Subcomponent" ),
-		FCanExecuteAction{}
-	);
+	AddKindToCollapseEntry(EUsdDefaultKind::Subcomponent, LOCTEXT("ModelSubcomponent", "Subcomponent"), FCanExecuteAction{});
 
 	MenuBuilder.EndSection();
 }
 
-void SUsdStage::FillAssetReuseSubMenu( FMenuBuilder& MenuBuilder )
+void SUsdStage::FillAssetReuseSubMenu(FMenuBuilder& MenuBuilder)
 {
 	MenuBuilder.AddMenuEntry(
 		LOCTEXT("ReuseIdenticalAssets", "Reuse identical assets"),
@@ -1901,34 +1973,38 @@ void SUsdStage::FillInterpolationTypeSubMenu(FMenuBuilder& MenuBuilder)
 		LOCTEXT("LinearType_ToolTip", "Attribute values are linearly interpolated between authored values"),
 		FSlateIcon(),
 		FUIAction(
-			FExecuteAction::CreateLambda([this]()
-			{
-				if ( AUsdStageActor* StageActor = GetStageActorOrCDO() )
+			FExecuteAction::CreateLambda(
+				[this]()
 				{
-					FScopedTransaction Transaction(FText::Format(
-						LOCTEXT("SetLinearInterpolationType", "Set USD stage actor '{0}' to linear interpolation"),
-						FText::FromString(StageActor->GetActorLabel())
-					));
+					if (AUsdStageActor* StageActor = GetStageActorOrCDO())
+					{
+						FScopedTransaction Transaction(FText::Format(
+							LOCTEXT("SetLinearInterpolationType", "Set USD stage actor '{0}' to linear interpolation"),
+							FText::FromString(StageActor->GetActorLabel())
+						));
 
-					// c.f. comment in SUsdStage::FillCollapsingSubMenu
-					TGuardValue<bool> MaintainSelectionGuard(bUpdatingViewportSelection, true);
+						// c.f. comment in SUsdStage::FillCollapsingSubMenu
+						TGuardValue<bool> MaintainSelectionGuard(bUpdatingViewportSelection, true);
 
-					StageActor->SetInterpolationType(EUsdInterpolationType::Linear);
-					if (StageActor->IsTemplate())
-					{
-						StageActor->SaveConfig();
+						StageActor->SetInterpolationType(EUsdInterpolationType::Linear);
+						if (StageActor->IsTemplate())
+						{
+							StageActor->SaveConfig();
+						}
 					}
 				}
-			}),
+			),
 			FCanExecuteAction{},
-			FIsActionChecked::CreateLambda([this]()
-			{
-				if ( AUsdStageActor* StageActor = GetStageActorOrCDO() )
+			FIsActionChecked::CreateLambda(
+				[this]()
 				{
-					return StageActor->InterpolationType == EUsdInterpolationType::Linear;
+					if (AUsdStageActor* StageActor = GetStageActorOrCDO())
+					{
+						return StageActor->InterpolationType == EUsdInterpolationType::Linear;
+					}
+					return false;
 				}
-				return false;
-			})
+			)
 		),
 		NAME_None,
 		EUserInterfaceActionType::RadioButton
@@ -1936,112 +2012,131 @@ void SUsdStage::FillInterpolationTypeSubMenu(FMenuBuilder& MenuBuilder)
 
 	MenuBuilder.AddMenuEntry(
 		LOCTEXT("HeldType", "Held"),
-		LOCTEXT("HeldType_ToolTip", "Attribute values are held constant between authored values. An attribute's value will be equal to the nearest preceding authored value. If there is no preceding authored value, the value will be equal to the nearest subsequent value."),
+		LOCTEXT(
+			"HeldType_ToolTip",
+			"Attribute values are held constant between authored values. An attribute's value will be equal to the nearest preceding authored value. "
+			"If there is no preceding authored value, the value will be equal to the nearest subsequent value."
+		),
 		FSlateIcon(),
 		FUIAction(
-			FExecuteAction::CreateLambda([this]()
-			{
-				if ( AUsdStageActor* StageActor = GetStageActorOrCDO() )
+			FExecuteAction::CreateLambda(
+				[this]()
 				{
-					FScopedTransaction Transaction(FText::Format(
-						LOCTEXT("SetHeldInterpolationType", "Set USD stage actor '{0}' to held interpolation"),
-						FText::FromString(StageActor->GetActorLabel())
-					));
+					if (AUsdStageActor* StageActor = GetStageActorOrCDO())
+					{
+						FScopedTransaction Transaction(FText::Format(
+							LOCTEXT("SetHeldInterpolationType", "Set USD stage actor '{0}' to held interpolation"),
+							FText::FromString(StageActor->GetActorLabel())
+						));
 
-					// c.f. comment in SUsdStage::FillCollapsingSubMenu
-					TGuardValue<bool> MaintainSelectionGuard(bUpdatingViewportSelection, true);
+						// c.f. comment in SUsdStage::FillCollapsingSubMenu
+						TGuardValue<bool> MaintainSelectionGuard(bUpdatingViewportSelection, true);
 
-					StageActor->SetInterpolationType(EUsdInterpolationType::Held);
-					if (StageActor->IsTemplate())
-					{
-						StageActor->SaveConfig();
+						StageActor->SetInterpolationType(EUsdInterpolationType::Held);
+						if (StageActor->IsTemplate())
+						{
+							StageActor->SaveConfig();
+						}
 					}
 				}
-			}),
+			),
 			FCanExecuteAction{},
-			FIsActionChecked::CreateLambda([this]()
-			{
-				if ( AUsdStageActor* StageActor = GetStageActorOrCDO() )
+			FIsActionChecked::CreateLambda(
+				[this]()
 				{
-					return StageActor->InterpolationType == EUsdInterpolationType::Held;
+					if (AUsdStageActor* StageActor = GetStageActorOrCDO())
+					{
+						return StageActor->InterpolationType == EUsdInterpolationType::Held;
+					}
+					return false;
 				}
-				return false;
-			})
+			)
 		),
 		NAME_None,
 		EUserInterfaceActionType::RadioButton
 	);
 }
 
-void SUsdStage::FillSelectionSubMenu( FMenuBuilder& MenuBuilder )
+void SUsdStage::FillSelectionSubMenu(FMenuBuilder& MenuBuilder)
 {
 	MenuBuilder.AddMenuEntry(
-		LOCTEXT( "SynchronizeText", "Synchronize with Editor" ),
+		LOCTEXT("SynchronizeText", "Synchronize with Editor"),
 		FText::GetEmpty(),
 		FSlateIcon(),
 		FUIAction(
-			FExecuteAction::CreateLambda([this]()
-			{
-				if ( UUsdStageEditorSettings* Settings = GetMutableDefault<UUsdStageEditorSettings>() )
-				{
-					Settings->bSelectionSynced = !Settings->bSelectionSynced;
-					Settings->SaveConfig();
-				}
-
-				// Immediately sync the selection, but only if the results are obvious (i.e. either our prim selection or viewport selection are empty)
-				if ( GEditor )
+			FExecuteAction::CreateLambda(
+				[this]()
 				{
-					int32 NumPrimsSelected = UsdStageTreeView->GetNumItemsSelected();
-					int32 NumViewportSelected = FMath::Max( GEditor->GetSelectedComponentCount(), GEditor->GetSelectedActorCount() );
-
-					if ( NumPrimsSelected == 0 )
+					if (UUsdStageEditorSettings* Settings = GetMutableDefault<UUsdStageEditorSettings>())
 					{
-						USelection* Selection = GEditor->GetSelectedComponentCount() > 0
-							? GEditor->GetSelectedComponents()
-							: GEditor->GetSelectedActors();
-
-						OnViewportSelectionChanged(Selection);
+						Settings->bSelectionSynced = !Settings->bSelectionSynced;
+						Settings->SaveConfig();
 					}
-					else if ( NumViewportSelected == 0 )
+
+					// Immediately sync the selection, but only if the results are obvious (i.e. either our prim selection or viewport selection are
+					// empty)
+					if (GEditor)
 					{
-						TGuardValue<bool> SelectionLoopGuard( bUpdatingViewportSelection, true );
+						int32 NumPrimsSelected = UsdStageTreeView->GetNumItemsSelected();
+						int32 NumViewportSelected = FMath::Max(GEditor->GetSelectedComponentCount(), GEditor->GetSelectedActorCount());
+
+						if (NumPrimsSelected == 0)
+						{
+							USelection* Selection = GEditor->GetSelectedComponentCount() > 0 ? GEditor->GetSelectedComponents()
+																							 : GEditor->GetSelectedActors();
 
-						SUSDStageImpl::SelectGeneratedComponentsAndActors( ViewModel.UsdStageActor.Get(), UsdStageTreeView->GetSelectedPrimPaths() );
+							OnViewportSelectionChanged(Selection);
+						}
+						else if (NumViewportSelected == 0)
+						{
+							TGuardValue<bool> SelectionLoopGuard(bUpdatingViewportSelection, true);
+
+							SUSDStageImpl::SelectGeneratedComponentsAndActors(
+								ViewModel.UsdStageActor.Get(),
+								UsdStageTreeView->GetSelectedPrimPaths()
+							);
+						}
 					}
 				}
-			}),
-			FCanExecuteAction::CreateLambda([]()
-			{
-				return true;
-			}),
-			FIsActionChecked::CreateLambda([]()
-			{
-				const UUsdStageEditorSettings* Settings = GetDefault<UUsdStageEditorSettings>();
-				return Settings && Settings->bSelectionSynced;
-			})
+			),
+			FCanExecuteAction::CreateLambda(
+				[]()
+				{
+					return true;
+				}
+			),
+			FIsActionChecked::CreateLambda(
+				[]()
+				{
+					const UUsdStageEditorSettings* Settings = GetDefault<UUsdStageEditorSettings>();
+					return Settings && Settings->bSelectionSynced;
+				}
+			)
 		),
 		NAME_None,
 		EUserInterfaceActionType::ToggleButton
 	);
 }
 
-void SUsdStage::FillNaniteThresholdSubMenu( FMenuBuilder& MenuBuilder )
+void SUsdStage::FillNaniteThresholdSubMenu(FMenuBuilder& MenuBuilder)
 {
-	if ( AUsdStageActor* StageActor = GetStageActorOrCDO() )
+	if (AUsdStageActor* StageActor = GetStageActorOrCDO())
 	{
 		CurrentNaniteThreshold = StageActor->NaniteTriangleThreshold;
 	}
 
-	TSharedRef<SSpinBox<int32>> Slider = SNew( SSpinBox<int32> )
-		.MinValue( 0 )
-		.ToolTipText( LOCTEXT( "TriangleThresholdTooltip", "Try enabling Nanite for static meshes that are generated with at least this many triangles" ) )
-		.Value( this, &SUsdStage::GetNaniteTriangleThresholdValue )
-		.OnValueChanged(this, &SUsdStage::OnNaniteTriangleThresholdValueChanged )
-		.SupportDynamicSliderMaxValue( true )
-		.OnValueCommitted( this, &SUsdStage::OnNaniteTriangleThresholdValueCommitted );
+	// clang-format off
+	TSharedRef<SSpinBox<int32>> Slider = SNew(SSpinBox<int32>)
+		.MinValue(0)
+		.ToolTipText(LOCTEXT("TriangleThresholdTooltip", "Try enabling Nanite for static meshes that are generated with at least this many triangles"))
+		.Value(this, &SUsdStage::GetNaniteTriangleThresholdValue)
+		.OnValueChanged(this, &SUsdStage::OnNaniteTriangleThresholdValueChanged)
+		.SupportDynamicSliderMaxValue(true)
+		.OnValueCommitted(this, &SUsdStage::OnNaniteTriangleThresholdValueCommitted);
+	// clang-format on
 
 	const bool bNoIndent = true;
-	MenuBuilder.AddWidget( Slider, FText::FromString( TEXT( "Triangle threshold: " ) ), bNoIndent );
+	MenuBuilder.AddWidget(Slider, FText::FromString(TEXT("Triangle threshold: ")), bNoIndent);
 }
 
 void SUsdStage::FileNew()
@@ -2049,13 +2144,13 @@ void SUsdStage::FileNew()
 	ViewModel.NewStage();
 }
 
-void SUsdStage::FileOpen( const FString& FilePath )
+void SUsdStage::FileOpen(const FString& FilePath)
 {
 	FString FilePathCopy = FilePath;
 	if (FilePathCopy.IsEmpty())
 	{
-		TOptional< FString > UsdFilePath = UsdUtils::BrowseUsdFile( UsdUtils::EBrowseFileMode::Open );
-		if( !UsdFilePath.IsSet() )
+		TOptional<FString> UsdFilePath = UsdUtils::BrowseUsdFile(UsdUtils::EBrowseFileMode::Open);
+		if (!UsdFilePath.IsSet())
 		{
 			return;
 		}
@@ -2066,24 +2161,24 @@ void SUsdStage::FileOpen( const FString& FilePath )
 	OpenStage(*FilePathCopy);
 }
 
-void SUsdStage::FileSave( const FString& FilePath )
+void SUsdStage::FileSave(const FString& FilePath)
 {
 	const AUsdStageActor* StageActor = ViewModel.UsdStageActor.Get();
-	if ( !StageActor )
+	if (!StageActor)
 	{
 		return;
 	}
 
-	UE::FUsdStage UsdStage = static_cast< const AUsdStageActor* >( StageActor )->GetUsdStage();
-	if ( !UsdStage )
+	UE::FUsdStage UsdStage = static_cast<const AUsdStageActor*>(StageActor)->GetUsdStage();
+	if (!UsdStage)
 	{
 		return;
 	}
 
-	if ( UE::FSdfLayer RootLayer = UsdStage.GetRootLayer() )
+	if (UE::FSdfLayer RootLayer = UsdStage.GetRootLayer())
 	{
 		FString RealPath = RootLayer.GetRealPath();
-		if ( FPaths::FileExists( RealPath ) )
+		if (FPaths::FileExists(RealPath))
 		{
 			ViewModel.SaveStage();
 		}
@@ -2092,8 +2187,8 @@ void SUsdStage::FileSave( const FString& FilePath )
 			FString FilePathCopy = FilePath;
 			if (FilePathCopy.IsEmpty())
 			{
-				TOptional< FString > UsdFilePath = UsdUtils::BrowseUsdFile( UsdUtils::EBrowseFileMode::Save );
-				if( !UsdFilePath.IsSet() )
+				TOptional<FString> UsdFilePath = UsdUtils::BrowseUsdFile(UsdUtils::EBrowseFileMode::Save);
+				if (!UsdFilePath.IsSet())
 				{
 					return;
 				}
@@ -2106,47 +2201,43 @@ void SUsdStage::FileSave( const FString& FilePath )
 	}
 }
 
-void SUsdStage::FileExportAllLayers( const FString& OutputDirectory )
+void SUsdStage::FileExportAllLayers(const FString& OutputDirectory)
 {
 	const AUsdStageActor* StageActor = ViewModel.UsdStageActor.Get();
-	if ( !StageActor )
+	if (!StageActor)
 	{
 		return;
 	}
 
 	UE::FUsdStage UsdStage = StageActor->GetUsdStage();
-	if ( !UsdStage )
+	if (!UsdStage)
 	{
 		return;
 	}
 
 	const bool bIncludeClipLayers = false;
-	TArray<UE::FSdfLayer> LayerStack = UsdStage.GetUsedLayers( bIncludeClipLayers );
-	if ( LayerStack.Num() < 1 )
+	TArray<UE::FSdfLayer> LayerStack = UsdStage.GetUsedLayers(bIncludeClipLayers);
+	if (LayerStack.Num() < 1)
 	{
 		return;
 	}
 
 	IDesktopPlatform* DesktopPlatform = FDesktopPlatformModule::Get();
-	if ( !DesktopPlatform )
+	if (!DesktopPlatform)
 	{
 		return;
 	}
 
-	TSharedPtr< SWindow > ParentWindow = FSlateApplication::Get().FindWidgetWindow( AsShared() );
-	void* ParentWindowHandle = ( ParentWindow.IsValid() && ParentWindow->GetNativeWindow().IsValid() )
-		? ParentWindow->GetNativeWindow()->GetOSWindowHandle()
-		: nullptr;
+	TSharedPtr<SWindow> ParentWindow = FSlateApplication::Get().FindWidgetWindow(AsShared());
+	void* ParentWindowHandle = (ParentWindow.IsValid() && ParentWindow->GetNativeWindow().IsValid())
+								   ? ParentWindow->GetNativeWindow()->GetOSWindowHandle()
+								   : nullptr;
 
 	FString OutputDirectoryCopy = OutputDirectory;
 	if (OutputDirectoryCopy.IsEmpty())
 	{
-		if ( !DesktopPlatform->OpenDirectoryDialog(
-			ParentWindowHandle,
-			LOCTEXT( "ChooseFolder", "Choose output folder" ).ToString(),
-			TEXT( "" ),
-			OutputDirectoryCopy
-		) )
+		if (!DesktopPlatform
+				 ->OpenDirectoryDialog(ParentWindowHandle, LOCTEXT("ChooseFolder", "Choose output folder").ToString(), TEXT(""), OutputDirectoryCopy))
 		{
 			return;
 		}
@@ -2159,102 +2250,102 @@ void SUsdStage::FileExportAllLayers( const FString& OutputDirectory )
 	// Manually load any layer that is referenced by any of the stage's layers, but not currently loaded.
 	// This can happen if e.g. an unselected variant appends a reference or has a payload.
 	// We will need to actually load these as opposed to just copy-paste the files as we may need to update references/paths
-	SUSDStageImpl::AppendAllExternalReferences( LayerStack );
+	SUSDStageImpl::AppendAllExternalReferences(LayerStack);
 
 	// Discard session layers
-	for ( int32 Index = LayerStack.Num() - 1; Index >= 0; --Index )
+	for (int32 Index = LayerStack.Num() - 1; Index >= 0; --Index)
 	{
-		UE::FSdfLayer& Layer = LayerStack[ Index ];
-		if ( Layer.IsAnonymous() )
+		UE::FSdfLayer& Layer = LayerStack[Index];
+		if (Layer.IsAnonymous())
 		{
 			const int32 Count = 1;
 			const bool bAllowShrinking = false;
-			LayerStack.RemoveAt( Index, Count, bAllowShrinking );
+			LayerStack.RemoveAt(Index, Count, bAllowShrinking);
 		}
 	}
 
 #if PLATFORM_LINUX
-	TMap<FString, FString, FDefaultSetAllocator, UsdUtils::FCaseSensitiveStringMapFuncs< FString > > OldPathToSavedPath;
+	TMap<FString, FString, FDefaultSetAllocator, UsdUtils::FCaseSensitiveStringMapFuncs<FString>> OldPathToSavedPath;
 #else
 	TMap<FString, FString> OldPathToSavedPath;
 #endif
 
 	// If the stage has layers referencing each other, we will need to remap those to point exclusively at the newly saved files
 	TSet<FString> SavedPaths;
-	for ( const UE::FSdfLayer& Layer : LayerStack )
+	for (const UE::FSdfLayer& Layer : LayerStack)
 	{
 		FString LayerPath = Layer.GetRealPath();
-		FPaths::NormalizeFilename( LayerPath );
+		FPaths::NormalizeFilename(LayerPath);
 
-		FString TargetPath = FPaths::Combine( OutputDirectory, Layer.GetDisplayName() );
-		FPaths::NormalizeFilename( TargetPath );
+		FString TargetPath = FPaths::Combine(OutputDirectory, Layer.GetDisplayName());
+		FPaths::NormalizeFilename(TargetPath);
 
 		// Filename collision (should be rare, but possible given that we're discarding the folder structure)
-		if ( SavedPaths.Contains( TargetPath ) )
+		if (SavedPaths.Contains(TargetPath))
 		{
-			FString TargetPathNoExt = FPaths::SetExtension( TargetPath, TEXT( "" ) );
-			FString Ext = FPaths::GetExtension( TargetPath );
+			FString TargetPathNoExt = FPaths::SetExtension(TargetPath, TEXT(""));
+			FString Ext = FPaths::GetExtension(TargetPath);
 			int32 Suffix = 0;
 
 			do
 			{
-				TargetPath = FString::Printf( TEXT( "%s_%d.%s" ), *TargetPathNoExt, Suffix++, *Ext );
-			} while ( SavedPaths.Contains( TargetPath ) );
+				TargetPath = FString::Printf(TEXT("%s_%d.%s"), *TargetPathNoExt, Suffix++, *Ext);
+			} while (SavedPaths.Contains(TargetPath));
 		}
 
-		OldPathToSavedPath.Add( LayerPath, TargetPath );
-		SavedPaths.Add( TargetPath );
+		OldPathToSavedPath.Add(LayerPath, TargetPath);
+		SavedPaths.Add(TargetPath);
 	}
 
 	// Actually save the output layers
-	for ( UE::FSdfLayer& Layer : LayerStack )
+	for (UE::FSdfLayer& Layer : LayerStack)
 	{
 		FString LayerPath = Layer.GetRealPath();
-		FPaths::NormalizeFilename( LayerPath );
+		FPaths::NormalizeFilename(LayerPath);
 
-		if ( FString* TargetLayerPath = OldPathToSavedPath.Find( LayerPath ) )
+		if (FString* TargetLayerPath = OldPathToSavedPath.Find(LayerPath))
 		{
 			// Clone the layer so that we don't modify the currently opened stage
-			UE::FSdfLayer OutputLayer = UE::FSdfLayer::CreateNew( **TargetLayerPath );
-			OutputLayer.TransferContent( Layer );
+			UE::FSdfLayer OutputLayer = UE::FSdfLayer::CreateNew(**TargetLayerPath);
+			OutputLayer.TransferContent(Layer);
 
 			// Update references to assets (e.g. textures) so that they're absolute and also work from the new file
-			UsdUtils::ConvertAssetRelativePathsToAbsolute( OutputLayer, Layer );
+			UsdUtils::ConvertAssetRelativePathsToAbsolute(OutputLayer, Layer);
 
 			// Remap references to layers so that they point at the other newly saved files. Note that SUSDStageImpl::AppendAllExternalReferences
 			// will have collected all external references already, so OldPathToSavedPath should have entries for all references we'll find
-			for ( const FString& AssetDependency : OutputLayer.GetCompositionAssetDependencies() )
+			for (const FString& AssetDependency : OutputLayer.GetCompositionAssetDependencies())
 			{
-				FString AbsRef = FPaths::ConvertRelativePathToFull( FPaths::GetPath( LayerPath ), AssetDependency ); // Relative to the original file
-				FPaths::NormalizeFilename( AbsRef );
+				FString AbsRef = FPaths::ConvertRelativePathToFull(FPaths::GetPath(LayerPath), AssetDependency);	// Relative to the original file
+				FPaths::NormalizeFilename(AbsRef);
 
-				if ( FString* SavedReference = OldPathToSavedPath.Find( AbsRef ) )
+				if (FString* SavedReference = OldPathToSavedPath.Find(AbsRef))
 				{
-					OutputLayer.UpdateCompositionAssetDependency( *AssetDependency, **SavedReference );
+					OutputLayer.UpdateCompositionAssetDependency(*AssetDependency, **SavedReference);
 				}
 			}
 
 			bool bForce = true;
-			OutputLayer.Save( bForce );
+			OutputLayer.Save(bForce);
 		}
 	}
 
 	// Send analytics
-	if ( FEngineAnalytics::IsAvailable() )
+	if (FEngineAnalytics::IsAvailable())
 	{
 		TArray<FAnalyticsEventAttribute> EventAttributes;
 
-		EventAttributes.Emplace( TEXT( "NumLayersExported" ), LayerStack.Num() );
+		EventAttributes.Emplace(TEXT("NumLayersExported"), LayerStack.Num());
 
 		bool bAutomated = false;
-		double ElapsedSeconds = FPlatformTime::ToSeconds64( FPlatformTime::Cycles64() - StartTime );
-		FString Extension = FPaths::GetExtension( UsdStage.GetRootLayer().GetDisplayName() );
+		double ElapsedSeconds = FPlatformTime::ToSeconds64(FPlatformTime::Cycles64() - StartTime);
+		FString Extension = FPaths::GetExtension(UsdStage.GetRootLayer().GetDisplayName());
 		IUsdClassesModule::SendAnalytics(
-			MoveTemp( EventAttributes ),
-			TEXT( "ExportStageLayers" ),
+			MoveTemp(EventAttributes),
+			TEXT("ExportStageLayers"),
 			bAutomated,
 			ElapsedSeconds,
-			UsdUtils::GetUsdStageNumFrames( UsdStage ),
+			UsdUtils::GetUsdStageNumFrames(UsdStage),
 			Extension
 		);
 	}
@@ -2263,13 +2354,13 @@ void SUsdStage::FileExportAllLayers( const FString& OutputDirectory )
 void SUsdStage::FileExportFlattenedStage(const FString& OutputLayer)
 {
 	const AUsdStageActor* StageActor = ViewModel.UsdStageActor.Get();
-	if ( !StageActor )
+	if (!StageActor)
 	{
 		return;
 	}
 
 	UE::FUsdStage UsdStage = StageActor->GetUsdStage();
-	if ( !UsdStage )
+	if (!UsdStage)
 	{
 		return;
 	}
@@ -2277,8 +2368,8 @@ void SUsdStage::FileExportFlattenedStage(const FString& OutputLayer)
 	FString OutputLayerCopy = OutputLayer;
 	if (OutputLayerCopy.IsEmpty())
 	{
-		TOptional< FString > UsdFilePath = UsdUtils::BrowseUsdFile( UsdUtils::EBrowseFileMode::Save );
-		if ( !UsdFilePath.IsSet() )
+		TOptional<FString> UsdFilePath = UsdUtils::BrowseUsdFile(UsdUtils::EBrowseFileMode::Save);
+		if (!UsdFilePath.IsSet())
 		{
 			return;
 		}
@@ -2291,17 +2382,17 @@ void SUsdStage::FileExportFlattenedStage(const FString& OutputLayer)
 	UsdStage.Export(*OutputLayerCopy);
 
 	// Send analytics
-	if ( FEngineAnalytics::IsAvailable() )
+	if (FEngineAnalytics::IsAvailable())
 	{
 		bool bAutomated = false;
-		double ElapsedSeconds = FPlatformTime::ToSeconds64( FPlatformTime::Cycles64() - StartTime );
+		double ElapsedSeconds = FPlatformTime::ToSeconds64(FPlatformTime::Cycles64() - StartTime);
 		FString Extension = FPaths::GetExtension(OutputLayerCopy);
 		IUsdClassesModule::SendAnalytics(
 			{},
-			TEXT( "ExportStageFlattened" ),
+			TEXT("ExportStageFlattened"),
 			bAutomated,
 			ElapsedSeconds,
-			UsdUtils::GetUsdStageNumFrames( UsdStage ),
+			UsdUtils::GetUsdStageNumFrames(UsdStage),
 			Extension
 		);
 	}
@@ -2309,16 +2400,16 @@ void SUsdStage::FileExportFlattenedStage(const FString& OutputLayer)
 
 void SUsdStage::FileReload()
 {
-	FScopedTransaction Transaction( LOCTEXT( "ReloadTransaction", "Reload USD stage" ) );
+	FScopedTransaction Transaction(LOCTEXT("ReloadTransaction", "Reload USD stage"));
 
 	ViewModel.ReloadStage();
 
 	const AUsdStageActor* StageActor = ViewModel.UsdStageActor.Get();
 
-	if ( UsdLayersTreeView && StageActor )
+	if (UsdLayersTreeView && StageActor)
 	{
 		const bool bResync = true;
-		UsdLayersTreeView->Refresh( StageActor->GetBaseUsdStage(), StageActor->GetIsolatedUsdStage(), bResync );
+		UsdLayersTreeView->Refresh(StageActor->GetBaseUsdStage(), StageActor->GetIsolatedUsdStage(), bResync);
 	}
 }
 
@@ -2328,10 +2419,10 @@ void SUsdStage::FileReset()
 
 	const AUsdStageActor* StageActor = ViewModel.UsdStageActor.Get();
 
-	if ( UsdLayersTreeView && StageActor )
+	if (UsdLayersTreeView && StageActor)
 	{
 		const bool bResync = true;
-		UsdLayersTreeView->Refresh( StageActor->GetBaseUsdStage(), StageActor->GetIsolatedUsdStage(), bResync );
+		UsdLayersTreeView->Refresh(StageActor->GetBaseUsdStage(), StageActor->GetIsolatedUsdStage(), bResync);
 	}
 }
 
@@ -2356,11 +2447,11 @@ void SUsdStage::FileClose()
 	}
 }
 
-void SUsdStage::OnLayerIsolated( const UE::FSdfLayer& IsolatedLayer )
+void SUsdStage::OnLayerIsolated(const UE::FSdfLayer& IsolatedLayer)
 {
-	if ( AUsdStageActor* StageActor = ViewModel.UsdStageActor.Get() )
+	if (AUsdStageActor* StageActor = ViewModel.UsdStageActor.Get())
 	{
-		StageActor->IsolateLayer( IsolatedLayer );
+		StageActor->IsolateLayer(IsolatedLayer);
 
 		RequestFullRefresh();
 	}
@@ -2371,34 +2462,34 @@ void SUsdStage::ActionsImportWithDialog()
 	ViewModel.ImportStage();
 }
 
-void SUsdStage::ActionsImport( const FString& OutputContentFolder, UUsdStageImportOptions* Options )
+void SUsdStage::ActionsImport(const FString& OutputContentFolder, UUsdStageImportOptions* Options)
 {
-	ViewModel.ImportStage( *OutputContentFolder, Options );
+	ViewModel.ImportStage(*OutputContentFolder, Options);
 }
 
 void SUsdStage::ExportSelectedLayers(const FString& OutputLayerOrDirectory)
 {
-	if ( UsdLayersTreeView )
+	if (UsdLayersTreeView)
 	{
 		UsdLayersTreeView->ExportSelectedLayers(OutputLayerOrDirectory);
 	}
 }
 
-void SUsdStage::OnPrimSelectionChanged( const TArray<FString>& PrimPaths )
+void SUsdStage::OnPrimSelectionChanged(const TArray<FString>& PrimPaths)
 {
-	if ( bUpdatingPrimSelection )
+	if (bUpdatingPrimSelection)
 	{
 		return;
 	}
 
 	AUsdStageActor* StageActor = ViewModel.UsdStageActor.Get();
-	if ( !StageActor )
+	if (!StageActor)
 	{
 		return;
 	}
 
-	UE::FUsdStage UsdStage = static_cast< const AUsdStageActor* >( StageActor )->GetUsdStage();
-	SelectedPrimPath = PrimPaths.Num() == 1 ? PrimPaths[ 0 ] : TEXT( "" );
+	UE::FUsdStage UsdStage = static_cast<const AUsdStageActor*>(StageActor)->GetUsdStage();
+	SelectedPrimPath = PrimPaths.Num() == 1 ? PrimPaths[0] : TEXT("");
 
 	if (UsdPrimInfoWidget)
 	{
@@ -2412,39 +2503,36 @@ void SUsdStage::OnPrimSelectionChanged( const TArray<FString>& PrimPaths )
 	// because we need the InfoCache to be complete at this point. Strangely though, the SlowTask.EnterProgressFrame
 	// within AUsdStageActor::LoadUsdStage triggers a full slate tick right there, that calls us before
 	// LoadUsdStage is complete... in the future I think we can fix this by changing our delegates from
-	// AsyncTask( ENamedThreads::GameThread ) to finishing up on the next tick instead, but for now
+	// AsyncTask(ENamedThreads::GameThread) to finishing up on the next tick instead, but for now
 	// this will do
 	const bool bIsRootSelection = PrimPaths.Num() == 1 && PrimPaths[0] == TEXT("/");
 
 	const UUsdStageEditorSettings* Settings = GetDefault<UUsdStageEditorSettings>();
-	if ( Settings && Settings->bSelectionSynced && GEditor && !bIsRootSelection )
+	if (Settings && Settings->bSelectionSynced && GEditor && !bIsRootSelection)
 	{
-		TGuardValue<bool> SelectionLoopGuard( bUpdatingViewportSelection, true );
+		TGuardValue<bool> SelectionLoopGuard(bUpdatingViewportSelection, true);
 
-		SUSDStageImpl::SelectGeneratedComponentsAndActors( StageActor, PrimPaths );
+		SUSDStageImpl::SelectGeneratedComponentsAndActors(StageActor, PrimPaths);
 	}
 }
 
-void SUsdStage::OpenStage( const TCHAR* FilePath )
+void SUsdStage::OpenStage(const TCHAR* FilePath)
 {
 	TOptional<FScopedTransaction> Transaction;
 	if (!GDiscardUndoBufferOnStageOpenClose)
 	{
 		// Create the transaction before calling UsdStageModule.GetUsdStageActor as that may create the actor, and we want
 		// the actor spawning to be part of the transaction
-		Transaction.Emplace(FText::Format(
-			LOCTEXT("OpenStageTransaction", "Open USD stage '{0}'"),
-			FText::FromString(FilePath)
-		));
+		Transaction.Emplace(FText::Format(LOCTEXT("OpenStageTransaction", "Open USD stage '{0}'"), FText::FromString(FilePath)));
 	}
 
-	if ( !ViewModel.UsdStageActor.IsValid() )
+	if (!ViewModel.UsdStageActor.IsValid())
 	{
-		IUsdStageModule& UsdStageModule = FModuleManager::Get().LoadModuleChecked< IUsdStageModule >( "UsdStage" );
-		AttachToStageActor( &UsdStageModule.GetUsdStageActor( GWorld ) );
+		IUsdStageModule& UsdStageModule = FModuleManager::Get().LoadModuleChecked<IUsdStageModule>("UsdStage");
+		AttachToStageActor(&UsdStageModule.GetUsdStageActor(GWorld));
 	}
 
-	ViewModel.OpenStage( FilePath );
+	ViewModel.OpenStage(FilePath);
 
 	if (GDiscardUndoBufferOnStageOpenClose && GEditor)
 	{
@@ -2580,10 +2668,7 @@ void SUsdStage::OnSlateTick(float Time)
 					const bool bViewingStageProperties = SelectedPrimPath.IsEmpty() || SelectedPrimPath == TEXT("/");
 					const bool bStageUpdated = PrimPath == TEXT("/");
 
-					if (UsdPrimInfoWidget &&
-						StageActor &&
-						(bViewingTheUpdatedPrim || (bViewingStageProperties && bStageUpdated))
-					)
+					if (UsdPrimInfoWidget && StageActor && (bViewingTheUpdatedPrim || (bViewingStageProperties && bStageUpdated)))
 					{
 						PrimPathToSelect = PrimPath;
 					}
@@ -2600,9 +2685,7 @@ void SUsdStage::OnSlateTick(float Time)
 		{
 			// Refresh will generate brand new RootItems, so let's immediately select the new item
 			// that corresponds to the prim we're supposed to be selecting
-			UsdStageTreeView->SetSelectedPrimPaths(
-				{PrimPathToSelect.IsEmpty() ? SelectedPrimPath : PrimPathToSelect}
-			);
+			UsdStageTreeView->SetSelectedPrimPaths({PrimPathToSelect.IsEmpty() ? SelectedPrimPath : PrimPathToSelect});
 
 			// If we changed/redid the selection, let's make sure that our info panel will update too (right below this)
 			PrimPathToSelect = SelectedPrimPath;
@@ -2622,139 +2705,139 @@ void SUsdStage::OnSlateTick(float Time)
 	}
 }
 
-void SUsdStage::OnViewportSelectionChanged( UObject* NewSelection )
+void SUsdStage::OnViewportSelectionChanged(UObject* NewSelection)
 {
 	// This may be called when first opening a project, before the our widgets are fully initialized
-	if ( !UsdStageTreeView )
+	if (!UsdStageTreeView)
 	{
 		return;
 	}
 
 	const UUsdStageEditorSettings* Settings = GetDefault<UUsdStageEditorSettings>();
-	if ( !Settings || !Settings->bSelectionSynced || bUpdatingViewportSelection )
+	if (!Settings || !Settings->bSelectionSynced || bUpdatingViewportSelection)
 	{
 		return;
 	}
 
 	AUsdStageActor* AttachedStageActor = ViewModel.UsdStageActor.Get();
-	USelection* Selection = Cast<USelection>( NewSelection );
-	if ( !Selection || !AttachedStageActor )
+	USelection* Selection = Cast<USelection>(NewSelection);
+	if (!Selection || !AttachedStageActor)
 	{
 		return;
 	}
 
-	TGuardValue<bool> SelectionLoopGuard( bUpdatingPrimSelection, true );
+	TGuardValue<bool> SelectionLoopGuard(bUpdatingPrimSelection, true);
 
 	TArray<USceneComponent*> AllSelectedComponents;
 	{
 		TArray<AActor*> SelectedActors;
-		Selection->GetSelectedObjects<AActor>( SelectedActors );
+		Selection->GetSelectedObjects<AActor>(SelectedActors);
 
-		for ( AActor* SelectedActor : SelectedActors )
+		for (AActor* SelectedActor : SelectedActors)
 		{
 			// Some actors like AWorldDataLayer don't have a root component
-			if ( SelectedActor && SelectedActor->GetRootComponent() )
+			if (SelectedActor && SelectedActor->GetRootComponent())
 			{
-				AllSelectedComponents.Add( SelectedActor->GetRootComponent() );
+				AllSelectedComponents.Add(SelectedActor->GetRootComponent());
 			}
 		}
 
 		// Push the components last as we'll show info for the last selected component,
 		// and we must first select the actor before we can select its individual components
 		TArray<USceneComponent*> SelectedComponents;
-		Selection->GetSelectedObjects<USceneComponent>( SelectedComponents );
+		Selection->GetSelectedObjects<USceneComponent>(SelectedComponents);
 
-		AllSelectedComponents.Append( SelectedComponents );
+		AllSelectedComponents.Append(SelectedComponents);
 	}
-	if ( AllSelectedComponents.Num() == 0 )
+	if (AllSelectedComponents.Num() == 0)
 	{
 		return;
 	}
 
 	// Find the exact stage actor that owns these components, if any
 	AUsdStageActor* OwnerActor = nullptr;
-	if ( AllSelectedComponents.Num() > 0 )
+	if (AllSelectedComponents.Num() > 0)
 	{
-		USceneComponent* AttachParent = AllSelectedComponents[ AllSelectedComponents.Num() - 1 ]->GetAttachParent();
-		while ( AttachParent )
+		USceneComponent* AttachParent = AllSelectedComponents[AllSelectedComponents.Num() - 1]->GetAttachParent();
+		while (AttachParent)
 		{
 			USceneComponent* NewAttachParent = AttachParent->GetAttachParent();
-			if ( !NewAttachParent )
+			if (!NewAttachParent)
 			{
 				break;
 			}
 			AttachParent = NewAttachParent;
 		}
 
-		if ( AttachParent )
+		if (AttachParent)
 		{
-			OwnerActor = Cast<AUsdStageActor>( AttachParent->GetOwner() );
+			OwnerActor = Cast<AUsdStageActor>(AttachParent->GetOwner());
 		}
 	}
-	if ( !OwnerActor )
+	if (!OwnerActor)
 	{
 		// Selected component is not something we spawned
 		return;
 	}
 
 	TArray<FString> PrimPaths;
-	for ( USceneComponent* Component : AllSelectedComponents )
+	for (USceneComponent* Component : AllSelectedComponents)
 	{
-		FString FoundPrimPath = OwnerActor->GetSourcePrimPath( Component );
-		if ( !FoundPrimPath.IsEmpty() )
+		FString FoundPrimPath = OwnerActor->GetSourcePrimPath(Component);
+		if (!FoundPrimPath.IsEmpty())
 		{
-			PrimPaths.Add( FoundPrimPath );
+			PrimPaths.Add(FoundPrimPath);
 		}
 	}
 
-	if ( OwnerActor && OwnerActor != AttachedStageActor )
+	if (OwnerActor && OwnerActor != AttachedStageActor)
 	{
-		AttachToStageActor( OwnerActor );
+		AttachToStageActor(OwnerActor);
 	}
 
-	if ( PrimPaths.Num() > 0 )
+	if (PrimPaths.Num() > 0)
 	{
-		UsdStageTreeView->SetSelectedPrimPaths( PrimPaths );
+		UsdStageTreeView->SetSelectedPrimPaths(PrimPaths);
 
-		if ( UsdPrimInfoWidget )
+		if (UsdPrimInfoWidget)
 		{
-			UsdPrimInfoWidget->SetPrimPath( GetCurrentStage(), *PrimPaths[ PrimPaths.Num() - 1 ] );
+			UsdPrimInfoWidget->SetPrimPath(GetCurrentStage(), *PrimPaths[PrimPaths.Num() - 1]);
 		}
 	}
 }
 
-void SUsdStage::OnPostPIEStarted( bool bIsSimulating )
+void SUsdStage::OnPostPIEStarted(bool bIsSimulating)
 {
 	AUsdStageActor* CurrentStageActor = ViewModel.UsdStageActor.Get();
-	if ( !CurrentStageActor )
+	if (!CurrentStageActor)
 	{
 		return;
 	}
 
-	AUsdStageActor* PIEStageActor = Cast<AUsdStageActor>( EditorUtilities::GetSimWorldCounterpartActor( CurrentStageActor ) );
-	if ( !PIEStageActor )
+	AUsdStageActor* PIEStageActor = Cast<AUsdStageActor>(EditorUtilities::GetSimWorldCounterpartActor(CurrentStageActor));
+	if (!PIEStageActor)
 	{
 		return;
 	}
 
-	AttachToStageActor( PIEStageActor );
+	AttachToStageActor(PIEStageActor);
 }
 
-void SUsdStage::OnEndPIE( bool bIsSimulating )
+void SUsdStage::OnEndPIE(bool bIsSimulating)
 {
 	AUsdStageActor* CurrentStageActor = ViewModel.UsdStageActor.Get();
-	if ( !CurrentStageActor )
+	if (!CurrentStageActor)
 	{
 		return;
 	}
 
-	AUsdStageActor* EditorStageActor = Cast<AUsdStageActor>( EditorUtilities::GetEditorWorldCounterpartActor( CurrentStageActor ) );
-	if ( !EditorStageActor )
+	AUsdStageActor* EditorStageActor = Cast<AUsdStageActor>(EditorUtilities::GetEditorWorldCounterpartActor(CurrentStageActor));
+	if (!EditorStageActor)
 	{
 		return;
 	}
 
-	AttachToStageActor( EditorStageActor );
+	AttachToStageActor(EditorStageActor);
 }
 
 int32 SUsdStage::GetNaniteTriangleThresholdValue() const
@@ -2762,23 +2845,23 @@ int32 SUsdStage::GetNaniteTriangleThresholdValue() const
 	return CurrentNaniteThreshold;
 }
 
-void SUsdStage::OnNaniteTriangleThresholdValueChanged( int32 InValue )
+void SUsdStage::OnNaniteTriangleThresholdValueChanged(int32 InValue)
 {
 	CurrentNaniteThreshold = InValue;
 }
 
-void SUsdStage::OnNaniteTriangleThresholdValueCommitted( int32 InValue, ETextCommit::Type InCommitType )
+void SUsdStage::OnNaniteTriangleThresholdValueCommitted(int32 InValue, ETextCommit::Type InCommitType)
 {
 	AUsdStageActor* StageActor = GetStageActorOrCDO();
-	if ( !StageActor )
+	if (!StageActor)
 	{
 		return;
 	}
 
-	FScopedTransaction Transaction( FText::Format(
-		LOCTEXT( "NaniteTriangleThresholdCommittedTransaction", "Change Nanite triangle threshold for USD stage actor '{0}'" ),
-		FText::FromString( StageActor->GetActorLabel() )
-	) );
+	FScopedTransaction Transaction(FText::Format(
+		LOCTEXT("NaniteTriangleThresholdCommittedTransaction", "Change Nanite triangle threshold for USD stage actor '{0}'"),
+		FText::FromString(StageActor->GetActorLabel())
+	));
 
 	// c.f. comment in SUsdStage::FillCollapsingSubMenu
 	TGuardValue<bool> MaintainSelectionGuard(bUpdatingViewportSelection, true);
@@ -2828,7 +2911,7 @@ void SUsdStage::OnSubdivisionLevelValueCommitted(int32 InValue, ETextCommit::Typ
 UE::FUsdStageWeak SUsdStage::GetCurrentStage() const
 {
 	const AUsdStageActor* StageActor = ViewModel.UsdStageActor.Get();
-	if ( StageActor )
+	if (StageActor)
 	{
 		return StageActor->GetUsdStage();
 	}
@@ -2839,14 +2922,14 @@ UE::FUsdStageWeak SUsdStage::GetCurrentStage() const
 AUsdStageActor* SUsdStage::GetStageActorOrCDO() const
 {
 	AUsdStageActor* StageActor = ViewModel.UsdStageActor.Get();
-	if ( !StageActor )
+	if (!StageActor)
 	{
 		StageActor = GetMutableDefault<AUsdStageActor>();
 	}
-	ensure( StageActor );
+	ensure(StageActor);
 	return StageActor;
 }
 
-#endif // #if USE_USD_SDK
+#endif	  // #if USE_USD_SDK
 
 #undef LOCTEXT_NAMESPACE
