USD: Reformat USDStageEditor module.

diff --git a/Engine/Plugins/Importers/USDImporter/Source/USDStageEditor/Private/SUSDStageEditorStyle.cpp b/Engine/Plugins/Importers/USDImporter/Source/USDStageEditor/Private/SUSDStageEditorStyle.cpp
--- a/Engine/Plugins/Importers/USDImporter/Source/USDStageEditor/Private/SUSDStageEditorStyle.cpp
+++ b/Engine/Plugins/Importers/USDImporter/Source/USDStageEditor/Private/SUSDStageEditorStyle.cpp
@@ -2,17 +2,14 @@
 
 #include "SUSDStageEditorStyle.h"
 
-#include "Styling/AppStyle.h"
 #include "Interfaces/IPluginManager.h"
 #include "SlateOptMacros.h"
-#include "Styling/CoreStyle.h"
 #include "Styling/SlateStyleRegistry.h"
 #include "Styling/SlateTypes.h"
-#include "Styling/StyleColors.h"
 
-#define IMAGE_PLUGIN_BRUSH( RelativePath, ... ) FSlateImageBrush( FUsdStageEditorStyle::InContent( RelativePath, ".png" ), __VA_ARGS__ )
+#define IMAGE_PLUGIN_BRUSH(RelativePath, ...) FSlateImageBrush(FUsdStageEditorStyle::InContent(RelativePath, ".png"), __VA_ARGS__)
 #define IMAGE_BRUSH(RelativePath, ...) FSlateImageBrush(StyleSet->RootToContentDir(RelativePath, TEXT(".png")), __VA_ARGS__)
-#define IMAGE_BRUSH_SVG( RelativePath, ... ) FSlateVectorImageBrush(StyleSet->RootToContentDir(RelativePath, TEXT(".svg")), __VA_ARGS__)
+#define IMAGE_BRUSH_SVG(RelativePath, ...) FSlateVectorImageBrush(StyleSet->RootToContentDir(RelativePath, TEXT(".svg")), __VA_ARGS__)
 #define BOX_BRUSH(RelativePath, ...) FSlateBoxBrush(StyleSet->RootToContentDir(RelativePath, TEXT(".png")), __VA_ARGS__)
 #define DEFAULT_FONT(...) FCoreStyle::GetDefaultFontStyle(__VA_ARGS__)
 
@@ -22,8 +19,11 @@ FString FUsdStageEditorStyle::InContent(const FString& RelativePath, const ANSIC
 	return (ContentDir / RelativePath) + Extension;
 }
 
-TSharedPtr< FSlateStyleSet > FUsdStageEditorStyle::StyleSet = nullptr;
-TSharedPtr< class ISlateStyle > FUsdStageEditorStyle::Get() { return StyleSet; }
+TSharedPtr<FSlateStyleSet> FUsdStageEditorStyle::StyleSet = nullptr;
+TSharedPtr<class ISlateStyle> FUsdStageEditorStyle::Get()
+{
+	return StyleSet;
+}
 
 FName FUsdStageEditorStyle::GetStyleSetName()
 {
@@ -54,34 +54,34 @@ void FUsdStageEditorStyle::Initialize()
 
 	{
 		/** Height of items in the stage/layer tree views, prim info panel, etc. */
-		StyleSet->Set( "UsdStageEditor.ListItemHeight", 20.0f );
+		StyleSet->Set("UsdStageEditor.ListItemHeight", 20.0f);
 
 		/** Color of Prim labels that have composition arcs */
-		StyleSet->Set( "UsdStageEditor.PrimCompositionArcColor", FLinearColor( FColor::Orange ) );
+		StyleSet->Set("UsdStageEditor.PrimCompositionArcColor", FLinearColor(FColor::Orange));
 
 		/** Highlighted color of Prim labels that have composition arcs */
-		StyleSet->Set( "UsdStageEditor.HighlightPrimCompositionArcColor", FLinearColor( 0.2f, 0.08f, 0.0f, 1.0f ) );
+		StyleSet->Set("UsdStageEditor.HighlightPrimCompositionArcColor", FLinearColor(0.2f, 0.08f, 0.0f, 1.0f));
 
 		/**
 		 * Checkmark that is always white. We can't use the editor style one in some places (which by default is styled blue)
 		 * because the highlight color is also the same blue, so the checkbox would be invisible
 		 */
-		StyleSet->Set( "UsdStageEditor.CheckBoxImage", new IMAGE_BRUSH( "Common/Check", Icon14x14, FLinearColor::White ) );
+		StyleSet->Set("UsdStageEditor.CheckBoxImage", new IMAGE_BRUSH("Common/Check", Icon14x14, FLinearColor::White));
 
-		StyleSet->Set( "UsdStageEditor.ActorPickerFlashBrush", new FSlateRoundedBoxBrush( FLinearColor::White, 4.0f ) );
+		StyleSet->Set("UsdStageEditor.ActorPickerFlashBrush", new FSlateRoundedBoxBrush(FLinearColor::White, 4.0f));
 
 		/**
 		 * Button without a background, border or foreground color that doesn't move when pressed.
 		 * No foreground color is important to let the automatic slate foreground color propagate down.
 		 */
 		const FButtonStyle Style = FButtonStyle()
-			.SetNormal( FSlateNoResource() )
-			.SetHovered( FSlateNoResource() )
-			.SetPressed( FSlateNoResource() )
-			.SetDisabled( FSlateNoResource() )
-			.SetNormalPadding( FMargin( 2, 2, 2, 2 ) )
-			.SetPressedPadding( FMargin( 2, 2, 2, 2 ) );
-		StyleSet->Set( "NoBorder", Style );
+									   .SetNormal(FSlateNoResource())
+									   .SetHovered(FSlateNoResource())
+									   .SetPressed(FSlateNoResource())
+									   .SetDisabled(FSlateNoResource())
+									   .SetNormalPadding(FMargin(2, 2, 2, 2))
+									   .SetPressedPadding(FMargin(2, 2, 2, 2));
+		StyleSet->Set("NoBorder", Style);
 	}
 
 	FSlateStyleRegistry::RegisterSlateStyle(*StyleSet.Get());
diff --git a/Engine/Plugins/Importers/USDImporter/Source/USDStageEditor/Private/SUSDStageTreeView.cpp b/Engine/Plugins/Importers/USDImporter/Source/USDStageEditor/Private/SUSDStageTreeView.cpp
--- a/Engine/Plugins/Importers/USDImporter/Source/USDStageEditor/Private/SUSDStageTreeView.cpp
+++ b/Engine/Plugins/Importers/USDImporter/Source/USDStageEditor/Private/SUSDStageTreeView.cpp
@@ -11,21 +11,15 @@
 #include "USDOptionsWindow.h"
 #include "USDPrimViewModel.h"
 #include "USDReferenceOptions.h"
-#include "USDStageActor.h"
-#include "USDStageModule.h"
 #include "USDTypesConversion.h"
 
 #include "UsdWrappers/SdfChangeBlock.h"
-#include "UsdWrappers/SdfLayer.h"
 #include "UsdWrappers/SdfPath.h"
 #include "UsdWrappers/UsdPrim.h"
 #include "UsdWrappers/UsdStage.h"
 
-#include "Engine/World.h"
-#include "Framework/Application/SlateApplication.h"
 #include "Framework/Commands/GenericCommands.h"
 #include "Framework/MultiBox/MultiBoxBuilder.h"
-#include "Modules/ModuleManager.h"
 #include "ScopedTransaction.h"
 #include "Styling/AppStyle.h"
 #include "Templates/Function.h"
@@ -33,23 +27,23 @@
 #include "Widgets/Images/SImage.h"
 #include "Widgets/Input/SButton.h"
 #include "Widgets/Input/SCheckBox.h"
-#include "Widgets/Layout/SSeparator.h"
 #include "Widgets/SToolTip.h"
 #include "Widgets/Text/SInlineEditableTextBlock.h"
-#include "Widgets/Views/STableRow.h"
 
 #define LOCTEXT_NAMESPACE "UsdStageTreeView"
 
 #if USE_USD_SDK
 
 #include "USDIncludesStart.h"
-	#include "pxr/usd/usdPhysics/tokens.h"
+#include "pxr/usd/usdPhysics/tokens.h"
 #include "USDIncludesEnd.h"
 
 namespace UE::USDStageTreeView::Private
 {
-	static const FText NoSpecOnLocalLayerStack =
-		LOCTEXT( "NoLocalSpecToolTip", "This prim needs at least one spec on the stage's local layer stack for this option to be usable" );
+	static const FText NoSpecOnLocalLayerStack = LOCTEXT(
+		"NoLocalSpecToolTip",
+		"This prim needs at least one spec on the stage's local layer stack for this option to be usable"
+	);
 }
 
 enum class EPayloadsTrigger
@@ -59,87 +53,91 @@ enum class EPayloadsTrigger
 	Toggle
 };
 
-class FUsdStageNameColumn : public FUsdTreeViewColumn, public TSharedFromThis< FUsdStageNameColumn >
+class FUsdStageNameColumn
+	: public FUsdTreeViewColumn
+	, public TSharedFromThis<FUsdStageNameColumn>
 {
 public:
-	DECLARE_DELEGATE_TwoParams( FOnPrimNameCommitted, const FUsdPrimViewModelRef&, const FText& );
+	DECLARE_DELEGATE_TwoParams(FOnPrimNameCommitted, const FUsdPrimViewModelRef&, const FText&);
 	FOnPrimNameCommitted OnPrimNameCommitted;
 
-	DECLARE_DELEGATE_ThreeParams( FOnPrimNameUpdated, const FUsdPrimViewModelRef&, const FText&, FText& );
+	DECLARE_DELEGATE_ThreeParams(FOnPrimNameUpdated, const FUsdPrimViewModelRef&, const FText&, FText&);
 	FOnPrimNameUpdated OnPrimNameUpdated;
 
-	TWeakPtr< SUsdStageTreeView > OwnerTree;
+	TWeakPtr<SUsdStageTreeView> OwnerTree;
 
-	virtual TSharedRef< SWidget > GenerateWidget( const TSharedPtr< IUsdTreeViewItem > InTreeItem, const TSharedPtr< ITableRow > TableRow ) override
+	virtual TSharedRef<SWidget> GenerateWidget(const TSharedPtr<IUsdTreeViewItem> InTreeItem, const TSharedPtr<ITableRow> TableRow) override
 	{
-		if ( !InTreeItem )
+		if (!InTreeItem)
 		{
 			return SNullWidget::NullWidget;
 		}
 
-		TSharedPtr< FUsdPrimViewModel > TreeItem = StaticCastSharedPtr< FUsdPrimViewModel >( InTreeItem );
+		TSharedPtr<FUsdPrimViewModel> TreeItem = StaticCastSharedPtr<FUsdPrimViewModel>(InTreeItem);
 
-		TSharedRef< SInlineEditableTextBlock > Item =
-			SNew( SInlineEditableTextBlock )
-			.Text( TreeItem->RowData, &FUsdPrimModel::GetName )
-			.ColorAndOpacity( this, &FUsdStageNameColumn::GetTextColor, TreeItem )
-			.OnTextCommitted( this, &FUsdStageNameColumn::OnTextCommitted, TreeItem )
-			.OnVerifyTextChanged( this, &FUsdStageNameColumn::OnTextUpdated, TreeItem )
-			.IsReadOnly_Lambda( [ TreeItem ]()
+		// clang-format off
+		TSharedRef<SInlineEditableTextBlock> Item =
+			SNew(SInlineEditableTextBlock)
+			.Text(TreeItem->RowData, &FUsdPrimModel::GetName)
+			.ColorAndOpacity(this, &FUsdStageNameColumn::GetTextColor, TreeItem)
+			.OnTextCommitted(this, &FUsdStageNameColumn::OnTextCommitted, TreeItem)
+			.OnVerifyTextChanged(this, &FUsdStageNameColumn::OnTextUpdated, TreeItem)
+			.IsReadOnly_Lambda([TreeItem]()
 			{
 				return !TreeItem->bIsRenamingExistingPrim && (!TreeItem || TreeItem->UsdPrim);
-			} );
+			});
 
-		TreeItem->RenameRequestEvent.BindSP( &Item.Get(), &SInlineEditableTextBlock::EnterEditingMode );
+		TreeItem->RenameRequestEvent.BindSP(&Item.Get(), &SInlineEditableTextBlock::EnterEditingMode);
 
 		return SNew(SBox)
-			.VAlign( VAlign_Center )
+			.VAlign(VAlign_Center)
 			[
 				Item
 			];
+		// clang-format on
 	}
 
 protected:
-	void OnTextCommitted( const FText& InPrimName, ETextCommit::Type InCommitInfo, TSharedPtr< FUsdPrimViewModel > TreeItem )
+	void OnTextCommitted(const FText& InPrimName, ETextCommit::Type InCommitInfo, TSharedPtr<FUsdPrimViewModel> TreeItem)
 	{
-		if ( !TreeItem )
+		if (!TreeItem)
 		{
 			return;
 		}
 
-		OnPrimNameCommitted.ExecuteIfBound( TreeItem.ToSharedRef(), InPrimName );
+		OnPrimNameCommitted.ExecuteIfBound(TreeItem.ToSharedRef(), InPrimName);
 	}
 
-	bool OnTextUpdated(const FText& InPrimName, FText& ErrorMessage, TSharedPtr< FUsdPrimViewModel > TreeItem)
+	bool OnTextUpdated(const FText& InPrimName, FText& ErrorMessage, TSharedPtr<FUsdPrimViewModel> TreeItem)
 	{
 		if (!TreeItem)
 		{
 			return false;
 		}
 
-		OnPrimNameUpdated.ExecuteIfBound( TreeItem.ToSharedRef(), InPrimName, ErrorMessage );
+		OnPrimNameUpdated.ExecuteIfBound(TreeItem.ToSharedRef(), InPrimName, ErrorMessage);
 
 		return ErrorMessage.IsEmpty();
 	}
 
-	FSlateColor GetTextColor( TSharedPtr< FUsdPrimViewModel > TreeItem ) const
+	FSlateColor GetTextColor(TSharedPtr<FUsdPrimViewModel> TreeItem) const
 	{
 		FSlateColor TextColor = FSlateColor::UseForeground();
-		if ( !TreeItem )
+		if (!TreeItem)
 		{
 			return TextColor;
 		}
 
-		if ( TreeItem->RowData->HasCompositionArcs() )
+		if (TreeItem->RowData->HasCompositionArcs())
 		{
-			const TSharedPtr< SUsdStageTreeView > OwnerTreePtr = OwnerTree.Pin();
-			if ( OwnerTreePtr && OwnerTreePtr->IsItemSelected( TreeItem.ToSharedRef() ) )
+			const TSharedPtr<SUsdStageTreeView> OwnerTreePtr = OwnerTree.Pin();
+			if (OwnerTreePtr && OwnerTreePtr->IsItemSelected(TreeItem.ToSharedRef()))
 			{
-				TextColor = FUsdStageEditorStyle::Get()->GetColor( "UsdStageEditor.HighlightPrimCompositionArcColor" );
+				TextColor = FUsdStageEditorStyle::Get()->GetColor("UsdStageEditor.HighlightPrimCompositionArcColor");
 			}
 			else
 			{
-				TextColor = FUsdStageEditorStyle::Get()->GetColor( "UsdStageEditor.PrimCompositionArcColor" );
+				TextColor = FUsdStageEditorStyle::Get()->GetColor("UsdStageEditor.PrimCompositionArcColor");
 			}
 		}
 
@@ -147,14 +145,16 @@ protected:
 	}
 };
 
-class FUsdStagePayloadColumn : public FUsdTreeViewColumn, public TSharedFromThis< FUsdStagePayloadColumn >
+class FUsdStagePayloadColumn
+	: public FUsdTreeViewColumn
+	, public TSharedFromThis<FUsdStagePayloadColumn>
 {
 public:
-	ECheckBoxState IsChecked( const FUsdPrimViewModelPtr InTreeItem ) const
+	ECheckBoxState IsChecked(const FUsdPrimViewModelPtr InTreeItem) const
 	{
 		ECheckBoxState CheckedState = ECheckBoxState::Unchecked;
 
-		if ( InTreeItem && InTreeItem->RowData->HasPayload() )
+		if (InTreeItem && InTreeItem->RowData->HasPayload())
 		{
 			CheckedState = InTreeItem->RowData->IsLoaded() ? ECheckBoxState::Checked : ECheckBoxState::Unchecked;
 		}
@@ -162,84 +162,83 @@ public:
 		return CheckedState;
 	}
 
-	void OnCheckedPayload( ECheckBoxState NewCheckedState, const FUsdPrimViewModelPtr TreeItem )
+	void OnCheckedPayload(ECheckBoxState NewCheckedState, const FUsdPrimViewModelPtr TreeItem)
 	{
-		if ( !TreeItem )
+		if (!TreeItem)
 		{
 			return;
 		}
 
-		switch ( NewCheckedState )
+		switch (NewCheckedState)
 		{
-		case ECheckBoxState::Checked :
-			TreeItem->UsdPrim.Load();
-			break;
-		case ECheckBoxState::Unchecked:
-			TreeItem->UsdPrim.Unload();
-			break;
-		default:
-			break;
+			case ECheckBoxState::Checked:
+				TreeItem->UsdPrim.Load();
+				break;
+			case ECheckBoxState::Unchecked:
+				TreeItem->UsdPrim.Unload();
+				break;
+			default:
+				break;
 		}
 	}
 
-	virtual TSharedRef< SWidget > GenerateWidget( const TSharedPtr< IUsdTreeViewItem > InTreeItem, const TSharedPtr< ITableRow > TableRow ) override
+	virtual TSharedRef<SWidget> GenerateWidget(const TSharedPtr<IUsdTreeViewItem> InTreeItem, const TSharedPtr<ITableRow> TableRow) override
 	{
-		const TWeakPtr< FUsdPrimViewModel > TreeItem = StaticCastSharedPtr< FUsdPrimViewModel >( InTreeItem );
+		const TWeakPtr<FUsdPrimViewModel> TreeItem = StaticCastSharedPtr<FUsdPrimViewModel>(InTreeItem);
 
-		 TSharedRef< SWidget > Item = SNew( SCheckBox )
-			.Visibility_Lambda( [this, TreeItem]() -> EVisibility
+		// clang-format off
+		TSharedRef<SWidget> Item = SNew(SCheckBox)
+		.Visibility_Lambda([TreeItem]() -> EVisibility
+		{
+			if (const TSharedPtr<FUsdPrimViewModel> PinnedItem = TreeItem.Pin())
 			{
-				if ( const TSharedPtr< FUsdPrimViewModel > PinnedItem = TreeItem.Pin() )
-				{
-					return PinnedItem->RowData->HasPayload() ? EVisibility::Visible : EVisibility::Collapsed;
-				}
+				return PinnedItem->RowData->HasPayload() ? EVisibility::Visible : EVisibility::Collapsed;
+			}
 
-				return EVisibility::Collapsed;
-			})
-			.ToolTipText( LOCTEXT( "TogglePayloadToolTip", "Toggle payload" ) )
-			.IsChecked( this, &FUsdStagePayloadColumn::IsChecked, StaticCastSharedPtr< FUsdPrimViewModel >( InTreeItem ) )
-			.OnCheckStateChanged( this, &FUsdStagePayloadColumn::OnCheckedPayload, StaticCastSharedPtr< FUsdPrimViewModel >( InTreeItem ) );
+			return EVisibility::Collapsed;
+		})
+		.ToolTipText(LOCTEXT("TogglePayloadToolTip", "Toggle payload"))
+		.IsChecked(this, &FUsdStagePayloadColumn::IsChecked, StaticCastSharedPtr<FUsdPrimViewModel>(InTreeItem))
+		.OnCheckStateChanged(this, &FUsdStagePayloadColumn::OnCheckedPayload, StaticCastSharedPtr<FUsdPrimViewModel>(InTreeItem));
+		// clang-format on
 
 		return Item;
 	}
 };
 
-class FUsdStageVisibilityColumn : public FUsdTreeViewColumn, public TSharedFromThis< FUsdStageVisibilityColumn >
+class FUsdStageVisibilityColumn
+	: public FUsdTreeViewColumn
+	, public TSharedFromThis<FUsdStageVisibilityColumn>
 {
 public:
-	FReply OnToggleVisibility( const FUsdPrimViewModelPtr TreeItem )
+	FReply OnToggleVisibility(const FUsdPrimViewModelPtr TreeItem)
 	{
-		FScopedTransaction Transaction( FText::Format(
-			LOCTEXT( "VisibilityTransaction", "Toggle visibility of prim '{0}'" ),
-			FText::FromName( TreeItem->UsdPrim.GetName() )
-		) );
+		FScopedTransaction Transaction(
+			FText::Format(LOCTEXT("VisibilityTransaction", "Toggle visibility of prim '{0}'"), FText::FromName(TreeItem->UsdPrim.GetName()))
+		);
 
 		TreeItem->ToggleVisibility();
 
 		return FReply::Handled();
 	}
 
-	const FSlateBrush* GetBrush( const FUsdPrimViewModelPtr TreeItem, const TSharedPtr< SButton > Button ) const
+	const FSlateBrush* GetBrush(const FUsdPrimViewModelPtr TreeItem, const TSharedPtr<SButton> Button) const
 	{
 		const bool bIsButtonHovered = Button.IsValid() && Button->IsHovered();
 
-		if ( TreeItem->RowData->IsVisible() )
+		if (TreeItem->RowData->IsVisible())
 		{
-			return bIsButtonHovered
-				? FAppStyle::GetBrush( "Level.VisibleHighlightIcon16x" )
-				: FAppStyle::GetBrush( "Level.VisibleIcon16x" );
+			return bIsButtonHovered ? FAppStyle::GetBrush("Level.VisibleHighlightIcon16x") : FAppStyle::GetBrush("Level.VisibleIcon16x");
 		}
 		else
 		{
-			return bIsButtonHovered
-				? FAppStyle::GetBrush( "Level.NotVisibleHighlightIcon16x" )
-				: FAppStyle::GetBrush( "Level.NotVisibleIcon16x" );
+			return bIsButtonHovered ? FAppStyle::GetBrush("Level.NotVisibleHighlightIcon16x") : FAppStyle::GetBrush("Level.NotVisibleIcon16x");
 		}
 	}
 
-	FSlateColor GetForegroundColor( const FUsdPrimViewModelPtr TreeItem, const TSharedPtr< ITableRow > TableRow, const TSharedPtr< SButton > Button ) const
+	FSlateColor GetForegroundColor(const FUsdPrimViewModelPtr TreeItem, const TSharedPtr<ITableRow> TableRow, const TSharedPtr<SButton> Button) const
 	{
-		if ( !TreeItem.IsValid() || !TableRow.IsValid() || !Button.IsValid() )
+		if (!TreeItem.IsValid() || !TableRow.IsValid() || !Button.IsValid())
 		{
 			return FSlateColor::UseForeground();
 		}
@@ -249,164 +248,174 @@ public:
 		const bool bIsRowSelected = TableRow->IsItemSelected();
 		const bool bIsPrimVisible = TreeItem->RowData->IsVisible();
 
-		if ( bIsPrimVisible && !bIsRowHovered && !bIsRowSelected )
+		if (bIsPrimVisible && !bIsRowHovered && !bIsRowSelected)
 		{
 			return FLinearColor::Transparent;
 		}
-		else if ( bIsButtonHovered && !bIsRowSelected )
+		else if (bIsButtonHovered && !bIsRowSelected)
 		{
-			return FAppStyle::GetSlateColor( TEXT( "Colors.ForegroundHover" ) );
+			return FAppStyle::GetSlateColor(TEXT("Colors.ForegroundHover"));
 		}
 
 		return FSlateColor::UseForeground();
 	}
 
-	virtual TSharedRef< SWidget > GenerateWidget( const TSharedPtr< IUsdTreeViewItem > InTreeItem, const TSharedPtr< ITableRow > TableRow ) override
+	virtual TSharedRef<SWidget> GenerateWidget(const TSharedPtr<IUsdTreeViewItem> InTreeItem, const TSharedPtr<ITableRow> TableRow) override
 	{
-		if ( !InTreeItem )
+		if (!InTreeItem)
 		{
 			return SNullWidget::NullWidget;
 		}
 
-		const TSharedPtr< FUsdPrimViewModel > TreeItem = StaticCastSharedPtr< FUsdPrimViewModel >( InTreeItem );
-		const float ItemSize = FUsdStageEditorStyle::Get()->GetFloat( "UsdStageEditor.ListItemHeight" );
+		const TSharedPtr<FUsdPrimViewModel> TreeItem = StaticCastSharedPtr<FUsdPrimViewModel>(InTreeItem);
+		const float ItemSize = FUsdStageEditorStyle::Get()->GetFloat("UsdStageEditor.ListItemHeight");
 
-		if ( !TreeItem->HasVisibilityAttribute() )
+		// clang-format off
+		if (!TreeItem->HasVisibilityAttribute())
 		{
 			return SNew(SBox)
-				.HeightOverride( ItemSize )
-				.WidthOverride( ItemSize )
-				.Visibility( EVisibility::Visible )
-				.ToolTip( SNew( SToolTip ).Text( LOCTEXT( "NoGeomImageable", "Only prims with the GeomImageable schema (or derived) have the visibility attribute!" ) ) );
+				.HeightOverride(ItemSize)
+				.WidthOverride(ItemSize)
+				.Visibility(EVisibility::Visible)
+				.ToolTip(SNew(SToolTip).Text(LOCTEXT("NoGeomImageable", "Only prims with the GeomImageable schema (or derived) have the visibility attribute!")));
 		}
 
-		TSharedPtr<SButton> Button = SNew( SButton )
-			.ContentPadding( 0 )
-			.ButtonStyle( FUsdStageEditorStyle::Get(), TEXT("NoBorder") )
-			.OnClicked( this, &FUsdStageVisibilityColumn::OnToggleVisibility, TreeItem )
-			.ToolTip( SNew( SToolTip ).Text( LOCTEXT( "GeomImageable", "Toggle the visibility of this prim" ) ) )
-			.HAlign( HAlign_Center )
-			.VAlign( VAlign_Center );
+		TSharedPtr<SButton> Button = SNew(SButton)
+			.ContentPadding(0)
+			.ButtonStyle(FUsdStageEditorStyle::Get(), TEXT("NoBorder"))
+			.OnClicked(this, &FUsdStageVisibilityColumn::OnToggleVisibility, TreeItem)
+			.ToolTip(SNew(SToolTip).Text(LOCTEXT("GeomImageable", "Toggle the visibility of this prim")))
+			.HAlign(HAlign_Center)
+			.VAlign(VAlign_Center);
 
-		TSharedPtr<SImage> Image = SNew( SImage )
-			.Image( this, &FUsdStageVisibilityColumn::GetBrush, TreeItem, Button )
-			.ColorAndOpacity( this, &FUsdStageVisibilityColumn::GetForegroundColor, TreeItem, TableRow, Button );
+		TSharedPtr<SImage> Image = SNew(SImage)
+			.Image(this, &FUsdStageVisibilityColumn::GetBrush, TreeItem, Button)
+			.ColorAndOpacity(this, &FUsdStageVisibilityColumn::GetForegroundColor, TreeItem, TableRow, Button);
 
-		Button->SetContent( Image.ToSharedRef() );
+		Button->SetContent(Image.ToSharedRef());
 
-		return SNew( SBox )
-			.HeightOverride( ItemSize )
-			.WidthOverride( ItemSize )
-			.Visibility( EVisibility::Visible )
+		return SNew(SBox)
+			.HeightOverride(ItemSize)
+			.WidthOverride(ItemSize)
+			.Visibility(EVisibility::Visible)
 			[
 				Button.ToSharedRef()
 			];
+		// clang-format on
 	}
 };
 
 class FUsdStagePrimTypeColumn : public FUsdTreeViewColumn
 {
 public:
-	virtual TSharedRef< SWidget > GenerateWidget( const TSharedPtr< IUsdTreeViewItem > InTreeItem, const TSharedPtr< ITableRow > TableRow ) override
+	virtual TSharedRef<SWidget> GenerateWidget(const TSharedPtr<IUsdTreeViewItem> InTreeItem, const TSharedPtr<ITableRow> TableRow) override
 	{
-		TSharedPtr< FUsdPrimViewModel > TreeItem = StaticCastSharedPtr< FUsdPrimViewModel >( InTreeItem );
+		TSharedPtr<FUsdPrimViewModel> TreeItem = StaticCastSharedPtr<FUsdPrimViewModel>(InTreeItem);
 
+		// clang-format off
 		return SNew(SBox)
-			.VAlign( VAlign_Center )
+			.VAlign(VAlign_Center)
 			[
 				SNew(STextBlock)
-				.Text( TreeItem->RowData, &FUsdPrimModel::GetType )
+				.Text(TreeItem->RowData, &FUsdPrimModel::GetType)
 			];
+		// clang-format on
 	}
 };
 
-void SUsdStageTreeView::Construct( const FArguments& InArgs )
+void SUsdStageTreeView::Construct(const FArguments& InArgs)
 {
-	SUsdTreeView::Construct( SUsdTreeView::FArguments() );
+	SUsdTreeView::Construct(SUsdTreeView::FArguments());
 
-	OnContextMenuOpening = FOnContextMenuOpening::CreateSP( this, &SUsdStageTreeView::ConstructPrimContextMenu );
+	OnContextMenuOpening = FOnContextMenuOpening::CreateSP(this, &SUsdStageTreeView::ConstructPrimContextMenu);
 
-	OnSelectionChanged = FOnSelectionChanged::CreateLambda( [ this ]( FUsdPrimViewModelPtr UsdStageTreeItem, ESelectInfo::Type SelectionType )
-	{
-		FString SelectedPrimPath;
-
-		if ( UsdStageTreeItem )
+	OnSelectionChanged = FOnSelectionChanged::CreateLambda(
+		[this](FUsdPrimViewModelPtr UsdStageTreeItem, ESelectInfo::Type SelectionType)
 		{
-			SelectedPrimPath = UsdToUnreal::ConvertPath( UsdStageTreeItem->UsdPrim.GetPrimPath() );
-		}
+			FString SelectedPrimPath;
 
-		TArray<FString> SelectedPrimPaths = GetSelectedPrimPaths();
-		this->OnPrimSelectionChanged.ExecuteIfBound( SelectedPrimPaths );
-	} );
+			if (UsdStageTreeItem)
+			{
+				SelectedPrimPath = UsdToUnreal::ConvertPath(UsdStageTreeItem->UsdPrim.GetPrimPath());
+			}
 
-	OnExpansionChanged = FOnExpansionChanged::CreateLambda([this]( const FUsdPrimViewModelPtr& UsdPrimViewModel, bool bIsExpanded)
-	{
-		if ( !UsdPrimViewModel )
-		{
-			return;
+			TArray<FString> SelectedPrimPaths = GetSelectedPrimPaths();
+			this->OnPrimSelectionChanged.ExecuteIfBound(SelectedPrimPaths);
 		}
+	);
 
-		const UE::FUsdPrim& Prim = UsdPrimViewModel->UsdPrim;
-		if ( !Prim )
+	OnExpansionChanged = FOnExpansionChanged::CreateLambda(
+		[this](const FUsdPrimViewModelPtr& UsdPrimViewModel, bool bIsExpanded)
 		{
-			return;
-		}
+			if (!UsdPrimViewModel)
+			{
+				return;
+			}
+
+			const UE::FUsdPrim& Prim = UsdPrimViewModel->UsdPrim;
+			if (!Prim)
+			{
+				return;
+			}
 
-		TreeItemExpansionStates.Add( Prim.GetPrimPath().GetString(), bIsExpanded );
-	});
+			TreeItemExpansionStates.Add(Prim.GetPrimPath().GetString(), bIsExpanded);
+		}
+	);
 
 	OnPrimSelectionChanged = InArgs._OnPrimSelectionChanged;
 
 	UICommandList = MakeShared<FUICommandList>();
 	UICommandList->MapAction(
 		FGenericCommands::Get().Cut,
-		FExecuteAction::CreateSP( this, &SUsdStageTreeView::OnCutPrim ),
-		FCanExecuteAction::CreateSP( this, &SUsdStageTreeView::DoesPrimHaveSpecOnLocalLayerStack )
+		FExecuteAction::CreateSP(this, &SUsdStageTreeView::OnCutPrim),
+		FCanExecuteAction::CreateSP(this, &SUsdStageTreeView::DoesPrimHaveSpecOnLocalLayerStack)
 	);
 	UICommandList->MapAction(
 		FGenericCommands::Get().Copy,
-		FExecuteAction::CreateSP( this, &SUsdStageTreeView::OnCopyPrim ),
-		FCanExecuteAction::CreateSP( this, &SUsdStageTreeView::DoesPrimExistOnStage )
+		FExecuteAction::CreateSP(this, &SUsdStageTreeView::OnCopyPrim),
+		FCanExecuteAction::CreateSP(this, &SUsdStageTreeView::DoesPrimExistOnStage)
 	);
 	UICommandList->MapAction(
 		FGenericCommands::Get().Paste,
-		FExecuteAction::CreateSP( this, &SUsdStageTreeView::OnPastePrim ),
-		FCanExecuteAction::CreateSP( this, &SUsdStageTreeView::CanPastePrim )
+		FExecuteAction::CreateSP(this, &SUsdStageTreeView::OnPastePrim),
+		FCanExecuteAction::CreateSP(this, &SUsdStageTreeView::CanPastePrim)
 	);
 	UICommandList->MapAction(
 		FGenericCommands::Get().Duplicate,
-		FExecuteAction::CreateSP( this, &SUsdStageTreeView::OnDuplicatePrim, EUsdDuplicateType::AllLocalLayerSpecs ),
-		FCanExecuteAction::CreateSP( this, &SUsdStageTreeView::DoesPrimExistOnStage )
+		FExecuteAction::CreateSP(this, &SUsdStageTreeView::OnDuplicatePrim, EUsdDuplicateType::AllLocalLayerSpecs),
+		FCanExecuteAction::CreateSP(this, &SUsdStageTreeView::DoesPrimExistOnStage)
 	);
 	UICommandList->MapAction(
 		FGenericCommands::Get().Delete,
-		FExecuteAction::CreateSP( this, &SUsdStageTreeView::OnDeletePrim ),
-		FCanExecuteAction::CreateSP( this, &SUsdStageTreeView::DoesPrimHaveSpecOnLocalLayerStack )
+		FExecuteAction::CreateSP(this, &SUsdStageTreeView::OnDeletePrim),
+		FCanExecuteAction::CreateSP(this, &SUsdStageTreeView::DoesPrimHaveSpecOnLocalLayerStack)
 	);
 	UICommandList->MapAction(
 		FGenericCommands::Get().Rename,
-		FExecuteAction::CreateSP( this, &SUsdStageTreeView::OnRenamePrim ),
-		FCanExecuteAction::CreateLambda( [this]()
-		{
-			return SUsdStageTreeView::DoesPrimHaveSpecOnLocalLayerStack() && GetSelectedItems().Num() == 1;
-		})
+		FExecuteAction::CreateSP(this, &SUsdStageTreeView::OnRenamePrim),
+		FCanExecuteAction::CreateLambda(
+			[this]()
+			{
+				return SUsdStageTreeView::DoesPrimHaveSpecOnLocalLayerStack() && GetSelectedItems().Num() == 1;
+			}
+		)
 	);
 }
 
-TSharedRef< ITableRow > SUsdStageTreeView::OnGenerateRow( FUsdPrimViewModelRef InDisplayNode, const TSharedRef<STableViewBase>& OwnerTable )
+TSharedRef<ITableRow> SUsdStageTreeView::OnGenerateRow(FUsdPrimViewModelRef InDisplayNode, const TSharedRef<STableViewBase>& OwnerTable)
 {
-	return SNew( SUsdTreeRow< FUsdPrimViewModelRef >, InDisplayNode, OwnerTable, SharedData );
+	return SNew(SUsdTreeRow<FUsdPrimViewModelRef>, InDisplayNode, OwnerTable, SharedData);
 }
 
-void SUsdStageTreeView::OnGetChildren( FUsdPrimViewModelRef InParent, TArray< FUsdPrimViewModelRef >& OutChildren ) const
+void SUsdStageTreeView::OnGetChildren(FUsdPrimViewModelRef InParent, TArray<FUsdPrimViewModelRef>& OutChildren) const
 {
-	for ( const FUsdPrimViewModelRef& Child : InParent->UpdateChildren() )
+	for (const FUsdPrimViewModelRef& Child : InParent->UpdateChildren())
 	{
-		OutChildren.Add( Child );
+		OutChildren.Add(Child);
 	}
 }
 
-void SUsdStageTreeView::Refresh( const UE::FUsdStageWeak& NewStage )
+void SUsdStageTreeView::Refresh(const UE::FUsdStageWeak& NewStage)
 {
 	UE::FUsdStageWeak OldStage = RootItems.Num() > 0 ? RootItems[0]->UsdStage : UE::FUsdStageWeak();
 
@@ -414,42 +423,42 @@ void SUsdStageTreeView::Refresh( const UE::FUsdStageWeak& NewStage )
 
 	UsdStage = NewStage;
 
-	if ( OldStage != NewStage )
+	if (OldStage != NewStage)
 	{
 		TreeItemExpansionStates.Reset();
 	}
 
-	if ( NewStage )
+	if (NewStage)
 	{
-		if ( UE::FUsdPrim RootPrim = NewStage.GetPseudoRoot() )
+		if (UE::FUsdPrim RootPrim = NewStage.GetPseudoRoot())
 		{
-			RootItems.Add( MakeShared< FUsdPrimViewModel >( nullptr, UsdStage, RootPrim ) );
+			RootItems.Add(MakeShared<FUsdPrimViewModel>(nullptr, UsdStage, RootPrim));
 		}
 
 		RestoreExpansionStates();
 	}
 }
 
-void SUsdStageTreeView::RefreshPrim( const FString& PrimPath, bool bResync )
+void SUsdStageTreeView::RefreshPrim(const FString& PrimPath, bool bResync)
 {
-	FScopedUnrealAllocs UnrealAllocs; // RefreshPrim can be called by a delegate for which we don't know the active allocator
+	FScopedUnrealAllocs UnrealAllocs;	 // RefreshPrim can be called by a delegate for which we don't know the active allocator
 
 	FUsdPrimViewModelPtr FoundItem = GetItemFromPrimPath(PrimPath);
 
-	if ( FoundItem.IsValid() )
+	if (FoundItem.IsValid())
 	{
-		FoundItem->RefreshData( true );
+		FoundItem->RefreshData(true);
 
 		// Item doesn't match any prim, needs to be removed
-		if ( !FoundItem->UsdPrim )
+		if (!FoundItem->UsdPrim)
 		{
-			if ( FoundItem->ParentItem )
+			if (FoundItem->ParentItem)
 			{
-				FoundItem->ParentItem->RefreshData( true );
+				FoundItem->ParentItem->RefreshData(true);
 			}
 			else
 			{
-				RootItems.Remove( FoundItem.ToSharedRef() );
+				RootItems.Remove(FoundItem.ToSharedRef());
 			}
 		}
 	}
@@ -460,31 +469,32 @@ void SUsdStageTreeView::RefreshPrim( const FString& PrimPath, bool bResync )
 		bNeedsFullUpdate = true;
 	}
 
-	if ( bResync )
+	if (bResync)
 	{
 		FWriteScopeLock StageTreeViewLock{RefreshStateLock};
 		bNeedsFullUpdate = true;
 	}
 }
 
-FUsdPrimViewModelPtr SUsdStageTreeView::GetItemFromPrimPath( const FString& PrimPath )
+FUsdPrimViewModelPtr SUsdStageTreeView::GetItemFromPrimPath(const FString& PrimPath)
 {
 	FScopedUnrealAllocs UnrealAllocs;
 
-	UE::FSdfPath UsdPrimPath( *PrimPath );
+	UE::FSdfPath UsdPrimPath(*PrimPath);
 
-	TFunction< FUsdPrimViewModelPtr( const UE::FSdfPath&, const FUsdPrimViewModelRef& ) > FindTreeItemFromPrimPath;
-	FindTreeItemFromPrimPath = [ &FindTreeItemFromPrimPath ]( const UE::FSdfPath& UsdPrimPath, const FUsdPrimViewModelRef& ItemRef ) -> FUsdPrimViewModelPtr
+	TFunction<FUsdPrimViewModelPtr(const UE::FSdfPath&, const FUsdPrimViewModelRef&)> FindTreeItemFromPrimPath;
+	FindTreeItemFromPrimPath =
+		[&FindTreeItemFromPrimPath](const UE::FSdfPath& UsdPrimPath, const FUsdPrimViewModelRef& ItemRef) -> FUsdPrimViewModelPtr
 	{
-		if ( ItemRef->UsdPrim.GetPrimPath() == UsdPrimPath )
+		if (ItemRef->UsdPrim.GetPrimPath() == UsdPrimPath)
 		{
 			return ItemRef;
 		}
 		else
 		{
-			for ( FUsdPrimViewModelRef ChildItem : ItemRef->Children )
+			for (FUsdPrimViewModelRef ChildItem : ItemRef->Children)
 			{
-				if ( FUsdPrimViewModelPtr ChildValue = FindTreeItemFromPrimPath( UsdPrimPath, ChildItem ) )
+				if (FUsdPrimViewModelPtr ChildValue = FindTreeItemFromPrimPath(UsdPrimPath, ChildItem))
 				{
 					return ChildValue;
 				}
@@ -496,25 +506,25 @@ FUsdPrimViewModelPtr SUsdStageTreeView::GetItemFromPrimPath( const FString& Prim
 
 	// Search for the corresponding tree item to update
 	FUsdPrimViewModelPtr FoundItem = nullptr;
-	for ( FUsdPrimViewModelRef RootItem : this->RootItems )
+	for (FUsdPrimViewModelRef RootItem : this->RootItems)
 	{
 		UE::FSdfPath PrimPathToSearch = UsdPrimPath;
 
-		FoundItem = FindTreeItemFromPrimPath( PrimPathToSearch, RootItem );
+		FoundItem = FindTreeItemFromPrimPath(PrimPathToSearch, RootItem);
 
-		while ( !FoundItem.IsValid() )
+		while (!FoundItem.IsValid())
 		{
 			UE::FSdfPath ParentPrimPath = PrimPathToSearch.GetParentPath();
-			if ( ParentPrimPath == PrimPathToSearch )
+			if (ParentPrimPath == PrimPathToSearch)
 			{
 				break;
 			}
-			PrimPathToSearch = MoveTemp( ParentPrimPath );
+			PrimPathToSearch = MoveTemp(ParentPrimPath);
 
-			FoundItem = FindTreeItemFromPrimPath( PrimPathToSearch, RootItem );
+			FoundItem = FindTreeItemFromPrimPath(PrimPathToSearch, RootItem);
 		}
 
-		if ( FoundItem.IsValid() )
+		if (FoundItem.IsValid())
 		{
 			break;
 		}
@@ -523,9 +533,9 @@ FUsdPrimViewModelPtr SUsdStageTreeView::GetItemFromPrimPath( const FString& Prim
 	return FoundItem;
 }
 
-void SUsdStageTreeView::SelectItemsInternal( const TArray<FUsdPrimViewModelRef>& ItemsToSelect )
+void SUsdStageTreeView::SelectItemsInternal(const TArray<FUsdPrimViewModelRef>& ItemsToSelect)
 {
-	if ( ItemsToSelect.Num() > 0 )
+	if (ItemsToSelect.Num() > 0)
 	{
 		// Clear selection without emitting events, as we'll emit new events with SetItemSelection
 		// anyway. This prevents a UI blink as OnPrimSelectionChanged would otherwise fire for
@@ -533,8 +543,8 @@ void SUsdStageTreeView::SelectItemsInternal( const TArray<FUsdPrimViewModelRef>&
 		Private_ClearSelection();
 
 		const bool bSelected = true;
-		SetItemSelection( ItemsToSelect, bSelected );
-		ScrollItemIntoView( ItemsToSelect.Last() );
+		SetItemSelection(ItemsToSelect, bSelected);
+		ScrollItemIntoView(ItemsToSelect.Last());
 	}
 	else
 	{
@@ -542,63 +552,63 @@ void SUsdStageTreeView::SelectItemsInternal( const TArray<FUsdPrimViewModelRef>&
 
 		// ClearSelection is not going to fire the OnSelectionChanged event in case we have nothing selected, but we
 		// need to do that to refresh the prim properties panel to display the stage properties instead
-		OnPrimSelectionChanged.ExecuteIfBound( {} );
+		OnPrimSelectionChanged.ExecuteIfBound({});
 	}
 }
 
-void SUsdStageTreeView::SetSelectedPrimPaths( const TArray<FString>& PrimPaths )
+void SUsdStageTreeView::SetSelectedPrimPaths(const TArray<FString>& PrimPaths)
 {
-	TArray< FUsdPrimViewModelRef > ItemsToSelect;
-	ItemsToSelect.Reserve( PrimPaths.Num() );
+	TArray<FUsdPrimViewModelRef> ItemsToSelect;
+	ItemsToSelect.Reserve(PrimPaths.Num());
 
-	for ( const FString& PrimPath : PrimPaths )
+	for (const FString& PrimPath : PrimPaths)
 	{
-		if ( FUsdPrimViewModelPtr FoundItem = GetItemFromPrimPath( PrimPath ) )
+		if (FUsdPrimViewModelPtr FoundItem = GetItemFromPrimPath(PrimPath))
 		{
-			ItemsToSelect.Add( FoundItem.ToSharedRef() );
+			ItemsToSelect.Add(FoundItem.ToSharedRef());
 		}
 	}
 
-	SelectItemsInternal( ItemsToSelect );
+	SelectItemsInternal(ItemsToSelect);
 }
 
-void SUsdStageTreeView::SetSelectedPrims( const TArray<UE::FUsdPrim>& Prims )
+void SUsdStageTreeView::SetSelectedPrims(const TArray<UE::FUsdPrim>& Prims)
 {
 	TSet<UE::FSdfPath> PrimPaths;
-	for ( const UE::FUsdPrim& Prim : Prims )
+	for (const UE::FUsdPrim& Prim : Prims)
 	{
-		PrimPaths.Add( Prim.GetPrimPath() );
+		PrimPaths.Add(Prim.GetPrimPath());
 	}
 
-	TArray< FUsdPrimViewModelRef > ItemsToSelect;
-	ItemsToSelect.Reserve( Prims.Num() );
+	TArray<FUsdPrimViewModelRef> ItemsToSelect;
+	ItemsToSelect.Reserve(Prims.Num());
 
-	TFunction<void( const TArray< FUsdPrimViewModelRef >& )> Traverse;
-	Traverse = [this, &PrimPaths, &ItemsToSelect, &Traverse]( const TArray< FUsdPrimViewModelRef >& Items )
+	TFunction<void(const TArray<FUsdPrimViewModelRef>&)> Traverse;
+	Traverse = [this, &PrimPaths, &ItemsToSelect, &Traverse](const TArray<FUsdPrimViewModelRef>& Items)
 	{
-		for ( const FUsdPrimViewModelRef& Item : Items )
+		for (const FUsdPrimViewModelRef& Item : Items)
 		{
-			if ( PrimPaths.Contains( Item->UsdPrim.GetPrimPath() ) )
+			if (PrimPaths.Contains(Item->UsdPrim.GetPrimPath()))
 			{
-				ItemsToSelect.Add( Item );
+				ItemsToSelect.Add(Item);
 			}
 
-			Traverse( Item->Children );
+			Traverse(Item->Children);
 		}
 	};
-	Traverse( RootItems );
+	Traverse(RootItems);
 
-	SelectItemsInternal( ItemsToSelect );
+	SelectItemsInternal(ItemsToSelect);
 }
 
 TArray<FString> SUsdStageTreeView::GetSelectedPrimPaths()
 {
 	TArray<FString> SelectedPrimPaths;
-	SelectedPrimPaths.Reserve( GetNumItemsSelected() );
+	SelectedPrimPaths.Reserve(GetNumItemsSelected());
 
-	for ( FUsdPrimViewModelRef SelectedItem : GetSelectedItems() )
+	for (FUsdPrimViewModelRef SelectedItem : GetSelectedItems())
 	{
-		SelectedPrimPaths.Add( SelectedItem->UsdPrim.GetPrimPath().GetString() );
+		SelectedPrimPaths.Add(SelectedItem->UsdPrim.GetPrimPath().GetString());
 	}
 
 	return SelectedPrimPaths;
@@ -607,11 +617,11 @@ TArray<FString> SUsdStageTreeView::GetSelectedPrimPaths()
 TArray<UE::FUsdPrim> SUsdStageTreeView::GetSelectedPrims()
 {
 	TArray<UE::FUsdPrim> SelectedPrims;
-	SelectedPrims.Reserve( GetNumItemsSelected() );
+	SelectedPrims.Reserve(GetNumItemsSelected());
 
-	for ( FUsdPrimViewModelRef SelectedItem : GetSelectedItems() )
+	for (FUsdPrimViewModelRef SelectedItem : GetSelectedItems())
 	{
-		SelectedPrims.Add( SelectedItem->UsdPrim );
+		SelectedPrims.Add(SelectedItem->UsdPrim);
 	}
 
 	return SelectedPrims;
@@ -622,57 +632,58 @@ void SUsdStageTreeView::SetupColumns()
 	HeaderRowWidget->ClearColumns();
 
 	SHeaderRow::FColumn::FArguments VisibilityColumnArguments;
-	VisibilityColumnArguments.FixedWidth( 24.f );
+	VisibilityColumnArguments.FixedWidth(24.f);
 
-	AddColumn( TEXT( "Visibility" ), FText(), MakeShared< FUsdStageVisibilityColumn >(), VisibilityColumnArguments );
+	AddColumn(TEXT("Visibility"), FText(), MakeShared<FUsdStageVisibilityColumn>(), VisibilityColumnArguments);
 
 	{
-		TSharedRef< FUsdStageNameColumn > PrimNameColumn = MakeShared< FUsdStageNameColumn >();
-		PrimNameColumn->OwnerTree = SharedThis( this );
+		TSharedRef<FUsdStageNameColumn> PrimNameColumn = MakeShared<FUsdStageNameColumn>();
+		PrimNameColumn->OwnerTree = SharedThis(this);
 		PrimNameColumn->bIsMainColumn = true;
-		PrimNameColumn->OnPrimNameCommitted.BindRaw( this, &SUsdStageTreeView::OnPrimNameCommitted );
-		PrimNameColumn->OnPrimNameUpdated.BindRaw( this, &SUsdStageTreeView::OnPrimNameUpdated );
+		PrimNameColumn->OnPrimNameCommitted.BindRaw(this, &SUsdStageTreeView::OnPrimNameCommitted);
+		PrimNameColumn->OnPrimNameUpdated.BindRaw(this, &SUsdStageTreeView::OnPrimNameUpdated);
 
 		SHeaderRow::FColumn::FArguments PrimNameColumnArguments;
-		PrimNameColumnArguments.FillWidth( 70.f );
+		PrimNameColumnArguments.FillWidth(70.f);
 
-		AddColumn( TEXT("Prim"), LOCTEXT( "Prim", "Prim" ), PrimNameColumn, PrimNameColumnArguments );
+		AddColumn(TEXT("Prim"), LOCTEXT("Prim", "Prim"), PrimNameColumn, PrimNameColumnArguments);
 	}
 
 	SHeaderRow::FColumn::FArguments TypeColumnArguments;
-	TypeColumnArguments.FillWidth( 15.f );
-	AddColumn( TEXT("Type"), LOCTEXT( "Type", "Type" ), MakeShared< FUsdStagePrimTypeColumn >(), TypeColumnArguments );
+	TypeColumnArguments.FillWidth(15.f);
+	AddColumn(TEXT("Type"), LOCTEXT("Type", "Type"), MakeShared<FUsdStagePrimTypeColumn>(), TypeColumnArguments);
 
 	SHeaderRow::FColumn::FArguments PayloadColumnArguments;
-	PayloadColumnArguments.FillWidth( 15.f ).HAlignHeader( HAlign_Center ).HAlignCell( HAlign_Center );
-	AddColumn( TEXT("Payload"), LOCTEXT( "Payload", "Payload" ), MakeShared< FUsdStagePayloadColumn >(), PayloadColumnArguments );
+	PayloadColumnArguments.FillWidth(15.f).HAlignHeader(HAlign_Center).HAlignCell(HAlign_Center);
+	AddColumn(TEXT("Payload"), LOCTEXT("Payload", "Payload"), MakeShared<FUsdStagePayloadColumn>(), PayloadColumnArguments);
 }
 
-TSharedPtr< SWidget > SUsdStageTreeView::ConstructPrimContextMenu()
+TSharedPtr<SWidget> SUsdStageTreeView::ConstructPrimContextMenu()
 {
-	TSharedRef< SWidget > MenuWidget = SNullWidget::NullWidget;
+	TSharedRef<SWidget> MenuWidget = SNullWidget::NullWidget;
 
 	const bool bInShouldCloseWindowAfterMenuSelection = true;
-	FMenuBuilder PrimOptions( bInShouldCloseWindowAfterMenuSelection, UICommandList );
-	PrimOptions.BeginSection( "Edit", LOCTEXT("EditText", "Edit") );
+	FMenuBuilder PrimOptions(bInShouldCloseWindowAfterMenuSelection, UICommandList);
+	PrimOptions.BeginSection("Edit", LOCTEXT("EditText", "Edit"));
 	{
 		PrimOptions.AddMenuEntry(
-			TAttribute<FText>::Create( TAttribute<FText>::FGetter::CreateLambda( [this]()
-			{
-				return GetSelectedItems().Num() == 0
-					? LOCTEXT( "AddTopLevelPrim", "Add Prim" )
-					: LOCTEXT( "AddPrim", "Add Child" );
-			})),
-			TAttribute<FText>::Create( TAttribute<FText>::FGetter::CreateLambda( [this]()
-			{
-				return GetSelectedItems().Num() == 0
-					? LOCTEXT( "AddTopPrim_ToolTip", "Adds a new top-level prim" )
-					: LOCTEXT( "AddPrim_ToolTip", "Adds a new prim as a child of this one" );
-			})),
-			FSlateIcon( FAppStyle::GetAppStyleSetName(), "Icons.PlaceActors" ),
+			TAttribute<FText>::Create(TAttribute<FText>::FGetter::CreateLambda(
+				[this]()
+				{
+					return GetSelectedItems().Num() == 0 ? LOCTEXT("AddTopLevelPrim", "Add Prim") : LOCTEXT("AddPrim", "Add Child");
+				}
+			)),
+			TAttribute<FText>::Create(TAttribute<FText>::FGetter::CreateLambda(
+				[this]()
+				{
+					return GetSelectedItems().Num() == 0 ? LOCTEXT("AddTopPrim_ToolTip", "Adds a new top-level prim")
+														 : LOCTEXT("AddPrim_ToolTip", "Adds a new prim as a child of this one");
+				}
+			)),
+			FSlateIcon(FAppStyle::GetAppStyleSetName(), "Icons.PlaceActors"),
 			FUIAction(
-				FExecuteAction::CreateSP( this, &SUsdStageTreeView::OnAddChildPrim ),
-				FCanExecuteAction::CreateSP( this, &SUsdStageTreeView::CanAddChildPrim )
+				FExecuteAction::CreateSP(this, &SUsdStageTreeView::OnAddChildPrim),
+				FCanExecuteAction::CreateSP(this, &SUsdStageTreeView::CanAddChildPrim)
 			),
 			NAME_None,
 			EUserInterfaceActionType::Button
@@ -681,74 +692,77 @@ TSharedPtr< SWidget > SUsdStageTreeView::ConstructPrimContextMenu()
 		PrimOptions.AddMenuEntry(
 			FGenericCommands::Get().Cut,
 			NAME_None,
-			LOCTEXT( "Cut_Text", "Cut" ),
-			TAttribute<FText>::Create( TAttribute<FText>::FGetter::CreateLambda( [this]()
-			{
-				return DoesPrimHaveSpecOnLocalLayerStack()
-					? LOCTEXT( "Cut_ToolTip", "Cuts the selected prim's specs from the stage's local layer stack" )
-					: UE::USDStageTreeView::Private::NoSpecOnLocalLayerStack;
-			}))
+			LOCTEXT("Cut_Text", "Cut"),
+			TAttribute<FText>::Create(TAttribute<FText>::FGetter::CreateLambda(
+				[this]()
+				{
+					return DoesPrimHaveSpecOnLocalLayerStack()
+							   ? LOCTEXT("Cut_ToolTip", "Cuts the selected prim's specs from the stage's local layer stack")
+							   : UE::USDStageTreeView::Private::NoSpecOnLocalLayerStack;
+				}
+			))
 		);
 
 		PrimOptions.AddMenuEntry(
 			FGenericCommands::Get().Copy,
 			NAME_None,
-			LOCTEXT( "Copy_Text", "Copy" ),
-			LOCTEXT( "Copy_ToolTip", "Copies the selected prims" )
+			LOCTEXT("Copy_Text", "Copy"),
+			LOCTEXT("Copy_ToolTip", "Copies the selected prims")
 		);
 
 		PrimOptions.AddMenuEntry(
 			FGenericCommands::Get().Paste,
 			NAME_None,
-			LOCTEXT( "Paste_Text", "Paste" ),
-			LOCTEXT( "Paste_ToolTip", "Pastes a flattened representation of the cut/copied prims as children of this prim, on the current edit target" )
+			LOCTEXT("Paste_Text", "Paste"),
+			LOCTEXT("Paste_ToolTip", "Pastes a flattened representation of the cut/copied prims as children of this prim, on the current edit target")
 		);
 
 		const bool bInOpenSubMenuOnClick = false;
 		PrimOptions.AddSubMenu(
-			LOCTEXT( "Duplicate_Text", "Duplicate..." ),
+			LOCTEXT("Duplicate_Text", "Duplicate..."),
 			FText::GetEmpty(),
-			FNewMenuDelegate::CreateSP( this, &SUsdStageTreeView::FillDuplicateSubmenu ),
+			FNewMenuDelegate::CreateSP(this, &SUsdStageTreeView::FillDuplicateSubmenu),
 			bInOpenSubMenuOnClick,
-			FSlateIcon( FAppStyle::GetAppStyleSetName(), "Icons.Duplicate" )
+			FSlateIcon(FAppStyle::GetAppStyleSetName(), "Icons.Duplicate")
 		);
 
 		PrimOptions.AddMenuEntry(
 			FGenericCommands::Get().Delete,
 			NAME_None,
-			LOCTEXT( "Delete_Text", "Delete" ),
-			TAttribute<FText>::Create( TAttribute<FText>::FGetter::CreateLambda( [this]()
-			{
-				return DoesPrimHaveSpecOnLocalLayerStack()
-					? LOCTEXT( "Delete_ToolTip", "Deletes the selected prim's specs from the local layer stack" )
-					: UE::USDStageTreeView::Private::NoSpecOnLocalLayerStack;
-			}))
+			LOCTEXT("Delete_Text", "Delete"),
+			TAttribute<FText>::Create(TAttribute<FText>::FGetter::CreateLambda(
+				[this]()
+				{
+					return DoesPrimHaveSpecOnLocalLayerStack()
+							   ? LOCTEXT("Delete_ToolTip", "Deletes the selected prim's specs from the local layer stack")
+							   : UE::USDStageTreeView::Private::NoSpecOnLocalLayerStack;
+				}
+			))
 		);
 
 		PrimOptions.AddMenuEntry(
 			FGenericCommands::Get().Rename,
 			NAME_None,
-			LOCTEXT( "Rename_Text", "Rename" ),
-			TAttribute<FText>::Create( TAttribute<FText>::FGetter::CreateLambda( [this]()
-			{
-				return DoesPrimHaveSpecOnLocalLayerStack()
-					? LOCTEXT( "Rename_ToolTip", "Renames the selected prim's specs on the local layer stack" )
-					: UE::USDStageTreeView::Private::NoSpecOnLocalLayerStack;
-			}))
+			LOCTEXT("Rename_Text", "Rename"),
+			TAttribute<FText>::Create(TAttribute<FText>::FGetter::CreateLambda(
+				[this]()
+				{
+					return DoesPrimHaveSpecOnLocalLayerStack()
+							   ? LOCTEXT("Rename_ToolTip", "Renames the selected prim's specs on the local layer stack")
+							   : UE::USDStageTreeView::Private::NoSpecOnLocalLayerStack;
+				}
+			))
 		);
 	}
 	PrimOptions.EndSection();
 
-	PrimOptions.BeginSection( "Payloads", LOCTEXT("Payloads", "Payloads") );
+	PrimOptions.BeginSection("Payloads", LOCTEXT("Payloads", "Payloads"));
 	{
 		PrimOptions.AddMenuEntry(
 			LOCTEXT("TogglePayloads", "Toggle All Payloads"),
 			LOCTEXT("TogglePayloads_ToolTip", "Toggles all payloads for this prim and its children"),
 			FSlateIcon(),
-			FUIAction(
-				FExecuteAction::CreateSP( this, &SUsdStageTreeView::OnToggleAllPayloads, EPayloadsTrigger::Toggle ),
-				FCanExecuteAction()
-			),
+			FUIAction(FExecuteAction::CreateSP(this, &SUsdStageTreeView::OnToggleAllPayloads, EPayloadsTrigger::Toggle), FCanExecuteAction()),
 			NAME_None,
 			EUserInterfaceActionType::Button
 		);
@@ -757,10 +771,7 @@ TSharedPtr< SWidget > SUsdStageTreeView::ConstructPrimContextMenu()
 			LOCTEXT("LoadPayloads", "Load All Payloads"),
 			LOCTEXT("LoadPayloads_ToolTip", "Loads all payloads for this prim and its children"),
 			FSlateIcon(),
-			FUIAction(
-				FExecuteAction::CreateSP( this, &SUsdStageTreeView::OnToggleAllPayloads, EPayloadsTrigger::Load ),
-				FCanExecuteAction()
-			),
+			FUIAction(FExecuteAction::CreateSP(this, &SUsdStageTreeView::OnToggleAllPayloads, EPayloadsTrigger::Load), FCanExecuteAction()),
 			NAME_None,
 			EUserInterfaceActionType::Button
 		);
@@ -769,25 +780,27 @@ TSharedPtr< SWidget > SUsdStageTreeView::ConstructPrimContextMenu()
 			LOCTEXT("UnloadPayloads", "Unload All Payloads"),
 			LOCTEXT("UnloadPayloads_ToolTip", "Unloads all payloads for this prim and its children"),
 			FSlateIcon(),
-			FUIAction(
-				FExecuteAction::CreateSP( this, &SUsdStageTreeView::OnToggleAllPayloads, EPayloadsTrigger::Unload ),
-				FCanExecuteAction()
-			),
+			FUIAction(FExecuteAction::CreateSP(this, &SUsdStageTreeView::OnToggleAllPayloads, EPayloadsTrigger::Unload), FCanExecuteAction()),
 			NAME_None,
 			EUserInterfaceActionType::Button
 		);
 	}
 	PrimOptions.EndSection();
 
-	PrimOptions.BeginSection( "Composition", LOCTEXT("Composition", "Composition") );
+	PrimOptions.BeginSection("Composition", LOCTEXT("Composition", "Composition"));
 	{
 		PrimOptions.AddMenuEntry(
 			LOCTEXT("AddReference", "Add Reference"),
 			LOCTEXT("AddReference_ToolTip", "Adds a reference for this prim"),
 			FSlateIcon(),
 			FUIAction(
-				FExecuteAction::CreateSP( this, &SUsdStageTreeView::OnAddReference ),
-				FCanExecuteAction::CreateLambda( [this]() { return SUsdStageTreeView::DoesPrimExistOnEditTarget() && GetSelectedItems().Num() == 1; } )
+				FExecuteAction::CreateSP(this, &SUsdStageTreeView::OnAddReference),
+				FCanExecuteAction::CreateLambda(
+					[this]()
+					{
+						return SUsdStageTreeView::DoesPrimExistOnEditTarget() && GetSelectedItems().Num() == 1;
+					}
+				)
 			),
 			NAME_None,
 			EUserInterfaceActionType::Button
@@ -798,32 +811,37 @@ TSharedPtr< SWidget > SUsdStageTreeView::ConstructPrimContextMenu()
 			LOCTEXT("ClearReferences_ToolTip", "Clears the references for this prim"),
 			FSlateIcon(),
 			FUIAction(
-				FExecuteAction::CreateSP( this, &SUsdStageTreeView::OnClearReferences ),
-				FCanExecuteAction::CreateSP( this, &SUsdStageTreeView::DoesPrimExistOnEditTarget )
+				FExecuteAction::CreateSP(this, &SUsdStageTreeView::OnClearReferences),
+				FCanExecuteAction::CreateSP(this, &SUsdStageTreeView::DoesPrimExistOnEditTarget)
 			),
 			NAME_None,
 			EUserInterfaceActionType::Button
 		);
 
 		PrimOptions.AddMenuEntry(
-			LOCTEXT( "AddPayload", "Add Payload" ),
-			LOCTEXT( "AddPayload_ToolTip", "Adds a payload for this prim" ),
+			LOCTEXT("AddPayload", "Add Payload"),
+			LOCTEXT("AddPayload_ToolTip", "Adds a payload for this prim"),
 			FSlateIcon(),
 			FUIAction(
-				FExecuteAction::CreateSP( this, &SUsdStageTreeView::OnAddPayload ),
-				FCanExecuteAction::CreateLambda( [this]() { return SUsdStageTreeView::DoesPrimExistOnEditTarget() && GetSelectedItems().Num() == 1; } )
+				FExecuteAction::CreateSP(this, &SUsdStageTreeView::OnAddPayload),
+				FCanExecuteAction::CreateLambda(
+					[this]()
+					{
+						return SUsdStageTreeView::DoesPrimExistOnEditTarget() && GetSelectedItems().Num() == 1;
+					}
+				)
 			),
 			NAME_None,
 			EUserInterfaceActionType::Button
 		);
 
 		PrimOptions.AddMenuEntry(
-			LOCTEXT( "ClearPayloads", "Clear Payloads" ),
-			LOCTEXT( "ClearPayloads_ToolTip", "Clears the payloads for this prim" ),
+			LOCTEXT("ClearPayloads", "Clear Payloads"),
+			LOCTEXT("ClearPayloads_ToolTip", "Clears the payloads for this prim"),
 			FSlateIcon(),
 			FUIAction(
-				FExecuteAction::CreateSP( this, &SUsdStageTreeView::OnClearPayloads ),
-				FCanExecuteAction::CreateSP( this, &SUsdStageTreeView::DoesPrimExistOnEditTarget )
+				FExecuteAction::CreateSP(this, &SUsdStageTreeView::OnClearPayloads),
+				FCanExecuteAction::CreateSP(this, &SUsdStageTreeView::DoesPrimExistOnEditTarget)
 			),
 			NAME_None,
 			EUserInterfaceActionType::Button
@@ -831,22 +849,18 @@ TSharedPtr< SWidget > SUsdStageTreeView::ConstructPrimContextMenu()
 	}
 	PrimOptions.EndSection();
 
-	PrimOptions.BeginSection( "Integrations", LOCTEXT( "Integrations", "Integrations" ) );
+	PrimOptions.BeginSection("Integrations", LOCTEXT("Integrations", "Integrations"));
 	{
 		PrimOptions.AddMenuEntry(
-			LOCTEXT( "SetUpLiveLink", "Set up Live Link" ),
-			LOCTEXT( "SetUpLiveLink_ToolTip", "Sets up the generated component for Live Link and store the connection details to the USD Stage" ),
+			LOCTEXT("SetUpLiveLink", "Set up Live Link"),
+			LOCTEXT("SetUpLiveLink_ToolTip", "Sets up the generated component for Live Link and store the connection details to the USD Stage"),
 			FSlateIcon(),
 			FUIAction(
-				FExecuteAction::CreateSP(
-					this,
-					&SUsdStageTreeView::OnApplySchema,
-					FName{ *UsdToUnreal::ConvertToken( UnrealIdentifiers::LiveLinkAPI ) }
-				),
+				FExecuteAction::CreateSP(this, &SUsdStageTreeView::OnApplySchema, FName{*UsdToUnreal::ConvertToken(UnrealIdentifiers::LiveLinkAPI)}),
 				FCanExecuteAction::CreateSP(
 					this,
 					&SUsdStageTreeView::CanApplySchema,
-					FName{ *UsdToUnreal::ConvertToken( UnrealIdentifiers::LiveLinkAPI ) }
+					FName{*UsdToUnreal::ConvertToken(UnrealIdentifiers::LiveLinkAPI)}
 				)
 			),
 			NAME_None,
@@ -854,19 +868,18 @@ TSharedPtr< SWidget > SUsdStageTreeView::ConstructPrimContextMenu()
 		);
 
 		PrimOptions.AddMenuEntry(
-			LOCTEXT( "RemoveLiveLink", "Remove Live Link" ),
-			LOCTEXT( "RemoveLiveLink_ToolTip", "Reverses the Live Link configuration on the component and removes the connection details from the USD Stage" ),
+			LOCTEXT("RemoveLiveLink", "Remove Live Link"),
+			LOCTEXT(
+				"RemoveLiveLink_ToolTip",
+				"Reverses the Live Link configuration on the component and removes the connection details from the USD Stage"
+			),
 			FSlateIcon(),
 			FUIAction(
-				FExecuteAction::CreateSP(
-					this,
-					&SUsdStageTreeView::OnRemoveSchema,
-					FName{ *UsdToUnreal::ConvertToken( UnrealIdentifiers::LiveLinkAPI ) }
-				),
+				FExecuteAction::CreateSP(this, &SUsdStageTreeView::OnRemoveSchema, FName{*UsdToUnreal::ConvertToken(UnrealIdentifiers::LiveLinkAPI)}),
 				FCanExecuteAction::CreateSP(
 					this,
 					&SUsdStageTreeView::CanRemoveSchema,
-					FName{ *UsdToUnreal::ConvertToken( UnrealIdentifiers::LiveLinkAPI ) }
+					FName{*UsdToUnreal::ConvertToken(UnrealIdentifiers::LiveLinkAPI)}
 				)
 			),
 			NAME_None,
@@ -874,19 +887,22 @@ TSharedPtr< SWidget > SUsdStageTreeView::ConstructPrimContextMenu()
 		);
 
 		PrimOptions.AddMenuEntry(
-			LOCTEXT( "SetUpControlRig", "Set up Control Rig" ),
-			LOCTEXT( "SetUpControlRig_ToolTip", "Sets up the generated component for Control Rig integration and store the connection details to the USD Stage" ),
+			LOCTEXT("SetUpControlRig", "Set up Control Rig"),
+			LOCTEXT(
+				"SetUpControlRig_ToolTip",
+				"Sets up the generated component for Control Rig integration and store the connection details to the USD Stage"
+			),
 			FSlateIcon(),
 			FUIAction(
 				FExecuteAction::CreateSP(
 					this,
 					&SUsdStageTreeView::OnApplySchema,
-					FName{ *UsdToUnreal::ConvertToken( UnrealIdentifiers::ControlRigAPI ) }
+					FName{*UsdToUnreal::ConvertToken(UnrealIdentifiers::ControlRigAPI)}
 				),
 				FCanExecuteAction::CreateSP(
 					this,
 					&SUsdStageTreeView::CanApplySchema,
-					FName{ *UsdToUnreal::ConvertToken( UnrealIdentifiers::ControlRigAPI ) }
+					FName{*UsdToUnreal::ConvertToken(UnrealIdentifiers::ControlRigAPI)}
 				)
 			),
 			NAME_None,
@@ -894,19 +910,22 @@ TSharedPtr< SWidget > SUsdStageTreeView::ConstructPrimContextMenu()
 		);
 
 		PrimOptions.AddMenuEntry(
-			LOCTEXT( "RemoveControlRig", "Remove Control Rig" ),
-			LOCTEXT( "RemoveControlRig_ToolTip", "Reverses the Control Rig configuration on the component and removes the connection details from the USD Stage" ),
+			LOCTEXT("RemoveControlRig", "Remove Control Rig"),
+			LOCTEXT(
+				"RemoveControlRig_ToolTip",
+				"Reverses the Control Rig configuration on the component and removes the connection details from the USD Stage"
+			),
 			FSlateIcon(),
 			FUIAction(
 				FExecuteAction::CreateSP(
 					this,
 					&SUsdStageTreeView::OnRemoveSchema,
-					FName{ *UsdToUnreal::ConvertToken( UnrealIdentifiers::ControlRigAPI ) }
+					FName{*UsdToUnreal::ConvertToken(UnrealIdentifiers::ControlRigAPI)}
 				),
 				FCanExecuteAction::CreateSP(
 					this,
 					&SUsdStageTreeView::CanRemoveSchema,
-					FName{ *UsdToUnreal::ConvertToken( UnrealIdentifiers::ControlRigAPI ) }
+					FName{*UsdToUnreal::ConvertToken(UnrealIdentifiers::ControlRigAPI)}
 				)
 			),
 			NAME_None,
@@ -914,40 +933,24 @@ TSharedPtr< SWidget > SUsdStageTreeView::ConstructPrimContextMenu()
 		);
 
 		PrimOptions.AddMenuEntry(
-			LOCTEXT( "ApplyGroomSchema", "Apply Groom schema" ),
-			LOCTEXT( "ApplyGroomSchema_ToolTip", "Applies the Groom schema to interpret the prim and its children as a groom" ),
+			LOCTEXT("ApplyGroomSchema", "Apply Groom schema"),
+			LOCTEXT("ApplyGroomSchema_ToolTip", "Applies the Groom schema to interpret the prim and its children as a groom"),
 			FSlateIcon(),
 			FUIAction(
-				FExecuteAction::CreateSP(
-					this,
-					&SUsdStageTreeView::OnApplySchema,
-					FName{ *UsdToUnreal::ConvertToken( UnrealIdentifiers::GroomAPI ) }
-				),
-				FCanExecuteAction::CreateSP(
-					this,
-					&SUsdStageTreeView::CanApplySchema,
-					FName{ *UsdToUnreal::ConvertToken( UnrealIdentifiers::GroomAPI ) }
-				)
+				FExecuteAction::CreateSP(this, &SUsdStageTreeView::OnApplySchema, FName{*UsdToUnreal::ConvertToken(UnrealIdentifiers::GroomAPI)}),
+				FCanExecuteAction::CreateSP(this, &SUsdStageTreeView::CanApplySchema, FName{*UsdToUnreal::ConvertToken(UnrealIdentifiers::GroomAPI)})
 			),
 			NAME_None,
 			EUserInterfaceActionType::Button
 		);
 
 		PrimOptions.AddMenuEntry(
-			LOCTEXT( "RemoveGroomSchema", "Remove Groom schema" ),
-			LOCTEXT( "RemoveGroomSchema_ToolTip", "Removes the Groom schema from the prim to stop interpreting it as a groom" ),
+			LOCTEXT("RemoveGroomSchema", "Remove Groom schema"),
+			LOCTEXT("RemoveGroomSchema_ToolTip", "Removes the Groom schema from the prim to stop interpreting it as a groom"),
 			FSlateIcon(),
 			FUIAction(
-				FExecuteAction::CreateSP(
-					this,
-					&SUsdStageTreeView::OnRemoveSchema,
-					FName{ *UsdToUnreal::ConvertToken( UnrealIdentifiers::GroomAPI ) }
-				),
-				FCanExecuteAction::CreateSP(
-					this,
-					&SUsdStageTreeView::CanRemoveSchema,
-					FName{ *UsdToUnreal::ConvertToken( UnrealIdentifiers::GroomAPI ) }
-				)
+				FExecuteAction::CreateSP(this, &SUsdStageTreeView::OnRemoveSchema, FName{*UsdToUnreal::ConvertToken(UnrealIdentifiers::GroomAPI)}),
+				FCanExecuteAction::CreateSP(this, &SUsdStageTreeView::CanRemoveSchema, FName{*UsdToUnreal::ConvertToken(UnrealIdentifiers::GroomAPI)})
 			),
 			NAME_None,
 			EUserInterfaceActionType::Button
@@ -995,7 +998,10 @@ TSharedPtr< SWidget > SUsdStageTreeView::ConstructPrimContextMenu()
 
 		PrimOptions.AddMenuEntry(
 			LOCTEXT("ApplyPhysicsMeshCollisionchema", "Apply PhysicsMeshCollision schema"),
-			LOCTEXT("ApplyPhysicsMeshCollisionSchema_ToolTip", "Applies the PhysicsMeshCollision schema to control how the mesh is made into a collider"),
+			LOCTEXT(
+				"ApplyPhysicsMeshCollisionSchema_ToolTip",
+				"Applies the PhysicsMeshCollision schema to control how the mesh is made into a collider"
+			),
 			FSlateIcon(),
 			FUIAction(
 				FExecuteAction::CreateSP(
@@ -1032,7 +1038,6 @@ TSharedPtr< SWidget > SUsdStageTreeView::ConstructPrimContextMenu()
 			NAME_None,
 			EUserInterfaceActionType::Button
 		);
-
 	}
 	PrimOptions.EndSection();
 
@@ -1043,28 +1048,28 @@ TSharedPtr< SWidget > SUsdStageTreeView::ConstructPrimContextMenu()
 
 void SUsdStageTreeView::OnAddChildPrim()
 {
-	TArray< FUsdPrimViewModelRef > MySelectedItems = GetSelectedItems();
+	TArray<FUsdPrimViewModelRef> MySelectedItems = GetSelectedItems();
 
 	// Add a new child prim
 	if (MySelectedItems.Num() > 0)
 	{
-		for ( FUsdPrimViewModelRef SelectedItem : MySelectedItems )
+		for (FUsdPrimViewModelRef SelectedItem : MySelectedItems)
 		{
-			FUsdPrimViewModelRef TreeItem = MakeShared< FUsdPrimViewModel >( &SelectedItem.Get(), SelectedItem->UsdStage );
-			SelectedItem->Children.Add( TreeItem );
+			FUsdPrimViewModelRef TreeItem = MakeShared<FUsdPrimViewModel>(&SelectedItem.Get(), SelectedItem->UsdStage);
+			SelectedItem->Children.Add(TreeItem);
 
 			PendingRenameItem = TreeItem;
-			ScrollItemIntoView( TreeItem );
+			ScrollItemIntoView(TreeItem);
 		}
 	}
 	// Add a new top-level prim (direct child of the pseudo-root prim)
 	else
 	{
-		FUsdPrimViewModelRef TreeItem = MakeShared< FUsdPrimViewModel >( nullptr, UsdStage );
-		RootItems.Add( TreeItem );
+		FUsdPrimViewModelRef TreeItem = MakeShared<FUsdPrimViewModel>(nullptr, UsdStage);
+		RootItems.Add(TreeItem);
 
 		PendingRenameItem = TreeItem;
-		ScrollItemIntoView( TreeItem );
+		ScrollItemIntoView(TreeItem);
 	}
 
 	RequestTreeRefresh();
@@ -1081,45 +1086,45 @@ void SUsdStageTreeView::OnCutPrim()
 		}
 	}
 
-	FScopedTransaction Transaction( LOCTEXT( "CutPrimTransaction", "Cut prims" ) );
+	FScopedTransaction Transaction(LOCTEXT("CutPrimTransaction", "Cut prims"));
 
 	UE::FSdfChangeBlock Block;
 
 	TArray<UE::FUsdPrim> Prims;
-	Prims.Reserve( MySelectedItems.Num() );
+	Prims.Reserve(MySelectedItems.Num());
 
-	for ( FUsdPrimViewModelRef SelectedItem : MySelectedItems )
+	for (FUsdPrimViewModelRef SelectedItem : MySelectedItems)
 	{
-		if ( SelectedItem->UsdPrim )
+		if (SelectedItem->UsdPrim)
 		{
-			Prims.Add( SelectedItem->UsdPrim );
+			Prims.Add(SelectedItem->UsdPrim);
 		}
 	}
 
-	UsdUtils::CutPrims( Prims );
+	UsdUtils::CutPrims(Prims);
 }
 
 void SUsdStageTreeView::OnCopyPrim()
 {
-	TArray< FUsdPrimViewModelRef > MySelectedItems = GetSelectedItems();
+	TArray<FUsdPrimViewModelRef> MySelectedItems = GetSelectedItems();
 
 	TArray<UE::FUsdPrim> Prims;
-	Prims.Reserve( MySelectedItems.Num() );
+	Prims.Reserve(MySelectedItems.Num());
 
-	for ( FUsdPrimViewModelRef SelectedItem : MySelectedItems )
+	for (FUsdPrimViewModelRef SelectedItem : MySelectedItems)
 	{
-		if ( SelectedItem->UsdPrim )
+		if (SelectedItem->UsdPrim)
 		{
-			Prims.Add( SelectedItem->UsdPrim );
+			Prims.Add(SelectedItem->UsdPrim);
 		}
 	}
 
-	UsdUtils::CopyPrims( Prims );
+	UsdUtils::CopyPrims(Prims);
 }
 
 void SUsdStageTreeView::OnPastePrim()
 {
-	if ( !UsdStage )
+	if (!UsdStage)
 	{
 		return;
 	}
@@ -1133,41 +1138,41 @@ void SUsdStageTreeView::OnPastePrim()
 		}
 	}
 
-	FScopedTransaction Transaction( LOCTEXT( "PastePrimTransaction", "Paste prims" ) );
+	FScopedTransaction Transaction(LOCTEXT("PastePrimTransaction", "Paste prims"));
 
 	UE::FSdfChangeBlock Block;
 
 	TArray<UE::FUsdPrim> ParentPrims;
 
 	// This happens when right-clicking the background area without selecting any prim
-	if ( MySelectedItems.IsEmpty() )
+	if (MySelectedItems.IsEmpty())
 	{
-		ParentPrims.Add( UsdStage.GetPseudoRoot() );
+		ParentPrims.Add(UsdStage.GetPseudoRoot());
 	}
 	else
 	{
-		ParentPrims.Reserve( MySelectedItems.Num() );
+		ParentPrims.Reserve(MySelectedItems.Num());
 
 		// A bit unusual that we can paste to multiple locations at the same time, but why not?
-		for ( FUsdPrimViewModelRef SelectedItem : MySelectedItems )
+		for (FUsdPrimViewModelRef SelectedItem : MySelectedItems)
 		{
-			ParentPrims.Add( SelectedItem->UsdPrim );
+			ParentPrims.Add(SelectedItem->UsdPrim);
 		}
 	}
 
-	for ( const UE::FUsdPrim& ParentPrim : ParentPrims )
+	for (const UE::FUsdPrim& ParentPrim : ParentPrims)
 	{
 		// Preemptively mark the parent prims as expanded so that we can always see what we pasted
 		const bool bIsExpanded = true;
-		TreeItemExpansionStates.Add( ParentPrim.GetPrimPath().GetString(), bIsExpanded );
+		TreeItemExpansionStates.Add(ParentPrim.GetPrimPath().GetString(), bIsExpanded);
 
-		UsdUtils::PastePrims( ParentPrim );
+		UsdUtils::PastePrims(ParentPrim);
 	}
 }
 
-void SUsdStageTreeView::OnDuplicatePrim( EUsdDuplicateType DuplicateType )
+void SUsdStageTreeView::OnDuplicatePrim(EUsdDuplicateType DuplicateType)
 {
-	if ( !UsdStage )
+	if (!UsdStage)
 	{
 		return;
 	}
@@ -1181,20 +1186,20 @@ void SUsdStageTreeView::OnDuplicatePrim( EUsdDuplicateType DuplicateType )
 		}
 	}
 
-	FScopedTransaction Transaction( LOCTEXT( "DuplicatePrimTransaction", "Duplicate prims" ) );
+	FScopedTransaction Transaction(LOCTEXT("DuplicatePrimTransaction", "Duplicate prims"));
 
 	TArray<UE::FUsdPrim> Prims;
-	Prims.Reserve( MySelectedItems.Num() );
+	Prims.Reserve(MySelectedItems.Num());
 
-	for ( FUsdPrimViewModelRef SelectedItem : MySelectedItems )
+	for (FUsdPrimViewModelRef SelectedItem : MySelectedItems)
 	{
-		if ( SelectedItem->UsdPrim )
+		if (SelectedItem->UsdPrim)
 		{
-			Prims.Add( SelectedItem->UsdPrim );
+			Prims.Add(SelectedItem->UsdPrim);
 		}
 	}
 
-	UsdUtils::DuplicatePrims( Prims, DuplicateType, UsdStage.GetEditTarget() );
+	UsdUtils::DuplicatePrims(Prims, DuplicateType, UsdStage.GetEditTarget());
 }
 
 void SUsdStageTreeView::OnDeletePrim()
@@ -1208,60 +1213,56 @@ void SUsdStageTreeView::OnDeletePrim()
 		}
 	}
 
-	FScopedTransaction Transaction( LOCTEXT( "DeletePrimTransaction", "Delete prims" ) );
+	FScopedTransaction Transaction(LOCTEXT("DeletePrimTransaction", "Delete prims"));
 
 	UE::FSdfChangeBlock Block;
 
-	for ( FUsdPrimViewModelRef SelectedItem : MySelectedItems )
+	for (FUsdPrimViewModelRef SelectedItem : MySelectedItems)
 	{
-		UsdUtils::RemoveAllLocalPrimSpecs( SelectedItem->UsdPrim );
+		UsdUtils::RemoveAllLocalPrimSpecs(SelectedItem->UsdPrim);
 	}
 }
 
 void SUsdStageTreeView::OnRenamePrim()
 {
-	TArray< FUsdPrimViewModelRef > MySelectedItems = GetSelectedItems();
+	TArray<FUsdPrimViewModelRef> MySelectedItems = GetSelectedItems();
 
-	if ( MySelectedItems.Num() > 0 )
+	if (MySelectedItems.Num() > 0)
 	{
-		FUsdPrimViewModelRef TreeItem = MySelectedItems[ 0 ];
+		FUsdPrimViewModelRef TreeItem = MySelectedItems[0];
 
 		TreeItem->bIsRenamingExistingPrim = true;
 		PendingRenameItem = TreeItem;
-		RequestScrollIntoView( TreeItem );
+		RequestScrollIntoView(TreeItem);
 	}
 }
 
 void SUsdStageTreeView::OnAddReference()
 {
-	if ( !UsdStage || !UsdStage.IsEditTargetValid() )
+	if (!UsdStage || !UsdStage.IsEditTargetValid())
 	{
 		return;
 	}
 
-	TStrongObjectPtr<UUsdReferenceOptions> Options = TStrongObjectPtr<UUsdReferenceOptions>{ NewObject<UUsdReferenceOptions>() };
+	TStrongObjectPtr<UUsdReferenceOptions> Options = TStrongObjectPtr<UUsdReferenceOptions>{NewObject<UUsdReferenceOptions>()};
 	UUsdReferenceOptions* OptionsPtr = Options.Get();
-	if ( !OptionsPtr )
+	if (!OptionsPtr)
 	{
 		return;
 	}
 
-	bool bContinue = SUsdOptionsWindow::ShowOptions(
-		*OptionsPtr,
-		LOCTEXT( "AddReferenceTitle", "Add reference" ),
-		LOCTEXT( "AddReferenceAccept", "OK" )
-	);
-	if ( !bContinue )
+	bool bContinue = SUsdOptionsWindow::ShowOptions(*OptionsPtr, LOCTEXT("AddReferenceTitle", "Add reference"), LOCTEXT("AddReferenceAccept", "OK"));
+	if (!bContinue)
 	{
 		return;
 	}
 
-	TArray< FUsdPrimViewModelRef > MySelectedItems = GetSelectedItems();
-	if ( MySelectedItems.Num() != 1 )
+	TArray<FUsdPrimViewModelRef> MySelectedItems = GetSelectedItems();
+	if (MySelectedItems.Num() != 1)
 	{
 		return;
 	}
-	UE::FUsdPrim Referencer = MySelectedItems[ 0 ]->UsdPrim;
+	UE::FUsdPrim Referencer = MySelectedItems[0]->UsdPrim;
 	if (UsdUtils::NotifyIfInstanceProxy(Referencer))
 	{
 		return;
@@ -1269,10 +1270,9 @@ void SUsdStageTreeView::OnAddReference()
 
 	// This transaction is important as adding a reference may trigger the creation of new unreal assets, which need to be
 	// destroyed if we spam undo afterwards. Undoing won't remove the actual reference from the stage yet though, sadly...
-	FScopedTransaction Transaction( FText::Format(
-		LOCTEXT( "AddReferenceTransaction", "Add reference from prim '{0}'" ),
-		FText::FromString( Referencer.GetPrimPath().GetString() )
-	) );
+	FScopedTransaction Transaction(
+		FText::Format(LOCTEXT("AddReferenceTransaction", "Add reference from prim '{0}'"), FText::FromString(Referencer.GetPrimPath().GetString()))
+	);
 
 	// AddReference expects absolute file paths, so let's try ensuring that
 	FString ReferencedLayerPath = OptionsPtr->TargetFile.FilePath;
@@ -1292,7 +1292,7 @@ void SUsdStageTreeView::OnAddReference()
 	UsdUtils::AddReference(
 		Referencer,
 		*ReferencedLayerPath,
-		OptionsPtr->bUseDefaultPrim ? UE::FSdfPath{} : UE::FSdfPath{ *OptionsPtr->TargetPrimPath },
+		OptionsPtr->bUseDefaultPrim ? UE::FSdfPath{} : UE::FSdfPath{*OptionsPtr->TargetPrimPath},
 		OptionsPtr->TimeCodeOffset,
 		OptionsPtr->TimeCodeScale
 	);
@@ -1309,9 +1309,9 @@ void SUsdStageTreeView::OnClearReferences()
 		}
 	}
 
-	FScopedTransaction Transaction( LOCTEXT( "ClearReferenceTransaction", "Clear references to USD layers" ) );
+	FScopedTransaction Transaction(LOCTEXT("ClearReferenceTransaction", "Clear references to USD layers"));
 
-	for ( FUsdPrimViewModelRef SelectedItem : MySelectedItems )
+	for (FUsdPrimViewModelRef SelectedItem : MySelectedItems)
 	{
 		SelectedItem->ClearReferences();
 	}
@@ -1319,34 +1319,30 @@ void SUsdStageTreeView::OnClearReferences()
 
 void SUsdStageTreeView::OnAddPayload()
 {
-	if ( !UsdStage || !UsdStage.IsEditTargetValid() )
+	if (!UsdStage || !UsdStage.IsEditTargetValid())
 	{
 		return;
 	}
 
-	TStrongObjectPtr<UUsdReferenceOptions> Options = TStrongObjectPtr<UUsdReferenceOptions>{ NewObject<UUsdReferenceOptions>() };
+	TStrongObjectPtr<UUsdReferenceOptions> Options = TStrongObjectPtr<UUsdReferenceOptions>{NewObject<UUsdReferenceOptions>()};
 	UUsdReferenceOptions* OptionsPtr = Options.Get();
-	if ( !OptionsPtr )
+	if (!OptionsPtr)
 	{
 		return;
 	}
 
-	bool bContinue = SUsdOptionsWindow::ShowOptions(
-		*OptionsPtr,
-		LOCTEXT( "AddPayloadTitle", "Add payload" ),
-		LOCTEXT( "AddPayloadAccept", "OK" )
-	);
-	if ( !bContinue )
+	bool bContinue = SUsdOptionsWindow::ShowOptions(*OptionsPtr, LOCTEXT("AddPayloadTitle", "Add payload"), LOCTEXT("AddPayloadAccept", "OK"));
+	if (!bContinue)
 	{
 		return;
 	}
 
-	TArray< FUsdPrimViewModelRef > MySelectedItems = GetSelectedItems();
-	if ( MySelectedItems.Num() != 1 )
+	TArray<FUsdPrimViewModelRef> MySelectedItems = GetSelectedItems();
+	if (MySelectedItems.Num() != 1)
 	{
 		return;
 	}
-	UE::FUsdPrim Referencer = MySelectedItems[ 0 ]->UsdPrim;
+	UE::FUsdPrim Referencer = MySelectedItems[0]->UsdPrim;
 	if (UsdUtils::NotifyIfInstanceProxy(Referencer))
 	{
 		return;
@@ -1354,10 +1350,9 @@ void SUsdStageTreeView::OnAddPayload()
 
 	// This transaction is important as adding a payload may trigger the creation of new unreal assets, which need to be
 	// destroyed if we spam undo afterwards. Undoing won't remove the actual payload from the stage yet though, sadly...
-	FScopedTransaction Transaction( FText::Format(
-		LOCTEXT( "AddPayloadTransaction", "Add payload from prim '{0}'" ),
-		FText::FromString( Referencer.GetPrimPath().GetString() )
-	) );
+	FScopedTransaction Transaction(
+		FText::Format(LOCTEXT("AddPayloadTransaction", "Add payload from prim '{0}'"), FText::FromString(Referencer.GetPrimPath().GetString()))
+	);
 
 	// AddPayload expects absolute file paths too, so let's try ensuring that
 	FString ReferencedLayerPath = OptionsPtr->TargetFile.FilePath;
@@ -1377,7 +1372,7 @@ void SUsdStageTreeView::OnAddPayload()
 	UsdUtils::AddPayload(
 		Referencer,
 		*ReferencedLayerPath,
-		OptionsPtr->bUseDefaultPrim ? UE::FSdfPath{} : UE::FSdfPath{ *OptionsPtr->TargetPrimPath },
+		OptionsPtr->bUseDefaultPrim ? UE::FSdfPath{} : UE::FSdfPath{*OptionsPtr->TargetPrimPath},
 		OptionsPtr->TimeCodeOffset,
 		OptionsPtr->TimeCodeScale
 	);
@@ -1394,15 +1389,15 @@ void SUsdStageTreeView::OnClearPayloads()
 		}
 	}
 
-	FScopedTransaction Transaction( LOCTEXT( "ClearPayloadTransaction", "Clear payloads to USD layers" ) );
+	FScopedTransaction Transaction(LOCTEXT("ClearPayloadTransaction", "Clear payloads to USD layers"));
 
-	for ( FUsdPrimViewModelRef SelectedItem : MySelectedItems )
+	for (FUsdPrimViewModelRef SelectedItem : MySelectedItems)
 	{
 		SelectedItem->ClearPayloads();
 	}
 }
 
-void SUsdStageTreeView::OnApplySchema( FName SchemaName )
+void SUsdStageTreeView::OnApplySchema(FName SchemaName)
 {
 	TArray<FUsdPrimViewModelRef> MySelectedItems = GetSelectedItems();
 	for (FUsdPrimViewModelRef SelectedItem : MySelectedItems)
@@ -1413,20 +1408,19 @@ void SUsdStageTreeView::OnApplySchema( FName SchemaName )
 		}
 	}
 
-	FScopedTransaction Transaction( FText::Format(
-		LOCTEXT( "ApplySchemaTransaction", "Apply the '{0}' schema onto selected prims" ),
-		FText::FromName( SchemaName )
-	));
+	FScopedTransaction Transaction(
+		FText::Format(LOCTEXT("ApplySchemaTransaction", "Apply the '{0}' schema onto selected prims"), FText::FromName(SchemaName))
+	);
 
 	UE::FSdfChangeBlock Block;
 
-	for ( FUsdPrimViewModelRef SelectedItem : MySelectedItems )
+	for (FUsdPrimViewModelRef SelectedItem : MySelectedItems)
 	{
-		SelectedItem->ApplySchema( SchemaName );
+		SelectedItem->ApplySchema(SchemaName);
 	}
 }
 
-void SUsdStageTreeView::OnRemoveSchema( FName SchemaName )
+void SUsdStageTreeView::OnRemoveSchema(FName SchemaName)
 {
 	TArray<FUsdPrimViewModelRef> MySelectedItems = GetSelectedItems();
 	for (FUsdPrimViewModelRef SelectedItem : MySelectedItems)
@@ -1437,31 +1431,30 @@ void SUsdStageTreeView::OnRemoveSchema( FName SchemaName )
 		}
 	}
 
-	FScopedTransaction Transaction( FText::Format(
-		LOCTEXT( "RemoveSchemaTransaction", "Remove the '{0}' schema from selected prims" ),
-		FText::FromName( SchemaName )
-	) );
+	FScopedTransaction Transaction(
+		FText::Format(LOCTEXT("RemoveSchemaTransaction", "Remove the '{0}' schema from selected prims"), FText::FromName(SchemaName))
+	);
 
 	UE::FSdfChangeBlock Block;
 
-	for ( FUsdPrimViewModelRef SelectedItem : MySelectedItems )
+	for (FUsdPrimViewModelRef SelectedItem : MySelectedItems)
 	{
-		SelectedItem->RemoveSchema( SchemaName );
+		SelectedItem->RemoveSchema(SchemaName);
 	}
 }
 
-bool SUsdStageTreeView::CanApplySchema( FName SchemaName )
+bool SUsdStageTreeView::CanApplySchema(FName SchemaName)
 {
-	if ( !UsdStage || !UsdStage.IsEditTargetValid() )
+	if (!UsdStage || !UsdStage.IsEditTargetValid())
 	{
 		return false;
 	}
 
-	TArray< FUsdPrimViewModelRef > MySelectedItems = GetSelectedItems();
+	TArray<FUsdPrimViewModelRef> MySelectedItems = GetSelectedItems();
 
-	for ( FUsdPrimViewModelRef SelectedItem : MySelectedItems )
+	for (FUsdPrimViewModelRef SelectedItem : MySelectedItems)
 	{
-		if ( SelectedItem->CanApplySchema( SchemaName ) )
+		if (SelectedItem->CanApplySchema(SchemaName))
 		{
 			return true;
 		}
@@ -1470,18 +1463,18 @@ bool SUsdStageTreeView::CanApplySchema( FName SchemaName )
 	return false;
 }
 
-bool SUsdStageTreeView::CanRemoveSchema( FName SchemaName )
+bool SUsdStageTreeView::CanRemoveSchema(FName SchemaName)
 {
-	if ( !UsdStage || !UsdStage.IsEditTargetValid() )
+	if (!UsdStage || !UsdStage.IsEditTargetValid())
 	{
 		return false;
 	}
 
-	TArray< FUsdPrimViewModelRef > MySelectedItems = GetSelectedItems();
+	TArray<FUsdPrimViewModelRef> MySelectedItems = GetSelectedItems();
 
-	for ( FUsdPrimViewModelRef SelectedItem : MySelectedItems )
+	for (FUsdPrimViewModelRef SelectedItem : MySelectedItems)
 	{
-		if ( SelectedItem->CanRemoveSchema( SchemaName ) )
+		if (SelectedItem->CanRemoveSchema(SchemaName))
 		{
 			return true;
 		}
@@ -1492,15 +1485,15 @@ bool SUsdStageTreeView::CanRemoveSchema( FName SchemaName )
 
 bool SUsdStageTreeView::CanAddChildPrim() const
 {
-	if ( !UsdStage )
+	if (!UsdStage)
 	{
 		return false;
 	}
 
-	TArray< FUsdPrimViewModelRef > MySelectedItems = GetSelectedItems();
+	TArray<FUsdPrimViewModelRef> MySelectedItems = GetSelectedItems();
 
 	// Allow adding a new top-level prim
-	if ( MySelectedItems.IsEmpty() )
+	if (MySelectedItems.IsEmpty())
 	{
 		return true;
 	}
@@ -1508,13 +1501,13 @@ bool SUsdStageTreeView::CanAddChildPrim() const
 	{
 		// We use the "rename" text input workflow to specify the target name,
 		// so this doesn't work very well for multiple prims yet
-		if ( MySelectedItems.Num() > 1 )
+		if (MySelectedItems.Num() > 1)
 		{
 			return false;
 		}
 
 		// If we have something selected it must be valid
-		if ( !MySelectedItems[ 0 ]->UsdPrim.IsValid() )
+		if (!MySelectedItems[0]->UsdPrim.IsValid())
 		{
 			return false;
 		}
@@ -1525,7 +1518,7 @@ bool SUsdStageTreeView::CanAddChildPrim() const
 
 bool SUsdStageTreeView::CanPastePrim() const
 {
-	if ( !UsdStage )
+	if (!UsdStage)
 	{
 		return false;
 	}
@@ -1535,16 +1528,16 @@ bool SUsdStageTreeView::CanPastePrim() const
 
 bool SUsdStageTreeView::DoesPrimExistOnStage() const
 {
-	if ( !UsdStage || !UsdStage.IsEditTargetValid() )
+	if (!UsdStage || !UsdStage.IsEditTargetValid())
 	{
 		return false;
 	}
 
-	TArray< FUsdPrimViewModelRef > MySelectedItems = GetSelectedItems();
+	TArray<FUsdPrimViewModelRef> MySelectedItems = GetSelectedItems();
 
-	for ( FUsdPrimViewModelRef SelectedItem : MySelectedItems )
+	for (FUsdPrimViewModelRef SelectedItem : MySelectedItems)
 	{
-		if ( !SelectedItem->UsdPrim.IsPseudoRoot() && SelectedItem->UsdPrim.IsValid() )
+		if (!SelectedItem->UsdPrim.IsPseudoRoot() && SelectedItem->UsdPrim.IsValid())
 		{
 			return true;
 		}
@@ -1555,17 +1548,17 @@ bool SUsdStageTreeView::DoesPrimExistOnStage() const
 
 bool SUsdStageTreeView::DoesPrimExistOnEditTarget() const
 {
-	if ( !UsdStage || !UsdStage.IsEditTargetValid() )
+	if (!UsdStage || !UsdStage.IsEditTargetValid())
 	{
 		return false;
 	}
 
-	TArray< FUsdPrimViewModelRef > MySelectedItems = GetSelectedItems();
+	TArray<FUsdPrimViewModelRef> MySelectedItems = GetSelectedItems();
 
-	for ( FUsdPrimViewModelRef SelectedItem : MySelectedItems )
+	for (FUsdPrimViewModelRef SelectedItem : MySelectedItems)
 	{
-		UE::FSdfPath SpecPath = UsdUtils::GetPrimSpecPathForLayer( SelectedItem->UsdPrim, UsdStage.GetEditTarget() );
-		if ( !SpecPath.IsAbsoluteRootPath() && !SpecPath.IsEmpty() )
+		UE::FSdfPath SpecPath = UsdUtils::GetPrimSpecPathForLayer(SelectedItem->UsdPrim, UsdStage.GetEditTarget());
+		if (!SpecPath.IsAbsoluteRootPath() && !SpecPath.IsEmpty())
 		{
 			return true;
 		}
@@ -1576,11 +1569,11 @@ bool SUsdStageTreeView::DoesPrimExistOnEditTarget() const
 
 bool SUsdStageTreeView::DoesPrimHaveSpecOnLocalLayerStack() const
 {
-	TArray< FUsdPrimViewModelRef > MySelectedItems = GetSelectedItems();
+	TArray<FUsdPrimViewModelRef> MySelectedItems = GetSelectedItems();
 
-	for ( FUsdPrimViewModelRef SelectedItem : MySelectedItems )
+	for (FUsdPrimViewModelRef SelectedItem : MySelectedItems)
 	{
-		if ( SelectedItem->HasSpecsOnLocalLayer() )
+		if (SelectedItem->HasSpecsOnLocalLayer())
 		{
 			return true;
 		}
@@ -1591,48 +1584,48 @@ bool SUsdStageTreeView::DoesPrimHaveSpecOnLocalLayerStack() const
 
 void SUsdStageTreeView::RequestListRefresh()
 {
-	SUsdTreeView< FUsdPrimViewModelRef >::RequestListRefresh();
+	SUsdTreeView<FUsdPrimViewModelRef>::RequestListRefresh();
 	RestoreExpansionStates();
 }
 
 void SUsdStageTreeView::RestoreExpansionStates()
 {
-	TFunction< void( const FUsdPrimViewModelRef& ) > SetExpansionRecursive = [&]( const FUsdPrimViewModelRef& Item )
+	TFunction<void(const FUsdPrimViewModelRef&)> SetExpansionRecursive = [&](const FUsdPrimViewModelRef& Item)
 	{
-		if ( const UE::FUsdPrim& Prim = Item->UsdPrim )
+		if (const UE::FUsdPrim& Prim = Item->UsdPrim)
 		{
-			if (bool* bFoundExpansionState = TreeItemExpansionStates.Find( Prim.GetPrimPath().GetString() ) )
+			if (bool* bFoundExpansionState = TreeItemExpansionStates.Find(Prim.GetPrimPath().GetString()))
 			{
-				SetItemExpansion( Item, *bFoundExpansionState );
+				SetItemExpansion(Item, *bFoundExpansionState);
 			}
 			// Default to showing the root level expanded
-			else if ( Prim.GetStage().GetPseudoRoot() == Prim )
+			else if (Prim.GetStage().GetPseudoRoot() == Prim)
 			{
 				const bool bShouldExpand = true;
-				SetItemExpansion( Item, bShouldExpand );
+				SetItemExpansion(Item, bShouldExpand);
 			}
 		}
 
-		for ( const FUsdPrimViewModelRef& Child : Item->Children )
+		for (const FUsdPrimViewModelRef& Child : Item->Children)
 		{
-			SetExpansionRecursive( Child );
+			SetExpansionRecursive(Child);
 		}
 	};
 
-	for ( const FUsdPrimViewModelRef& RootItem : RootItems )
+	for (const FUsdPrimViewModelRef& RootItem : RootItems)
 	{
-		SetExpansionRecursive( RootItem );
+		SetExpansionRecursive(RootItem);
 	}
 }
 
-void SUsdStageTreeView::OnToggleAllPayloads( EPayloadsTrigger PayloadsTrigger )
+void SUsdStageTreeView::OnToggleAllPayloads(EPayloadsTrigger PayloadsTrigger)
 {
-	if ( !UsdStage )
+	if (!UsdStage)
 	{
 		return;
 	}
 
-	TArray< FUsdPrimViewModelRef > MySelectedItems = GetSelectedItems();
+	TArray<FUsdPrimViewModelRef> MySelectedItems = GetSelectedItems();
 
 	// Ideally we'd just use a UE::FSdfChangeBlock here, but for whatever reason this doesn't seem to affect the
 	// notices USD emits when loading/unloading prim payloads, so we must do this via the UsdStage directly
@@ -1640,140 +1633,145 @@ void SUsdStageTreeView::OnToggleAllPayloads( EPayloadsTrigger PayloadsTrigger )
 	TSet<UE::FSdfPath> PrimsToLoad;
 	TSet<UE::FSdfPath> PrimsToUnload;
 
-	for ( FUsdPrimViewModelRef SelectedItem : MySelectedItems )
+	for (FUsdPrimViewModelRef SelectedItem : MySelectedItems)
 	{
-		if ( SelectedItem->UsdPrim )
+		if (SelectedItem->UsdPrim)
 		{
-			TFunction< void( FUsdPrimViewModelRef ) > RecursiveTogglePayloads;
-			RecursiveTogglePayloads = [ &RecursiveTogglePayloads, PayloadsTrigger, &PrimsToLoad, &PrimsToUnload ]( FUsdPrimViewModelRef InSelectedItem ) -> void
+			TFunction<void(FUsdPrimViewModelRef)> RecursiveTogglePayloads;
+			RecursiveTogglePayloads = [&RecursiveTogglePayloads, PayloadsTrigger, &PrimsToLoad, &PrimsToUnload](FUsdPrimViewModelRef InSelectedItem
+									  ) -> void
 			{
 				UE::FUsdPrim& UsdPrim = InSelectedItem->UsdPrim;
 
-				if ( UsdPrim.HasAuthoredPayloads() )
+				if (UsdPrim.HasAuthoredPayloads())
 				{
 					bool bPrimIsLoaded = UsdPrim.IsLoaded();
 
-					if ( PayloadsTrigger == EPayloadsTrigger::Toggle )
+					if (PayloadsTrigger == EPayloadsTrigger::Toggle)
 					{
-						if ( bPrimIsLoaded )
+						if (bPrimIsLoaded)
 						{
-							PrimsToUnload.Add( UsdPrim.GetPrimPath() );
+							PrimsToUnload.Add(UsdPrim.GetPrimPath());
 						}
 						else
 						{
-							PrimsToLoad.Add( UsdPrim.GetPrimPath() );
+							PrimsToLoad.Add(UsdPrim.GetPrimPath());
 						}
 					}
-					else if ( PayloadsTrigger == EPayloadsTrigger::Load && !bPrimIsLoaded )
+					else if (PayloadsTrigger == EPayloadsTrigger::Load && !bPrimIsLoaded)
 					{
-						PrimsToLoad.Add( UsdPrim.GetPrimPath() );
+						PrimsToLoad.Add(UsdPrim.GetPrimPath());
 					}
-					else if ( PayloadsTrigger == EPayloadsTrigger::Unload && bPrimIsLoaded )
+					else if (PayloadsTrigger == EPayloadsTrigger::Unload && bPrimIsLoaded)
 					{
-						PrimsToUnload.Add( UsdPrim.GetPrimPath() );
+						PrimsToUnload.Add(UsdPrim.GetPrimPath());
 					}
 				}
 				else
 				{
-					for ( FUsdPrimViewModelRef Child : InSelectedItem->UpdateChildren() )
+					for (FUsdPrimViewModelRef Child : InSelectedItem->UpdateChildren())
 					{
-						RecursiveTogglePayloads( Child );
+						RecursiveTogglePayloads(Child);
 					}
 				}
 			};
 
-			RecursiveTogglePayloads( SelectedItem );
+			RecursiveTogglePayloads(SelectedItem);
 		}
 	}
 
-	if ( PrimsToLoad.Num() + PrimsToUnload.Num() > 0 )
+	if (PrimsToLoad.Num() + PrimsToUnload.Num() > 0)
 	{
 		UE::FSdfChangeBlock GroupNotices;
-		UsdStage.LoadAndUnload( PrimsToLoad, PrimsToUnload );
+		UsdStage.LoadAndUnload(PrimsToLoad, PrimsToUnload);
 	}
 }
 
-void SUsdStageTreeView::FillDuplicateSubmenu( FMenuBuilder& MenuBuilder )
+void SUsdStageTreeView::FillDuplicateSubmenu(FMenuBuilder& MenuBuilder)
 {
 	MenuBuilder.AddMenuEntry(
-		LOCTEXT( "DuplicateFlattened_Text", "Flatten composed prim" ),
-		LOCTEXT( "DuplicateFlattened_ToolTip", "Generate a flattened duplicate of the composed prim onto the current edit target" ),
+		LOCTEXT("DuplicateFlattened_Text", "Flatten composed prim"),
+		LOCTEXT("DuplicateFlattened_ToolTip", "Generate a flattened duplicate of the composed prim onto the current edit target"),
 		FSlateIcon(),
 		FUIAction(
-			FExecuteAction::CreateSP( this, &SUsdStageTreeView::OnDuplicatePrim, EUsdDuplicateType::FlattenComposedPrim ),
-			FCanExecuteAction::CreateSP( this, &SUsdStageTreeView::DoesPrimExistOnStage )
+			FExecuteAction::CreateSP(this, &SUsdStageTreeView::OnDuplicatePrim, EUsdDuplicateType::FlattenComposedPrim),
+			FCanExecuteAction::CreateSP(this, &SUsdStageTreeView::DoesPrimExistOnStage)
 		),
 		NAME_None,
 		EUserInterfaceActionType::Button
 	);
 
 	MenuBuilder.AddMenuEntry(
-		LOCTEXT( "DuplicateSingle_Text", "Single layer specs" ),
-		TAttribute<FText>::Create( TAttribute<FText>::FGetter::CreateLambda( [this]()
-		{
-			return DoesPrimExistOnEditTarget()
-				? LOCTEXT( "DuplicateSingleValid_ToolTip", "Duplicate the prim's specs on the current edit target only" )
-				: UE::USDStageTreeView::Private::NoSpecOnLocalLayerStack;
-		})),
+		LOCTEXT("DuplicateSingle_Text", "Single layer specs"),
+		TAttribute<FText>::Create(TAttribute<FText>::FGetter::CreateLambda(
+			[this]()
+			{
+				return DoesPrimExistOnEditTarget()
+						   ? LOCTEXT("DuplicateSingleValid_ToolTip", "Duplicate the prim's specs on the current edit target only")
+						   : UE::USDStageTreeView::Private::NoSpecOnLocalLayerStack;
+			}
+		)),
 		FSlateIcon(),
 		FUIAction(
-			FExecuteAction::CreateSP( this, &SUsdStageTreeView::OnDuplicatePrim, EUsdDuplicateType::SingleLayerSpecs ),
-			FCanExecuteAction::CreateSP( this, &SUsdStageTreeView::DoesPrimExistOnEditTarget )
+			FExecuteAction::CreateSP(this, &SUsdStageTreeView::OnDuplicatePrim, EUsdDuplicateType::SingleLayerSpecs),
+			FCanExecuteAction::CreateSP(this, &SUsdStageTreeView::DoesPrimExistOnEditTarget)
 		),
 		NAME_None,
 		EUserInterfaceActionType::Button
 	);
 
 	MenuBuilder.AddMenuEntry(
-		LOCTEXT( "DuplicateAllLocal_Text", "All local layer specs" ),
-		TAttribute<FText>::Create( TAttribute<FText>::FGetter::CreateLambda( [this]()
-		{
-			return DoesPrimHaveSpecOnLocalLayerStack()
-				? LOCTEXT( "DuplicateAllLocalValid_ToolTip", "Duplicate each of the prim's specs across the entire stage" )
-				: UE::USDStageTreeView::Private::NoSpecOnLocalLayerStack;
-		})),
+		LOCTEXT("DuplicateAllLocal_Text", "All local layer specs"),
+		TAttribute<FText>::Create(TAttribute<FText>::FGetter::CreateLambda(
+			[this]()
+			{
+				return DoesPrimHaveSpecOnLocalLayerStack()
+						   ? LOCTEXT("DuplicateAllLocalValid_ToolTip", "Duplicate each of the prim's specs across the entire stage")
+						   : UE::USDStageTreeView::Private::NoSpecOnLocalLayerStack;
+			}
+		)),
 		FSlateIcon(),
 		FUIAction(
-			FExecuteAction::CreateSP( this, &SUsdStageTreeView::OnDuplicatePrim, EUsdDuplicateType::AllLocalLayerSpecs ),
-			FCanExecuteAction::CreateSP( this, &SUsdStageTreeView::DoesPrimHaveSpecOnLocalLayerStack )
+			FExecuteAction::CreateSP(this, &SUsdStageTreeView::OnDuplicatePrim, EUsdDuplicateType::AllLocalLayerSpecs),
+			FCanExecuteAction::CreateSP(this, &SUsdStageTreeView::DoesPrimHaveSpecOnLocalLayerStack)
 		),
 		NAME_None,
 		EUserInterfaceActionType::Button
 	);
 }
 
-FReply SUsdStageTreeView::OnKeyDown( const FGeometry& MyGeometry, const FKeyEvent& InKeyEvent )
+FReply SUsdStageTreeView::OnKeyDown(const FGeometry& MyGeometry, const FKeyEvent& InKeyEvent)
 {
-	if ( UICommandList->ProcessCommandBindings( InKeyEvent ) )
+	if (UICommandList->ProcessCommandBindings(InKeyEvent))
 	{
 		return FReply::Handled();
 	}
 
-	return SUsdTreeView::OnKeyDown( MyGeometry, InKeyEvent );
+	return SUsdTreeView::OnKeyDown(MyGeometry, InKeyEvent);
 }
 
-void SUsdStageTreeView::ScrollItemIntoView( FUsdPrimViewModelRef TreeItem )
+void SUsdStageTreeView::ScrollItemIntoView(FUsdPrimViewModelRef TreeItem)
 {
 	FUsdPrimViewModel* Parent = TreeItem->ParentItem;
-	while( Parent )
+	while (Parent)
 	{
-		SetItemExpansion( StaticCastSharedRef< FUsdPrimViewModel>( Parent->AsShared() ), true );
+		SetItemExpansion(StaticCastSharedRef<FUsdPrimViewModel>(Parent->AsShared()), true);
 		Parent = Parent->ParentItem;
 	}
 
-	RequestScrollIntoView( TreeItem );
+	RequestScrollIntoView(TreeItem);
 }
 
-void SUsdStageTreeView::OnTreeItemScrolledIntoView( FUsdPrimViewModelRef TreeItem, const TSharedPtr<ITableRow>& Widget )
+void SUsdStageTreeView::OnTreeItemScrolledIntoView(FUsdPrimViewModelRef TreeItem, const TSharedPtr<ITableRow>& Widget)
 {
-	if ( TreeItem == PendingRenameItem.Pin() )
+	if (TreeItem == PendingRenameItem.Pin())
 	{
 		PendingRenameItem = nullptr;
 		TreeItem->RenameRequestEvent.ExecuteIfBound();
 	}
 }
 
-void SUsdStageTreeView::OnPrimNameCommitted( const FUsdPrimViewModelRef& ViewModel, const FText& InPrimName )
+void SUsdStageTreeView::OnPrimNameCommitted(const FUsdPrimViewModelRef& ViewModel, const FText& InPrimName)
 {
 	// Reset this regardless of how we exit this function
 	const bool bRenamingExistingPrim = ViewModel->bIsRenamingExistingPrim;
@@ -1797,13 +1795,13 @@ void SUsdStageTreeView::OnPrimNameCommitted( const FUsdPrimViewModelRef& ViewMod
 		}
 	};
 
-	if ( InPrimName.IsEmptyOrWhitespace() )
+	if (InPrimName.IsEmptyOrWhitespace())
 	{
 		CancelInput();
 		return;
 	}
 
-	if ( bRenamingExistingPrim )
+	if (bRenamingExistingPrim)
 	{
 		if (UsdUtils::NotifyIfInstanceProxy(ViewModel->UsdPrim))
 		{
@@ -1811,7 +1809,7 @@ void SUsdStageTreeView::OnPrimNameCommitted( const FUsdPrimViewModelRef& ViewMod
 			return;
 		}
 
-		FScopedTransaction Transaction( LOCTEXT( "RenamePrimTransaction", "Rename a prim" ) );
+		FScopedTransaction Transaction(LOCTEXT("RenamePrimTransaction", "Rename a prim"));
 
 		// e.g. "/Root/OldPrim/"
 		FString OldPath = ViewModel->UsdPrim.GetPrimPath().GetString();
@@ -1823,23 +1821,23 @@ void SUsdStageTreeView::OnPrimNameCommitted( const FUsdPrimViewModelRef& ViewMod
 		// that will trigger refreshes of the tree view
 		{
 			// e.g. "/Root/NewPrim"
-			FString NewPath = FString::Printf( TEXT( "%s/%s" ), *FPaths::GetPath( OldPath ), *NewNameStr );
+			FString NewPath = FString::Printf(TEXT("%s/%s"), *FPaths::GetPath(OldPath), *NewNameStr);
 			TMap<FString, bool> PairsToAdd;
-			for ( TMap<FString, bool>::TIterator It( TreeItemExpansionStates ); It; ++It )
+			for (TMap<FString, bool>::TIterator It(TreeItemExpansionStates); It; ++It)
 			{
 				// e.g. "/Root/OldPrim/SomeChild"
 				FString SomePrimPath = It->Key;
-				if ( SomePrimPath.RemoveFromStart( OldPath ) )  // e.g. "/SomeChild"
+				if (SomePrimPath.RemoveFromStart(OldPath))	  // e.g. "/SomeChild"
 				{
 					// e.g. "/Root/NewPrim/SomeChild"
 					SomePrimPath = NewPath + SomePrimPath;
-					PairsToAdd.Add( SomePrimPath, It->Value );
+					PairsToAdd.Add(SomePrimPath, It->Value);
 				}
 			}
-			TreeItemExpansionStates.Append( PairsToAdd );
+			TreeItemExpansionStates.Append(PairsToAdd);
 		}
 
-		UsdUtils::RenamePrim( ViewModel->UsdPrim, *NewNameStr );
+		UsdUtils::RenamePrim(ViewModel->UsdPrim, *NewNameStr);
 	}
 	else
 	{
@@ -1852,23 +1850,23 @@ void SUsdStageTreeView::OnPrimNameCommitted( const FUsdPrimViewModelRef& ViewMod
 			}
 		}
 
-		FScopedTransaction Transaction( LOCTEXT( "AddPrimTransaction", "Add a new prim" ) );
+		FScopedTransaction Transaction(LOCTEXT("AddPrimTransaction", "Add a new prim"));
 
-		ViewModel->DefinePrim( *InPrimName.ToString() );
+		ViewModel->DefinePrim(*InPrimName.ToString());
 
 		const bool bResync = true;
 
 		// Renamed a child item
-		if ( ViewModel->ParentItem )
+		if (ViewModel->ParentItem)
 		{
-			ViewModel->ParentItem->Children.Remove( ViewModel );
+			ViewModel->ParentItem->Children.Remove(ViewModel);
 
-			RefreshPrim( ViewModel->ParentItem->UsdPrim.GetPrimPath().GetString(), bResync );
+			RefreshPrim(ViewModel->ParentItem->UsdPrim.GetPrimPath().GetString(), bResync);
 		}
 		// Renamed a root item
 		else
 		{
-			RefreshPrim( ViewModel->UsdPrim.GetPrimPath().GetString(), bResync );
+			RefreshPrim(ViewModel->UsdPrim.GetPrimPath().GetString(), bResync);
 		}
 	}
 }
@@ -1884,13 +1882,13 @@ void SUsdStageTreeView::OnPrimNameUpdated(const FUsdPrimViewModelRef& TreeItem,
 
 	{
 		const UE::FUsdStageWeak& Stage = TreeItem->UsdStage;
-		if ( !Stage )
+		if (!Stage)
 		{
 			return;
 		}
 
 		UE::FSdfPath ParentPrimPath;
-		if ( TreeItem->ParentItem )
+		if (TreeItem->ParentItem)
 		{
 			ParentPrimPath = TreeItem->ParentItem->UsdPrim.GetPrimPath();
 		}
@@ -1899,9 +1897,9 @@ void SUsdStageTreeView::OnPrimNameUpdated(const FUsdPrimViewModelRef& TreeItem,
 			ParentPrimPath = UE::FSdfPath::AbsoluteRootPath();
 		}
 
-		UE::FSdfPath NewPrimPath = ParentPrimPath.AppendChild( *NameStr );
-		const UE::FUsdPrim& Prim = Stage.GetPrimAtPath( NewPrimPath );
-		if ( Prim && Prim != TreeItem->UsdPrim )
+		UE::FSdfPath NewPrimPath = ParentPrimPath.AppendChild(*NameStr);
+		const UE::FUsdPrim& Prim = Stage.GetPrimAtPath(NewPrimPath);
+		if (Prim && Prim != TreeItem->UsdPrim)
 		{
 			ErrorMessage = LOCTEXT("DuplicatePrimName", "A Prim with this name already exists!");
 			return;
@@ -1909,6 +1907,6 @@ void SUsdStageTreeView::OnPrimNameUpdated(const FUsdPrimViewModelRef& TreeItem,
 	}
 }
 
-#endif // #if USE_USD_SDK
+#endif	  // #if USE_USD_SDK
 
 #undef LOCTEXT_NAMESPACE
diff --git a/Engine/Plugins/Importers/USDImporter/Source/USDStageEditor/Private/SUSDStageTreeView.h b/Engine/Plugins/Importers/USDImporter/Source/USDStageEditor/Private/SUSDStageTreeView.h
--- a/Engine/Plugins/Importers/USDImporter/Source/USDStageEditor/Private/SUSDStageTreeView.h
+++ b/Engine/Plugins/Importers/USDImporter/Source/USDStageEditor/Private/SUSDStageTreeView.h
@@ -6,9 +6,6 @@
 #include "UsdWrappers/UsdStage.h"
 #include "Widgets/SUSDTreeView.h"
 
-#include "Widgets/Views/SHeaderRow.h"
-#include "Widgets/Views/STreeView.h"
-
 class AUsdStageActor;
 class FUICommandList;
 enum class EPayloadsTrigger;
@@ -20,54 +17,59 @@ namespace UE
 
 #if USE_USD_SDK
 
-DECLARE_DELEGATE_OneParam( FOnPrimSelectionChanged, const TArray<FString>& /* NewSelection */);
-DECLARE_DELEGATE_OneParam( FOnAddPrim, FString );
+DECLARE_DELEGATE_OneParam(FOnPrimSelectionChanged, const TArray<FString>& /* NewSelection */);
+DECLARE_DELEGATE_OneParam(FOnAddPrim, FString);
 
-class SUsdStageTreeView : public SUsdTreeView< FUsdPrimViewModelRef >
+class SUsdStageTreeView : public SUsdTreeView<FUsdPrimViewModelRef>
 {
 public:
-	SLATE_BEGIN_ARGS( SUsdStageTreeView ) {}
-		SLATE_EVENT( FOnPrimSelectionChanged, OnPrimSelectionChanged )
+	SLATE_BEGIN_ARGS(SUsdStageTreeView)
+	{
+	}
+	SLATE_EVENT(FOnPrimSelectionChanged, OnPrimSelectionChanged)
 	SLATE_END_ARGS()
 
-	void Construct( const FArguments& InArgs );
+	void Construct(const FArguments& InArgs);
 
-	void Refresh( const UE::FUsdStageWeak& NewStage );
-	void RefreshPrim( const FString& PrimPath, bool bResync );
+	void Refresh(const UE::FUsdStageWeak& NewStage);
+	void RefreshPrim(const FString& PrimPath, bool bResync);
 
-	FUsdPrimViewModelPtr GetItemFromPrimPath( const FString& PrimPath );
+	FUsdPrimViewModelPtr GetItemFromPrimPath(const FString& PrimPath);
 
-	void SetSelectedPrimPaths( const TArray<FString>& PrimPaths );
-	void SetSelectedPrims( const TArray<UE::FUsdPrim>& Prims );
+	void SetSelectedPrimPaths(const TArray<FString>& PrimPaths);
+	void SetSelectedPrims(const TArray<UE::FUsdPrim>& Prims);
 	TArray<FString> GetSelectedPrimPaths();
 	TArray<UE::FUsdPrim> GetSelectedPrims();
 
 private:
-	virtual TSharedRef< ITableRow > OnGenerateRow( FUsdPrimViewModelRef InDisplayNode, const TSharedRef< STableViewBase >& OwnerTable ) override;
-	virtual void OnGetChildren( FUsdPrimViewModelRef InParent, TArray< FUsdPrimViewModelRef >& OutChildren ) const override;
+	virtual TSharedRef<ITableRow> OnGenerateRow(FUsdPrimViewModelRef InDisplayNode, const TSharedRef<STableViewBase>& OwnerTable) override;
+	virtual void OnGetChildren(FUsdPrimViewModelRef InParent, TArray<FUsdPrimViewModelRef>& OutChildren) const override;
 
 	// Required so that we can use the cut/copy/paste/etc. shortcuts
-	virtual bool SupportsKeyboardFocus() const override { return true; }
-	virtual FReply OnKeyDown( const FGeometry& MyGeometry, const FKeyEvent& InKeyEvent ) override;
+	virtual bool SupportsKeyboardFocus() const override
+	{
+		return true;
+	}
+	virtual FReply OnKeyDown(const FGeometry& MyGeometry, const FKeyEvent& InKeyEvent) override;
 
-	void ScrollItemIntoView( FUsdPrimViewModelRef TreeItem );
-	virtual void OnTreeItemScrolledIntoView( FUsdPrimViewModelRef TreeItem, const TSharedPtr< ITableRow >& Widget ) override ;
+	void ScrollItemIntoView(FUsdPrimViewModelRef TreeItem);
+	virtual void OnTreeItemScrolledIntoView(FUsdPrimViewModelRef TreeItem, const TSharedPtr<ITableRow>& Widget) override;
 
-	void OnPrimNameCommitted( const FUsdPrimViewModelRef& TreeItem, const FText& InPrimName );
-	void OnPrimNameUpdated( const FUsdPrimViewModelRef& TreeItem, const FText& InPrimName, FText& ErrorMessage );
+	void OnPrimNameCommitted(const FUsdPrimViewModelRef& TreeItem, const FText& InPrimName);
+	void OnPrimNameUpdated(const FUsdPrimViewModelRef& TreeItem, const FText& InPrimName, FText& ErrorMessage);
 
 	virtual void SetupColumns() override;
-	TSharedPtr< SWidget > ConstructPrimContextMenu();
+	TSharedPtr<SWidget> ConstructPrimContextMenu();
 
-	void OnToggleAllPayloads( EPayloadsTrigger PayloadsTrigger );
+	void OnToggleAllPayloads(EPayloadsTrigger PayloadsTrigger);
 
-	void FillDuplicateSubmenu( FMenuBuilder& MenuBuilder );
+	void FillDuplicateSubmenu(FMenuBuilder& MenuBuilder);
 
 	void OnAddChildPrim();
 	void OnCutPrim();
 	void OnCopyPrim();
 	void OnPastePrim();
-	void OnDuplicatePrim( EUsdDuplicateType DuplicateType );
+	void OnDuplicatePrim(EUsdDuplicateType DuplicateType);
 	void OnDeletePrim();
 	void OnRenamePrim();
 
@@ -77,10 +79,10 @@ private:
 	void OnAddPayload();
 	void OnClearPayloads();
 
-	void OnApplySchema( FName SchemaName );
-	void OnRemoveSchema( FName SchemaName );
-	bool CanApplySchema( FName SchemaName );
-	bool CanRemoveSchema( FName SchemaName );
+	void OnApplySchema(FName SchemaName);
+	void OnRemoveSchema(FName SchemaName);
+	bool CanApplySchema(FName SchemaName);
+	bool CanRemoveSchema(FName SchemaName);
 
 	bool CanAddChildPrim() const;
 	bool CanPastePrim() const;
@@ -92,7 +94,7 @@ private:
 	void RestoreExpansionStates();
 	virtual void RequestListRefresh() override;
 
-	void SelectItemsInternal( const TArray< FUsdPrimViewModelRef >& ItemsToSelect );
+	void SelectItemsInternal(const TArray<FUsdPrimViewModelRef>& ItemsToSelect);
 
 public:
 	// We update only on slate tick, but can receive a USD notice at any point, from any thread.
@@ -115,14 +117,14 @@ private:
 	// Should always be valid, we keep the one we're given on Refresh()
 	UE::FUsdStageWeak UsdStage;
 
-	TWeakPtr< FUsdPrimViewModel > PendingRenameItem;
+	TWeakPtr<FUsdPrimViewModel> PendingRenameItem;
 
 	// So that we can store these across refreshes
-	TMap< FString, bool > TreeItemExpansionStates;
+	TMap<FString, bool> TreeItemExpansionStates;
 
 	FOnPrimSelectionChanged OnPrimSelectionChanged;
 
 	TSharedPtr<FUICommandList> UICommandList;
 };
 
-#endif // #if USE_USD_SDK
+#endif	  // #if USE_USD_SDK
