USD: Reformat USDStageEditor module.

diff --git a/Engine/Plugins/Importers/USDImporter/Source/USDStageEditor/Private/USDStageEditorModule.cpp b/Engine/Plugins/Importers/USDImporter/Source/USDStageEditor/Private/USDStageEditorModule.cpp
--- a/Engine/Plugins/Importers/USDImporter/Source/USDStageEditor/Private/USDStageEditorModule.cpp
+++ b/Engine/Plugins/Importers/USDImporter/Source/USDStageEditor/Private/USDStageEditorModule.cpp
@@ -11,14 +11,12 @@
 #include "USDProjectSettings.h"
 #include "USDStageActor.h"
 #include "UsdWrappers/SdfLayer.h"
-#include "UsdWrappers/UsdAttribute.h"
 #include "UsdWrappers/UsdPrim.h"
 
 #include "Editor/TransBuffer.h"
 #include "Editor/UnrealEdEngine.h"
 #include "EngineUtils.h"
 #include "Framework/Docking/TabManager.h"
-#include "IAssetTools.h"
 #include "Interfaces/IMainFrameModule.h"
 #include "LevelEditor.h"
 #include "Modules/ModuleManager.h"
@@ -33,18 +31,18 @@
 
 namespace UE::UsdStageEditorModule::Private
 {
-	const static FTabId UsdStageEditorTabID{ TEXT( "USDStage" ) };
+	const static FTabId UsdStageEditorTabID{TEXT("USDStage")};
 
-	void SaveStageActorLayersForWorld( UWorld* World, bool bForClosing, AUsdStageActor* TargetStageActor = nullptr )
+	void SaveStageActorLayersForWorld(UWorld* World, bool bForClosing, AUsdStageActor* TargetStageActor = nullptr)
 	{
 #if USE_USD_SDK
-		if ( !World )
+		if (!World)
 		{
 			return;
 		}
 
 		UUsdProjectSettings* Settings = GetMutableDefault<UUsdProjectSettings>();
-		if ( !Settings )
+		if (!Settings)
 		{
 			return;
 		}
@@ -52,14 +50,14 @@ namespace UE::UsdStageEditorModule::Private
 		// Reentrant guard here because if we ever save an anonymous layer we'll update the stage actors that use it
 		// to point to the new (saved) layer, which will internally close the anonymous stage and get us back in here
 		static bool bIsReentrant = false;
-		if ( bIsReentrant )
+		if (bIsReentrant)
 		{
 			return;
 		}
-		TGuardValue<bool> ReentrantGuard( bIsReentrant, true );
+		TGuardValue<bool> ReentrantGuard(bIsReentrant, true);
 
 		bool bPrompt = false;
-		switch ( bForClosing ? Settings->ShowSaveLayersDialogWhenClosing : Settings->ShowSaveLayersDialogWhenSaving )
+		switch (bForClosing ? Settings->ShowSaveLayersDialogWhenClosing : Settings->ShowSaveLayersDialogWhenSaving)
 		{
 			case EUsdSaveDialogBehavior::NeverSave:
 			{
@@ -80,132 +78,125 @@ namespace UE::UsdStageEditorModule::Private
 		}
 
 		TArray<AUsdStageActor*> StageActorsToVisit;
-		if ( TargetStageActor )
+		if (TargetStageActor)
 		{
-			if ( TargetStageActor->GetWorld() == World )
+			if (TargetStageActor->GetWorld() == World)
 			{
-				StageActorsToVisit.Add( TargetStageActor );
+				StageActorsToVisit.Add(TargetStageActor);
 			}
 		}
 		else
 		{
-			for ( TActorIterator<AUsdStageActor> It{ World }; It; ++It )
+			for (TActorIterator<AUsdStageActor> It{World}; It; ++It)
 			{
-				StageActorsToVisit.Add( *It );
+				StageActorsToVisit.Add(*It);
 			}
 		}
 
 		// For now lets only care about stages opened on stage actors. The user could have additional stages,
 		// like opened via Python or custom C++ plugins, but lets ignore those
 		TMap<FString, FUsdSaveDialogRowData> RowsByIdentifier;
-		for ( AUsdStageActor* StageActor : StageActorsToVisit )
+		for (AUsdStageActor* StageActor : StageActorsToVisit)
 		{
-			if ( !StageActor )
+			if (!StageActor)
 			{
 				continue;
 			}
 
-			UE::FUsdStage UsdStage = static_cast< const AUsdStageActor* >( StageActor )->GetUsdStage();
-			if ( !UsdStage )
+			UE::FUsdStage UsdStage = static_cast<const AUsdStageActor*>(StageActor)->GetUsdStage();
+			if (!UsdStage)
 			{
 				continue;
 			}
 
 			TArray<UE::FSdfLayer> UsedLayers = UsdStage.GetUsedLayers();
-			RowsByIdentifier.Reserve( RowsByIdentifier.Num() + UsedLayers.Num() );
+			RowsByIdentifier.Reserve(RowsByIdentifier.Num() + UsedLayers.Num());
 
-			for ( const UE::FSdfLayer& UsedLayer : UsedLayers )
+			for (const UE::FSdfLayer& UsedLayer : UsedLayers)
 			{
 				// This comment is written to the layer when we're in the process of saving a memory-only
 				// stage, and indicates that this layer is already saved (even though it will show as dirty and
 				// anonymous)
-				if ( UsedLayer.IsDirty() && UsedLayer.GetComment() != UnrealIdentifiers::LayerSavedComment )
+				if (UsedLayer.IsDirty() && UsedLayer.GetComment() != UnrealIdentifiers::LayerSavedComment)
 				{
-					FUsdSaveDialogRowData& RowData = RowsByIdentifier.FindOrAdd( UsedLayer.GetIdentifier() );
+					FUsdSaveDialogRowData& RowData = RowsByIdentifier.FindOrAdd(UsedLayer.GetIdentifier());
 					RowData.Layer = UsedLayer;
-					RowData.ConsumerStages.AddUnique( UsdStage );
-					RowData.ConsumerActors.Add( StageActor );
+					RowData.ConsumerStages.AddUnique(UsdStage);
+					RowData.ConsumerActors.Add(StageActor);
 				}
 			}
 		}
 
-		if ( RowsByIdentifier.Num() == 0 )
+		if (RowsByIdentifier.Num() == 0)
 		{
 			return;
 		}
 
-		TArray< FUsdSaveDialogRowData > Rows;
-		RowsByIdentifier.GenerateValueArray( Rows );
+		TArray<FUsdSaveDialogRowData> Rows;
+		RowsByIdentifier.GenerateValueArray(Rows);
 
-		if ( bPrompt )
+		if (bPrompt)
 		{
-			Rows.Sort( []( const FUsdSaveDialogRowData& Left, const FUsdSaveDialogRowData& Right )
-			{
-				return ( Left.Layer && Right.Layer )
-					? Left.Layer.GetIdentifier() < Right.Layer.GetIdentifier()
-					// This shouldn't ever happen but just do something consistent here instead anyway
-					: Left.ConsumerStages.Num() < Right.ConsumerStages.Num();
-			});
+			Rows.Sort(
+				[](const FUsdSaveDialogRowData& Left, const FUsdSaveDialogRowData& Right)
+				{
+					return (Left.Layer && Right.Layer) ? Left.Layer.GetIdentifier() < Right.Layer.GetIdentifier()
+													   // This shouldn't ever happen but just do something consistent here instead anyway
+													   : Left.ConsumerStages.Num() < Right.ConsumerStages.Num();
+				}
+			);
 
-			const FText WindowTitle = LOCTEXT( "SaveDialogTitle", "Save USD Layers" );
+			// clang-format off
+			const FText WindowTitle = LOCTEXT("SaveDialogTitle", "Save USD Layers");
 			const FText DescriptionText = bForClosing
-				? LOCTEXT( "CloseDialogDescTextText", "Before closing these USD Stages, do you want to save these USD layers to disk?" )
-				: LOCTEXT( "SaveDialogDescTextText", "Since you're saving the Level, do you want to save these USD layers to disk?" );
+				? LOCTEXT("CloseDialogDescTextText", "Before closing these USD Stages, do you want to save these USD layers to disk?")
+				: LOCTEXT("SaveDialogDescTextText", "Since you're saving the Level, do you want to save these USD layers to disk?");
+			// clang-format on
 
 			bool bShouldSave = true;
 			bool bShouldPromptAgain = true;
-			Rows = SUsdSaveDialog::ShowDialog(
-				Rows,
-				WindowTitle,
-				DescriptionText,
-				&bShouldSave,
-				&bShouldPromptAgain
-			);
+			Rows = SUsdSaveDialog::ShowDialog(Rows, WindowTitle, DescriptionText, &bShouldSave, &bShouldPromptAgain);
 
-			EUsdSaveDialogBehavior* bSetting = bForClosing
-				? &Settings->ShowSaveLayersDialogWhenClosing
-				: &Settings->ShowSaveLayersDialogWhenSaving;
+			EUsdSaveDialogBehavior* bSetting = bForClosing ? &Settings->ShowSaveLayersDialogWhenClosing : &Settings->ShowSaveLayersDialogWhenSaving;
 
-			*bSetting = bShouldPromptAgain
-				? EUsdSaveDialogBehavior::ShowPrompt
-				: bShouldSave
-					? EUsdSaveDialogBehavior::AlwaysSave
-					: EUsdSaveDialogBehavior::NeverSave;
+			*bSetting = bShouldPromptAgain ? EUsdSaveDialogBehavior::ShowPrompt
+						: bShouldSave	   ? EUsdSaveDialogBehavior::AlwaysSave
+										   : EUsdSaveDialogBehavior::NeverSave;
 
 			Settings->SaveConfig();
 		}
 
-		for ( const FUsdSaveDialogRowData& ReturnedRow : Rows )
+		for (const FUsdSaveDialogRowData& ReturnedRow : Rows)
 		{
 			UE::FSdfLayer PinnedLayer = ReturnedRow.Layer;
-			if ( ReturnedRow.bSaveLayer && PinnedLayer )
+			if (ReturnedRow.bSaveLayer && PinnedLayer)
 			{
 				bool bSaved = false;
 
-				if ( PinnedLayer.IsAnonymous() )
+				if (PinnedLayer.IsAnonymous())
 				{
 					// For now we only allow anonymous stages, and not individual layers, so we don't have to
 					// patch up anything
-					TOptional< FString > UsdFilePath = UsdUtils::BrowseUsdFile( UsdUtils::EBrowseFileMode::Save );
-					if ( UsdFilePath )
+					TOptional<FString> UsdFilePath = UsdUtils::BrowseUsdFile(UsdUtils::EBrowseFileMode::Save);
+					if (UsdFilePath)
 					{
-						bSaved = PinnedLayer.Export( *UsdFilePath.GetValue() );
+						bSaved = PinnedLayer.Export(*UsdFilePath.GetValue());
 
 						// If any stage actors were pointing at the in-memory versions of these stages, update them
 						// to point to the saved versions
-						if ( bSaved && !bForClosing )
+						if (bSaved && !bForClosing)
 						{
 							// Even though we're potentially going to load actors and assets here, we don't need to
 							// use a scoped transaction as the Save All command already clears the transaction buffer
 							// anyway, and we won't get in here when closing.
 
 							FString ExpectedIdentifier = UnrealIdentifiers::IdentifierPrefix + PinnedLayer.GetIdentifier();
-							for ( AUsdStageActor* StageActor : ReturnedRow.ConsumerActors )
+							for (AUsdStageActor* StageActor : ReturnedRow.ConsumerActors)
 							{
-								if ( StageActor->RootLayer.FilePath == ExpectedIdentifier
-									&& static_cast<const AUsdStageActor*>(StageActor)->GetUsdStage() )
+								if (StageActor->RootLayer.FilePath == ExpectedIdentifier
+									&& static_cast<const AUsdStageActor*>(StageActor)->GetUsdStage())
 								{
-									StageActor->SetRootLayer( *UsdFilePath.GetValue() );
+									StageActor->SetRootLayer(*UsdFilePath.GetValue());
 								}
 							}
 						}
@@ -214,33 +205,32 @@ namespace UE::UsdStageEditorModule::Private
 				else
 				{
 					const bool bForce = true;
-					bSaved = PinnedLayer.Save( bForce );
+					bSaved = PinnedLayer.Save(bForce);
 				}
 
-				if ( !bSaved )
+				if (!bSaved)
 				{
-					UE_LOG( LogUsd, Warning, TEXT( "Failed to save layer '%s'" ), *PinnedLayer.GetIdentifier() );
+					UE_LOG(LogUsd, Warning, TEXT("Failed to save layer '%s'"), *PinnedLayer.GetIdentifier());
 				}
 			}
 		}
-#endif // USE_USD_SDK
+#endif	  // USE_USD_SDK
 	}
 
 #if USE_USD_SDK
-	TSharedPtr<SUsdStage> GetUsdStageEditor( bool bOpenIfNeeded = true )
+	TSharedPtr<SUsdStage> GetUsdStageEditor(bool bOpenIfNeeded = true)
 	{
-		FLevelEditorModule& LevelEditorModule = FModuleManager::GetModuleChecked< FLevelEditorModule >( "LevelEditor" );
-		TSharedPtr< FTabManager > LevelEditorTabManager = LevelEditorModule.GetLevelEditorTabManager();
+		FLevelEditorModule& LevelEditorModule = FModuleManager::GetModuleChecked<FLevelEditorModule>("LevelEditor");
+		TSharedPtr<FTabManager> LevelEditorTabManager = LevelEditorModule.GetLevelEditorTabManager();
 
-		TSharedPtr< SDockTab > Tab = bOpenIfNeeded
-			? LevelEditorTabManager->TryInvokeTab( UsdStageEditorTabID )
-			: LevelEditorTabManager->FindExistingLiveTab( UsdStageEditorTabID );
+		TSharedPtr<SDockTab> Tab = bOpenIfNeeded ? LevelEditorTabManager->TryInvokeTab(UsdStageEditorTabID)
+												 : LevelEditorTabManager->FindExistingLiveTab(UsdStageEditorTabID);
 
-		if ( Tab )
+		if (Tab)
 		{
-			if ( TSharedPtr<SBorder> ContentBorder = StaticCastSharedRef<SBorder>( Tab->GetContent() ) )
+			if (TSharedPtr<SBorder> ContentBorder = StaticCastSharedRef<SBorder>(Tab->GetContent()))
 			{
-				if ( TSharedPtr<SUsdStage> UsdStageEditor = StaticCastSharedRef<SUsdStage>( ContentBorder->GetContent() ) )
+				if (TSharedPtr<SUsdStage> UsdStageEditor = StaticCastSharedRef<SUsdStage>(ContentBorder->GetContent()))
 				{
 					return UsdStageEditor;
 				}
@@ -249,28 +239,30 @@ namespace UE::UsdStageEditorModule::Private
 
 		return nullptr;
 	}
-#endif // USE_USD_SDK
-}
+#endif	  // USE_USD_SDK
+}	 // namespace UE::UsdStageEditorModule::Private
 
 class FUsdStageEditorModule : public IUsdStageEditorModule
 {
 public:
 #if USE_USD_SDK
-	static TSharedRef< SDockTab > SpawnUsdStageTab( const FSpawnTabArgs& SpawnTabArgs )
+	static TSharedRef<SDockTab> SpawnUsdStageTab(const FSpawnTabArgs& SpawnTabArgs)
 	{
 		LLM_SCOPE_BYTAG(Usd);
 
+		// clang-format off
 		return SNew(SDockTab)
 			.TabRole(ETabRole::NomadTab)
-			.Label( LOCTEXT( "USDStage", "USD Stage" ) )
+			.Label(LOCTEXT("USDStage", "USD Stage"))
 			[
 				SNew(SBorder)
 				.Padding(0)
-				.BorderImage( FAppStyle::GetBrush("ToolPanel.GroupBorder") )
+				.BorderImage(FAppStyle::GetBrush("ToolPanel.GroupBorder"))
 				[
-					SNew( SUsdStage )
+					SNew(SUsdStage)
 				]
 			];
+		// clang-format on
 	}
 
 	virtual void StartupModule() override
@@ -279,133 +271,147 @@ public:
 
 		FUsdStageEditorStyle::Initialize();
 
-		FLevelEditorModule& LevelEditorModule = FModuleManager::GetModuleChecked< FLevelEditorModule >( "LevelEditor" );
+		FLevelEditorModule& LevelEditorModule = FModuleManager::GetModuleChecked<FLevelEditorModule>("LevelEditor");
 		LevelEditorTabManagerChangedHandle = LevelEditorModule.OnTabManagerChanged().AddLambda(
 			[]()
 			{
-				FLevelEditorModule& LevelEditorModule = FModuleManager::GetModuleChecked< FLevelEditorModule >( "LevelEditor" );
-				TSharedPtr< FTabManager > LevelEditorTabManager = LevelEditorModule.GetLevelEditorTabManager();
+				FLevelEditorModule& LevelEditorModule = FModuleManager::GetModuleChecked<FLevelEditorModule>("LevelEditor");
+				TSharedPtr<FTabManager> LevelEditorTabManager = LevelEditorModule.GetLevelEditorTabManager();
 
 				const FSlateIcon LayersIcon(FAppStyle::GetAppStyleSetName(), "LevelEditor.Tabs.USDStage");
 
+				// clang-format off
 				LevelEditorTabManager->RegisterTabSpawner(
 					UE::UsdStageEditorModule::Private::UsdStageEditorTabID.TabType,
-					FOnSpawnTab::CreateStatic( &FUsdStageEditorModule::SpawnUsdStageTab )
+					FOnSpawnTab::CreateStatic(&FUsdStageEditorModule::SpawnUsdStageTab)
 				)
-				.SetDisplayName( LOCTEXT( "USDStage", "USD Stage" ) )
-				.SetTooltipText( LOCTEXT( "USDStageTab", "Open USD Stage tab" ) )
+				.SetDisplayName(LOCTEXT("USDStage", "USD Stage"))
+				.SetTooltipText(LOCTEXT("USDStageTab", "Open USD Stage tab"))
 				.SetGroup(WorkspaceMenu::GetMenuStructure().GetLevelEditorVirtualProductionCategory())
-				.SetIcon( LayersIcon );
-			});
+				.SetIcon(LayersIcon);
+				// clang-format on
+			}
+		);
 
 		// Prompt to save modified USD layers when closing the editor
-		IMainFrameModule& MainFrame = FModuleManager::LoadModuleChecked<IMainFrameModule>( "MainFrame" );
-		EditorCanCloseDelegate = MainFrame.RegisterCanCloseEditor( IMainFrameModule::FMainFrameCanCloseEditor::CreateLambda( []() -> bool
-		{
-			if ( GEditor && GEngine )
+		IMainFrameModule& MainFrame = FModuleManager::LoadModuleChecked<IMainFrameModule>("MainFrame");
+		EditorCanCloseDelegate = MainFrame.RegisterCanCloseEditor(IMainFrameModule::FMainFrameCanCloseEditor::CreateLambda(
+			[]() -> bool
 			{
-				UWorld* EditorWorld = nullptr;
-				for ( const FWorldContext& Context : GEngine->GetWorldContexts() )
+				if (GEditor && GEngine)
 				{
-					if ( Context.WorldType == EWorldType::Editor )
+					UWorld* EditorWorld = nullptr;
+					for (const FWorldContext& Context : GEngine->GetWorldContexts())
 					{
-						EditorWorld = Context.World();
+						if (Context.WorldType == EWorldType::Editor)
+						{
+							EditorWorld = Context.World();
+						}
 					}
+
+					const bool bForClosing = true;
+					UE::UsdStageEditorModule::Private::SaveStageActorLayersForWorld(EditorWorld, bForClosing);
 				}
 
-				const bool bForClosing = true;
-				UE::UsdStageEditorModule::Private::SaveStageActorLayersForWorld( EditorWorld, bForClosing );
+				// We won't actually ever block the save
+				return true;
 			}
-
-			// We won't actually ever block the save
-			return true;
-		}));
+		));
 
 		// Prompt to save modified USD Layers when closing stageactor stages
-		StageActorLoadedHandle = AUsdStageActor::OnActorLoaded.AddLambda([this](AUsdStageActor* StageActor)
-		{
-			if ( !StageActor )
+		StageActorLoadedHandle = AUsdStageActor::OnActorLoaded.AddLambda(
+			[this](AUsdStageActor* StageActor)
 			{
-				return;
-			}
+				if (!StageActor)
+				{
+					return;
+				}
 
-			// We never want to prompt when undoing or redoing.
-			// We have to subscribe to this here as the UTransBuffer doesn't exist by the time the module is
-			// initializing
-			if ( UTransBuffer* TransBuffer = GUnrealEd ? Cast<UTransBuffer>( GUnrealEd->Trans ) : nullptr )
-			{
-				if ( !OnTransactionStateChangedHandle.IsValid() )
+				// We never want to prompt when undoing or redoing.
+				// We have to subscribe to this here as the UTransBuffer doesn't exist by the time the module is
+				// initializing
+				if (UTransBuffer* TransBuffer = GUnrealEd ? Cast<UTransBuffer>(GUnrealEd->Trans) : nullptr)
 				{
-					OnTransactionStateChangedHandle = TransBuffer->OnTransactionStateChanged().AddLambda(
-						[this]( const FTransactionContext& InTransactionContext, const ETransactionStateEventType InTransactionState )
-						{
-							if ( InTransactionState == ETransactionStateEventType::UndoRedoStarted )
-							{
-								bUndoRedoing = true;
-							}
-							else if ( InTransactionState == ETransactionStateEventType::UndoRedoFinalized )
+					if (!OnTransactionStateChangedHandle.IsValid())
+					{
+						OnTransactionStateChangedHandle = TransBuffer->OnTransactionStateChanged().AddLambda(
+							[this](const FTransactionContext& InTransactionContext, const ETransactionStateEventType InTransactionState)
 							{
-								bUndoRedoing = false;
+								if (InTransactionState == ETransactionStateEventType::UndoRedoStarted)
+								{
+									bUndoRedoing = true;
+								}
+								else if (InTransactionState == ETransactionStateEventType::UndoRedoFinalized)
+								{
+									bUndoRedoing = false;
+								}
 							}
-						});
+						);
+					}
 				}
+
+				StageActor->OnPreStageChanged.AddLambda(
+					[this, StageActor]()
+					{
+						if (!bUndoRedoing && StageActor && static_cast<const AUsdStageActor*>(StageActor)->GetUsdStage())
+						{
+							const bool bForClosing = true;
+							UE::UsdStageEditorModule::Private::SaveStageActorLayersForWorld(StageActor->GetWorld(), bForClosing, StageActor);
+						}
+					}
+				);
 			}
+		);
 
-			StageActor->OnPreStageChanged.AddLambda( [this, StageActor]()
+		// Prompt to save modified USD layers when saving the world
+		PreSaveWorldEditorDelegateHandle = FEditorDelegates::PreSaveWorldWithContext.AddLambda(
+			[](UWorld* World, FObjectPreSaveContext InContext)
 			{
-				if ( !bUndoRedoing && StageActor && static_cast< const AUsdStageActor* >( StageActor )->GetUsdStage() )
+				// Detect if we should actually do anything (check for autosaves, cooking, etc.)
+				if (InContext.GetSaveFlags() & ESaveFlags::SAVE_FromAutosave || InContext.IsProceduralSave())
 				{
-					const bool bForClosing = true;
-					UE::UsdStageEditorModule::Private::SaveStageActorLayersForWorld( StageActor->GetWorld(), bForClosing, StageActor );
+					return;
 				}
-			});
-		});
 
-		// Prompt to save modified USD layers when saving the world
-		PreSaveWorldEditorDelegateHandle = FEditorDelegates::PreSaveWorldWithContext.AddLambda( []( UWorld* World, FObjectPreSaveContext InContext )
-		{
-			// Detect if we should actually do anything (check for autosaves, cooking, etc.)
-			if ( InContext.GetSaveFlags() & ESaveFlags::SAVE_FromAutosave || InContext.IsProceduralSave() )
-			{
-				return;
+				const bool bForClosing = false;
+				UE::UsdStageEditorModule::Private::SaveStageActorLayersForWorld(World, bForClosing);
 			}
+		);
 
-			const bool bForClosing = false;
-			UE::UsdStageEditorModule::Private::SaveStageActorLayersForWorld( World, bForClosing );
-		});
-
-		OpenStageEditorClickedHandle = AUsdStageActor::OnOpenStageEditorClicked.AddLambda( []( AUsdStageActor* StageActor )
-		{
-			FLevelEditorModule& LevelEditorModule = FModuleManager::GetModuleChecked< FLevelEditorModule >( "LevelEditor" );
-			TSharedPtr< FTabManager > LevelEditorTabManager = LevelEditorModule.GetLevelEditorTabManager();
-			if ( TSharedPtr< SDockTab > Tab = LevelEditorTabManager->TryInvokeTab( FTabId{ TEXT( "USDStage" ) } ) )
+		OpenStageEditorClickedHandle = AUsdStageActor::OnOpenStageEditorClicked.AddLambda(
+			[](AUsdStageActor* StageActor)
 			{
-				if ( TSharedPtr<SBorder> ContentBorder = StaticCastSharedRef<SBorder>( Tab->GetContent() ) )
+				FLevelEditorModule& LevelEditorModule = FModuleManager::GetModuleChecked<FLevelEditorModule>("LevelEditor");
+				TSharedPtr<FTabManager> LevelEditorTabManager = LevelEditorModule.GetLevelEditorTabManager();
+				if (TSharedPtr<SDockTab> Tab = LevelEditorTabManager->TryInvokeTab(FTabId{TEXT("USDStage")}))
 				{
-					if ( TSharedPtr<SUsdStage> UsdStageEditor = StaticCastSharedRef<SUsdStage>( ContentBorder->GetContent() ) )
+					if (TSharedPtr<SBorder> ContentBorder = StaticCastSharedRef<SBorder>(Tab->GetContent()))
 					{
-						UsdStageEditor->AttachToStageActor( StageActor );
+						if (TSharedPtr<SUsdStage> UsdStageEditor = StaticCastSharedRef<SUsdStage>(ContentBorder->GetContent()))
+						{
+							UsdStageEditor->AttachToStageActor(StageActor);
+						}
 					}
 				}
 			}
-		});
+		);
 	}
 
 	virtual void ShutdownModule() override
 	{
-		AUsdStageActor::OnOpenStageEditorClicked.Remove( OpenStageEditorClickedHandle );
+		AUsdStageActor::OnOpenStageEditorClicked.Remove(OpenStageEditorClickedHandle);
 
-		FEditorDelegates::PreSaveWorldWithContext.Remove( PreSaveWorldEditorDelegateHandle );
+		FEditorDelegates::PreSaveWorldWithContext.Remove(PreSaveWorldEditorDelegateHandle);
 
-		AUsdStageActor::OnActorLoaded.Remove( StageActorLoadedHandle );
+		AUsdStageActor::OnActorLoaded.Remove(StageActorLoadedHandle);
 
-		IMainFrameModule& MainFrame = FModuleManager::LoadModuleChecked<IMainFrameModule>( "MainFrame" );
-		MainFrame.UnregisterCanCloseEditor( EditorCanCloseDelegate );
+		IMainFrameModule& MainFrame = FModuleManager::LoadModuleChecked<IMainFrameModule>("MainFrame");
+		MainFrame.UnregisterCanCloseEditor(EditorCanCloseDelegate);
 
-		if ( LevelEditorTabManagerChangedHandle.IsValid() && FModuleManager::Get().IsModuleLoaded( "LevelEditor" ) )
+		if (LevelEditorTabManagerChangedHandle.IsValid() && FModuleManager::Get().IsModuleLoaded("LevelEditor"))
 		{
-			FLevelEditorModule& LevelEditorModule = FModuleManager::GetModuleChecked<FLevelEditorModule>( "LevelEditor" );
-			LevelEditorModule.OnTabManagerChanged().Remove( LevelEditorTabManagerChangedHandle );
+			FLevelEditorModule& LevelEditorModule = FModuleManager::GetModuleChecked<FLevelEditorModule>("LevelEditor");
+			LevelEditorModule.OnTabManagerChanged().Remove(LevelEditorTabManagerChangedHandle);
 		}
 
 		FUsdStageEditorStyle::Shutdown();
@@ -420,7 +426,7 @@ private:
 	FDelegateHandle StageActorLoadedHandle;
 	FDelegateHandle OpenStageEditorClickedHandle;
 	FDelegateHandle OnTransactionStateChangedHandle;
-#endif // USE_USD_SDK
+#endif	  // USE_USD_SDK
 };
 
 bool IUsdStageEditorModule::OpenStageEditor() const
@@ -429,19 +435,19 @@ bool IUsdStageEditorModule::OpenStageEditor() const
 	return UE::UsdStageEditorModule::Private::GetUsdStageEditor().IsValid();
 #else
 	return false;
-#endif // USE_USD_SDK
+#endif	  // USE_USD_SDK
 }
 
 bool IUsdStageEditorModule::CloseStageEditor() const
 {
 #if USE_USD_SDK
-	FLevelEditorModule& LevelEditorModule = FModuleManager::GetModuleChecked< FLevelEditorModule >( "LevelEditor" );
-	TSharedPtr< FTabManager > LevelEditorTabManager = LevelEditorModule.GetLevelEditorTabManager();
-	if ( TSharedPtr< SDockTab > Tab = LevelEditorTabManager->FindExistingLiveTab( UE::UsdStageEditorModule::Private::UsdStageEditorTabID ) )
+	FLevelEditorModule& LevelEditorModule = FModuleManager::GetModuleChecked<FLevelEditorModule>("LevelEditor");
+	TSharedPtr<FTabManager> LevelEditorTabManager = LevelEditorModule.GetLevelEditorTabManager();
+	if (TSharedPtr<SDockTab> Tab = LevelEditorTabManager->FindExistingLiveTab(UE::UsdStageEditorModule::Private::UsdStageEditorTabID))
 	{
 		return Tab->RequestCloseTab();
 	}
-#endif // USE_USD_SDK
+#endif	  // USE_USD_SDK
 
 	return false;
 }
@@ -453,30 +459,30 @@ bool IUsdStageEditorModule::IsStageEditorOpened() const
 	return UE::UsdStageEditorModule::Private::GetUsdStageEditor(bOpenIfNeeded).IsValid();
 #else
 	return false;
-#endif // USE_USD_SDK
+#endif	  // USE_USD_SDK
 }
 
 AUsdStageActor* IUsdStageEditorModule::GetAttachedStageActor() const
 {
 #if USE_USD_SDK
-	if ( TSharedPtr<SUsdStage> UsdStageEditor = UE::UsdStageEditorModule::Private::GetUsdStageEditor() )
+	if (TSharedPtr<SUsdStage> UsdStageEditor = UE::UsdStageEditorModule::Private::GetUsdStageEditor())
 	{
 		return UsdStageEditor->GetAttachedStageActor();
 	}
-#endif // USE_USD_SDK
+#endif	  // USE_USD_SDK
 
 	return nullptr;
 }
 
-bool IUsdStageEditorModule::SetAttachedStageActor( AUsdStageActor* NewActor ) const
+bool IUsdStageEditorModule::SetAttachedStageActor(AUsdStageActor* NewActor) const
 {
 #if USE_USD_SDK
-	if ( TSharedPtr<SUsdStage> UsdStageEditor = UE::UsdStageEditorModule::Private::GetUsdStageEditor() )
+	if (TSharedPtr<SUsdStage> UsdStageEditor = UE::UsdStageEditorModule::Private::GetUsdStageEditor())
 	{
-		UsdStageEditor->AttachToStageActor( NewActor );
+		UsdStageEditor->AttachToStageActor(NewActor);
 		return true;
 	}
-#endif // USE_USD_SDK
+#endif	  // USE_USD_SDK
 
 	return false;
 }
@@ -484,67 +490,67 @@ bool IUsdStageEditorModule::SetAttachedStageActor( AUsdStageActor* NewActor ) co
 TArray<UE::FSdfLayer> IUsdStageEditorModule::GetSelectedLayers() const
 {
 #if USE_USD_SDK
-	if ( TSharedPtr<SUsdStage> UsdStageEditor = UE::UsdStageEditorModule::Private::GetUsdStageEditor() )
+	if (TSharedPtr<SUsdStage> UsdStageEditor = UE::UsdStageEditorModule::Private::GetUsdStageEditor())
 	{
 		return UsdStageEditor->GetSelectedLayers();
 	}
-#endif // USE_USD_SDK
+#endif	  // USE_USD_SDK
 
 	return {};
 }
 
-void IUsdStageEditorModule::SetSelectedLayers( const TArray<UE::FSdfLayer>& NewSelection ) const
+void IUsdStageEditorModule::SetSelectedLayers(const TArray<UE::FSdfLayer>& NewSelection) const
 {
 #if USE_USD_SDK
-	if ( TSharedPtr<SUsdStage> UsdStageEditor = UE::UsdStageEditorModule::Private::GetUsdStageEditor() )
+	if (TSharedPtr<SUsdStage> UsdStageEditor = UE::UsdStageEditorModule::Private::GetUsdStageEditor())
 	{
-		return UsdStageEditor->SetSelectedLayers( NewSelection );
+		return UsdStageEditor->SetSelectedLayers(NewSelection);
 	}
-#endif // USE_USD_SDK
+#endif	  // USE_USD_SDK
 }
 
 TArray<UE::FUsdPrim> IUsdStageEditorModule::GetSelectedPrims() const
 {
 #if USE_USD_SDK
-	if ( TSharedPtr<SUsdStage> UsdStageEditor = UE::UsdStageEditorModule::Private::GetUsdStageEditor() )
+	if (TSharedPtr<SUsdStage> UsdStageEditor = UE::UsdStageEditorModule::Private::GetUsdStageEditor())
 	{
 		return UsdStageEditor->GetSelectedPrims();
 	}
-#endif // USE_USD_SDK
+#endif	  // USE_USD_SDK
 
 	return {};
 }
 
-void IUsdStageEditorModule::SetSelectedPrims( const TArray<UE::FUsdPrim>& NewSelection ) const
+void IUsdStageEditorModule::SetSelectedPrims(const TArray<UE::FUsdPrim>& NewSelection) const
 {
 #if USE_USD_SDK
-	if ( TSharedPtr<SUsdStage> UsdStageEditor = UE::UsdStageEditorModule::Private::GetUsdStageEditor() )
+	if (TSharedPtr<SUsdStage> UsdStageEditor = UE::UsdStageEditorModule::Private::GetUsdStageEditor())
 	{
-		UsdStageEditor->SetSelectedPrims( NewSelection );
+		UsdStageEditor->SetSelectedPrims(NewSelection);
 	}
-#endif // USE_USD_SDK
+#endif	  // USE_USD_SDK
 }
 
 TArray<FString> IUsdStageEditorModule::GetSelectedPropertyNames() const
 {
 #if USE_USD_SDK
-	if ( TSharedPtr<SUsdStage> UsdStageEditor = UE::UsdStageEditorModule::Private::GetUsdStageEditor() )
+	if (TSharedPtr<SUsdStage> UsdStageEditor = UE::UsdStageEditorModule::Private::GetUsdStageEditor())
 	{
 		return UsdStageEditor->GetSelectedPropertyNames();
 	}
-#endif // USE_USD_SDK
+#endif	  // USE_USD_SDK
 
 	return {};
 }
 
-void IUsdStageEditorModule::SetSelectedPropertyNames( const TArray<FString>& NewSelection ) const
+void IUsdStageEditorModule::SetSelectedPropertyNames(const TArray<FString>& NewSelection) const
 {
 #if USE_USD_SDK
-	if ( TSharedPtr<SUsdStage> UsdStageEditor = UE::UsdStageEditorModule::Private::GetUsdStageEditor() )
+	if (TSharedPtr<SUsdStage> UsdStageEditor = UE::UsdStageEditorModule::Private::GetUsdStageEditor())
 	{
-		UsdStageEditor->SetSelectedPropertyNames( NewSelection );
+		UsdStageEditor->SetSelectedPropertyNames(NewSelection);
 	}
-#endif // USE_USD_SDK
+#endif	  // USE_USD_SDK
 }
 
 TArray<FString> IUsdStageEditorModule::GetSelectedPropertyMetadataNames() const
@@ -554,7 +560,7 @@ TArray<FString> IUsdStageEditorModule::GetSelectedPropertyMetadataNames() const
 	{
 		return UsdStageEditor->GetSelectedPropertyMetadataNames();
 	}
-#endif // USE_USD_SDK
+#endif	  // USE_USD_SDK
 
 	return {};
 }
@@ -566,116 +572,116 @@ void IUsdStageEditorModule::SetSelectedPropertyMetadataNames(const TArray<FStrin
 	{
 		UsdStageEditor->SetSelectedPropertyMetadataNames(NewSelection);
 	}
-#endif // USE_USD_SDK
+#endif	  // USE_USD_SDK
 }
 
 void IUsdStageEditorModule::FileNew() const
 {
 #if USE_USD_SDK
-	if ( TSharedPtr<SUsdStage> UsdStageEditor = UE::UsdStageEditorModule::Private::GetUsdStageEditor() )
+	if (TSharedPtr<SUsdStage> UsdStageEditor = UE::UsdStageEditorModule::Private::GetUsdStageEditor())
 	{
 		UsdStageEditor->FileNew();
 	}
-#endif // USE_USD_SDK
+#endif	  // USE_USD_SDK
 }
 
-void IUsdStageEditorModule::FileOpen( const FString& FilePath ) const
+void IUsdStageEditorModule::FileOpen(const FString& FilePath) const
 {
 #if USE_USD_SDK
-	if ( TSharedPtr<SUsdStage> UsdStageEditor = UE::UsdStageEditorModule::Private::GetUsdStageEditor() )
+	if (TSharedPtr<SUsdStage> UsdStageEditor = UE::UsdStageEditorModule::Private::GetUsdStageEditor())
 	{
-		UsdStageEditor->FileOpen( FilePath );
+		UsdStageEditor->FileOpen(FilePath);
 	}
-#endif // USE_USD_SDK
+#endif	  // USE_USD_SDK
 }
 
-void IUsdStageEditorModule::FileSave( const FString& OutputFilePathIfUnsaved ) const
+void IUsdStageEditorModule::FileSave(const FString& OutputFilePathIfUnsaved) const
 {
 #if USE_USD_SDK
-	if ( TSharedPtr<SUsdStage> UsdStageEditor = UE::UsdStageEditorModule::Private::GetUsdStageEditor() )
+	if (TSharedPtr<SUsdStage> UsdStageEditor = UE::UsdStageEditorModule::Private::GetUsdStageEditor())
 	{
-		UsdStageEditor->FileSave( OutputFilePathIfUnsaved );
+		UsdStageEditor->FileSave(OutputFilePathIfUnsaved);
 	}
-#endif // USE_USD_SDK
+#endif	  // USE_USD_SDK
 }
 
-void IUsdStageEditorModule::FileExportAllLayers( const FString& OutputDirectory ) const
+void IUsdStageEditorModule::FileExportAllLayers(const FString& OutputDirectory) const
 {
 #if USE_USD_SDK
-	if ( TSharedPtr<SUsdStage> UsdStageEditor = UE::UsdStageEditorModule::Private::GetUsdStageEditor() )
+	if (TSharedPtr<SUsdStage> UsdStageEditor = UE::UsdStageEditorModule::Private::GetUsdStageEditor())
 	{
-		UsdStageEditor->FileExportAllLayers( OutputDirectory  );
+		UsdStageEditor->FileExportAllLayers(OutputDirectory);
 	}
-#endif // USE_USD_SDK
+#endif	  // USE_USD_SDK
 }
 
-void IUsdStageEditorModule::FileExportFlattenedStage( const FString& OutputLayer ) const
+void IUsdStageEditorModule::FileExportFlattenedStage(const FString& OutputLayer) const
 {
 #if USE_USD_SDK
-	if ( TSharedPtr<SUsdStage> UsdStageEditor = UE::UsdStageEditorModule::Private::GetUsdStageEditor() )
+	if (TSharedPtr<SUsdStage> UsdStageEditor = UE::UsdStageEditorModule::Private::GetUsdStageEditor())
 	{
-		UsdStageEditor->FileExportFlattenedStage( OutputLayer  );
+		UsdStageEditor->FileExportFlattenedStage(OutputLayer);
 	}
-#endif // USE_USD_SDK
+#endif	  // USE_USD_SDK
 }
 
 void IUsdStageEditorModule::FileReload() const
 {
 #if USE_USD_SDK
-	if ( TSharedPtr<SUsdStage> UsdStageEditor = UE::UsdStageEditorModule::Private::GetUsdStageEditor() )
+	if (TSharedPtr<SUsdStage> UsdStageEditor = UE::UsdStageEditorModule::Private::GetUsdStageEditor())
 	{
 		UsdStageEditor->FileReload();
 	}
-#endif // USE_USD_SDK
+#endif	  // USE_USD_SDK
 }
 
 void IUsdStageEditorModule::FileReset() const
 {
 #if USE_USD_SDK
-	if ( TSharedPtr<SUsdStage> UsdStageEditor = UE::UsdStageEditorModule::Private::GetUsdStageEditor() )
+	if (TSharedPtr<SUsdStage> UsdStageEditor = UE::UsdStageEditorModule::Private::GetUsdStageEditor())
 	{
 		UsdStageEditor->FileReset();
 	}
-#endif // USE_USD_SDK
+#endif	  // USE_USD_SDK
 }
 
 void IUsdStageEditorModule::FileClose() const
 {
 #if USE_USD_SDK
-	if ( TSharedPtr<SUsdStage> UsdStageEditor = UE::UsdStageEditorModule::Private::GetUsdStageEditor() )
+	if (TSharedPtr<SUsdStage> UsdStageEditor = UE::UsdStageEditorModule::Private::GetUsdStageEditor())
 	{
 		UsdStageEditor->FileClose();
 	}
-#endif // USE_USD_SDK
+#endif	  // USE_USD_SDK
 }
 
-void IUsdStageEditorModule::ActionsImport( const FString& OutputContentFolder, UUsdStageImportOptions* Options ) const
+void IUsdStageEditorModule::ActionsImport(const FString& OutputContentFolder, UUsdStageImportOptions* Options) const
 {
 #if USE_USD_SDK
-	if ( TSharedPtr<SUsdStage> UsdStageEditor = UE::UsdStageEditorModule::Private::GetUsdStageEditor() )
+	if (TSharedPtr<SUsdStage> UsdStageEditor = UE::UsdStageEditorModule::Private::GetUsdStageEditor())
 	{
-		if ( OutputContentFolder.IsEmpty() )
+		if (OutputContentFolder.IsEmpty())
 		{
 			UsdStageEditor->ActionsImportWithDialog();
 		}
 		else
 		{
-			UsdStageEditor->ActionsImport( OutputContentFolder, Options );
+			UsdStageEditor->ActionsImport(OutputContentFolder, Options);
 		}
 	}
-#endif // USE_USD_SDK
+#endif	  // USE_USD_SDK
 }
 
-void IUsdStageEditorModule::ExportSelectedLayers( const FString& OutputLayerOrDirectory ) const
+void IUsdStageEditorModule::ExportSelectedLayers(const FString& OutputLayerOrDirectory) const
 {
 #if USE_USD_SDK
-	if ( TSharedPtr<SUsdStage> UsdStageEditor = UE::UsdStageEditorModule::Private::GetUsdStageEditor() )
+	if (TSharedPtr<SUsdStage> UsdStageEditor = UE::UsdStageEditorModule::Private::GetUsdStageEditor())
 	{
-		UsdStageEditor->ExportSelectedLayers( OutputLayerOrDirectory );
+		UsdStageEditor->ExportSelectedLayers(OutputLayerOrDirectory);
 	}
-#endif // USE_USD_SDK
+#endif	  // USE_USD_SDK
 }
 
-IMPLEMENT_MODULE_USD( FUsdStageEditorModule, USDStageEditor );
+IMPLEMENT_MODULE_USD(FUsdStageEditorModule, USDStageEditor);
 
 #undef LOCTEXT_NAMESPACE
diff --git a/Engine/Plugins/Importers/USDImporter/Source/USDStageEditor/Private/USDStageEditorSettings.h b/Engine/Plugins/Importers/USDImporter/Source/USDStageEditor/Private/USDStageEditorSettings.h
--- a/Engine/Plugins/Importers/USDImporter/Source/USDStageEditor/Private/USDStageEditorSettings.h
+++ b/Engine/Plugins/Importers/USDImporter/Source/USDStageEditor/Private/USDStageEditorSettings.h
@@ -3,23 +3,20 @@
 #pragma once
 
 #include "CoreMinimal.h"
-#include "UObject/ObjectMacros.h"
-#include "UObject/Object.h"
-#include "Engine/EngineTypes.h"
 
 #include "USDStageEditorSettings.generated.h"
 
-UCLASS(config=EditorPerProjectUserSettings)
+UCLASS(config = EditorPerProjectUserSettings)
 class UUsdStageEditorSettings : public UObject
 {
 	GENERATED_BODY()
 
 public:
 	/** Whether our prim selection in SUSDStageTreeView is kept synchronized with the viewport selection */
-	UPROPERTY( config, EditAnywhere, Category = USD )
+	UPROPERTY(config, EditAnywhere, Category = USD)
 	bool bSelectionSynced = false;
 
 	/** Whether to automatically set a layer as edit target when isolating it */
-	UPROPERTY( config, EditAnywhere, Category = USD )
+	UPROPERTY(config, EditAnywhere, Category = USD)
 	bool bIsolateLayerSyncedWithEditTarget = true;
-};
\ No newline at end of file
+};
diff --git a/Engine/Plugins/Importers/USDImporter/Source/USDStageEditor/Private/Widgets/SUSDIntegrationsPanel.cpp b/Engine/Plugins/Importers/USDImporter/Source/USDStageEditor/Private/Widgets/SUSDIntegrationsPanel.cpp
--- a/Engine/Plugins/Importers/USDImporter/Source/USDStageEditor/Private/Widgets/SUSDIntegrationsPanel.cpp
+++ b/Engine/Plugins/Importers/USDImporter/Source/USDStageEditor/Private/Widgets/SUSDIntegrationsPanel.cpp
@@ -11,10 +11,8 @@
 #include "USDValueConversion.h"
 
 #include "UsdWrappers/SdfChangeBlock.h"
-#include "UsdWrappers/SdfPath.h"
 #include "UsdWrappers/UsdAttribute.h"
 #include "UsdWrappers/UsdPrim.h"
-#include "UsdWrappers/UsdStage.h"
 #include "UsdWrappers/VtValue.h"
 
 #include "ISinglePropertyView.h"
@@ -26,75 +24,118 @@
 #include "SLiveLinkSubjectRepresentationPicker.h"
 #include "Styling/AppStyle.h"
 #include "Widgets/Input/SCheckBox.h"
-#include "Widgets/Input/SEditableTextBox.h"
 #include "Widgets/Input/SNumericEntryBox.h"
-#include "Widgets/Input/STextComboBox.h"
 #include "Widgets/Views/SHeaderRow.h"
 
 #define LOCTEXT_NAMESPACE "SUSDIntegrationsPanel"
 
 namespace UE::SUsdIntergrationsPanel::Private
 {
-	const FMargin LeftRowPadding( 6.0f, 0.0f, 2.0f, 0.0f );
-	const FMargin RightRowPadding( 3.0f, 0.0f, 2.0f, 0.0f );
-	const float DesiredNumericEntryBoxWidth = 80.0f;
+	const FMargin LeftRowPadding(6.0f, 0.0f, 2.0f, 0.0f);
+	const FMargin RightRowPadding(3.0f, 0.0f, 2.0f, 0.0f);
 
-	const TCHAR* NormalFont = TEXT( "PropertyWindow.NormalFont" );
+	const TCHAR* NormalFont = TEXT("PropertyWindow.NormalFont");
 }
 
-void SUsdIntegrationsPanelRow::Construct( const FArguments& InArgs, TSharedPtr<UE::FUsdAttribute> InAttr, const TSharedRef< STableViewBase >& OwnerTable )
+void SUsdIntegrationsPanelRow::Construct(const FArguments& InArgs, TSharedPtr<UE::FUsdAttribute> InAttr, const TSharedRef<STableViewBase>& OwnerTable)
 {
 	Attribute = InAttr;
 
-	SMultiColumnTableRow< TSharedPtr<UE::FUsdAttribute> >::Construct( SMultiColumnTableRow< TSharedPtr<UE::FUsdAttribute> >::FArguments(), OwnerTable );
+	SMultiColumnTableRow<TSharedPtr<UE::FUsdAttribute>>::Construct(SMultiColumnTableRow<TSharedPtr<UE::FUsdAttribute>>::FArguments(), OwnerTable);
 }
 
-TSharedRef< SWidget > SUsdIntegrationsPanelRow::GenerateWidgetForColumn( const FName& ColumnName )
+TSharedRef<SWidget> SUsdIntegrationsPanelRow::GenerateWidgetForColumn(const FName& ColumnName)
 {
-	TSharedRef< SWidget > ColumnWidget = SNullWidget::NullWidget;
+	TSharedRef<SWidget> ColumnWidget = SNullWidget::NullWidget;
 
 	bool bIsLeftColumn = true;
-	FOptionalSize RowHeight = FUsdStageEditorStyle::Get()->GetFloat( "UsdStageEditor.ListItemHeight" );
+	FOptionalSize RowHeight = FUsdStageEditorStyle::Get()->GetFloat("UsdStageEditor.ListItemHeight");
 
 	FName AttributeName = Attribute->GetName();
 
 	using DisplayTextForPropertiesEntry = TPairInitializer<const FName&, const FText&>;
-	const static TMap<FName, FText> DisplayTextForProperties
-	({
-		DisplayTextForPropertiesEntry( *UsdToUnreal::ConvertToken( UnrealIdentifiers::UnrealAnimBlueprintPath ), LOCTEXT( "AnimBlueprintPathText", "Anim Blueprint asset" ) ),
-		DisplayTextForPropertiesEntry( *UsdToUnreal::ConvertToken( UnrealIdentifiers::UnrealControlRigPath ), LOCTEXT( "ControlRigPathText", "Control Rig asset" ) ),
-		DisplayTextForPropertiesEntry( *UsdToUnreal::ConvertToken( UnrealIdentifiers::UnrealUseFKControlRig ), LOCTEXT( "UseFKControlRigText", "Use FKControlRig" ) ),
-		DisplayTextForPropertiesEntry( *UsdToUnreal::ConvertToken( UnrealIdentifiers::UnrealControlRigReduceKeys ), LOCTEXT( "ControlRigReduceKeysText", "Control Rig key reduction" ) ),
-		DisplayTextForPropertiesEntry( *UsdToUnreal::ConvertToken( UnrealIdentifiers::UnrealControlRigReductionTolerance ), LOCTEXT( "ControlRigReduceToleranceText", "Control Rig key reduction tolerance" ) ),
-		DisplayTextForPropertiesEntry( *UsdToUnreal::ConvertToken( UnrealIdentifiers::UnrealLiveLinkSubjectName ), LOCTEXT( "SubjectNameText", "Live Link subject name" ) ),
-		DisplayTextForPropertiesEntry( *UsdToUnreal::ConvertToken( UnrealIdentifiers::UnrealLiveLinkEnabled ), LOCTEXT( "LiveLinkEnabledText", "Enable LiveLink" ) ),
+	const static TMap<FName, FText> DisplayTextForProperties({
+		DisplayTextForPropertiesEntry(
+			*UsdToUnreal::ConvertToken(UnrealIdentifiers::UnrealAnimBlueprintPath),
+			LOCTEXT("AnimBlueprintPathText", "Anim Blueprint asset")
+		),
+		DisplayTextForPropertiesEntry(
+			*UsdToUnreal::ConvertToken(UnrealIdentifiers::UnrealControlRigPath),
+			LOCTEXT("ControlRigPathText", "Control Rig asset")
+		),
+		DisplayTextForPropertiesEntry(
+			*UsdToUnreal::ConvertToken(UnrealIdentifiers::UnrealUseFKControlRig),
+			LOCTEXT("UseFKControlRigText", "Use FKControlRig")
+		),
+		DisplayTextForPropertiesEntry(
+			*UsdToUnreal::ConvertToken(UnrealIdentifiers::UnrealControlRigReduceKeys),
+			LOCTEXT("ControlRigReduceKeysText", "Control Rig key reduction")
+		),
+		DisplayTextForPropertiesEntry(
+			*UsdToUnreal::ConvertToken(UnrealIdentifiers::UnrealControlRigReductionTolerance),
+			LOCTEXT("ControlRigReduceToleranceText", "Control Rig key reduction tolerance")
+		),
+		DisplayTextForPropertiesEntry(
+			*UsdToUnreal::ConvertToken(UnrealIdentifiers::UnrealLiveLinkSubjectName),
+			LOCTEXT("SubjectNameText", "Live Link subject name")
+		),
+		DisplayTextForPropertiesEntry(
+			*UsdToUnreal::ConvertToken(UnrealIdentifiers::UnrealLiveLinkEnabled),
+			LOCTEXT("LiveLinkEnabledText", "Enable LiveLink")
+		),
 	});
 
-	const static TMap<FName, FText> ToolTipTextForProperties
-	({
-		DisplayTextForPropertiesEntry( *UsdToUnreal::ConvertToken( UnrealIdentifiers::UnrealAnimBlueprintPath ), LOCTEXT( "AnimBlueprintPathToolTip", "Anim Blueprint asset to use on the component to connect with Live Link" ) ),
-		DisplayTextForPropertiesEntry( *UsdToUnreal::ConvertToken( UnrealIdentifiers::UnrealControlRigPath ), LOCTEXT( "ControlRigPathToolTip", "Control Rig Blueprint asset to use" ) ),
-		DisplayTextForPropertiesEntry( *UsdToUnreal::ConvertToken( UnrealIdentifiers::UnrealUseFKControlRig ), LOCTEXT( "UseFKControlRigToolTip", "Whether to use a generated FKControlRig instead of the Control Rig Blueprint asset" ) ),
-		DisplayTextForPropertiesEntry( *UsdToUnreal::ConvertToken( UnrealIdentifiers::UnrealControlRigReduceKeys ), LOCTEXT( "ControlRigReduceKeysToolTip", "Whether to enable key reduction when generating Control Rig tracks from the USD animation data" ) ),
-		DisplayTextForPropertiesEntry( *UsdToUnreal::ConvertToken( UnrealIdentifiers::UnrealControlRigReductionTolerance ), LOCTEXT( "ControlRigReduceToleranceToolTip", "Tolerance to use for the key reduction pass" ) ),
-		DisplayTextForPropertiesEntry( *UsdToUnreal::ConvertToken( UnrealIdentifiers::UnrealLiveLinkSubjectName ), LOCTEXT( "SubjectNameToolTip", "Which Live Link subject to use for this component" ) ),
-		DisplayTextForPropertiesEntry( *UsdToUnreal::ConvertToken( UnrealIdentifiers::UnrealLiveLinkEnabled ), LOCTEXT( "LiveLinkEnabledToolTip", "If checked will cause the component to be animated with the pose data coming from the Anim Blueprint. If unchecked will cause the component to be animated with the regular skeleton animation data coming from the USD Stage" ) ),
+	const static TMap<FName, FText> ToolTipTextForProperties({
+		DisplayTextForPropertiesEntry(
+			*UsdToUnreal::ConvertToken(UnrealIdentifiers::UnrealAnimBlueprintPath),
+			LOCTEXT("AnimBlueprintPathToolTip", "Anim Blueprint asset to use on the component to connect with Live Link")
+		),
+		DisplayTextForPropertiesEntry(
+			*UsdToUnreal::ConvertToken(UnrealIdentifiers::UnrealControlRigPath),
+			LOCTEXT("ControlRigPathToolTip", "Control Rig Blueprint asset to use")
+		),
+		DisplayTextForPropertiesEntry(
+			*UsdToUnreal::ConvertToken(UnrealIdentifiers::UnrealUseFKControlRig),
+			LOCTEXT("UseFKControlRigToolTip", "Whether to use a generated FKControlRig instead of the Control Rig Blueprint asset")
+		),
+		DisplayTextForPropertiesEntry(
+			*UsdToUnreal::ConvertToken(UnrealIdentifiers::UnrealControlRigReduceKeys),
+			LOCTEXT("ControlRigReduceKeysToolTip", "Whether to enable key reduction when generating Control Rig tracks from the USD animation data")
+		),
+		DisplayTextForPropertiesEntry(
+			*UsdToUnreal::ConvertToken(UnrealIdentifiers::UnrealControlRigReductionTolerance),
+			LOCTEXT("ControlRigReduceToleranceToolTip", "Tolerance to use for the key reduction pass")
+		),
+		DisplayTextForPropertiesEntry(
+			*UsdToUnreal::ConvertToken(UnrealIdentifiers::UnrealLiveLinkSubjectName),
+			LOCTEXT("SubjectNameToolTip", "Which Live Link subject to use for this component")
+		),
+		DisplayTextForPropertiesEntry(
+			*UsdToUnreal::ConvertToken(UnrealIdentifiers::UnrealLiveLinkEnabled),
+			LOCTEXT(
+				"LiveLinkEnabledToolTip",
+				"If checked will cause the component to be animated with the pose data coming from the Anim Blueprint. If unchecked will cause the "
+				"component to be animated with the regular skeleton animation data coming from the USD Stage"
+			)
+		),
 	});
 
-	const FText* ToolTipText = ToolTipTextForProperties.Find( AttributeName );
-	if ( !ToolTipText )
+	const FText* ToolTipText = ToolTipTextForProperties.Find(AttributeName);
+	if (!ToolTipText)
 	{
 		ToolTipText = &FText::GetEmpty();
 	}
 
-	if ( ColumnName == TEXT("PropertyName") )
+	if (ColumnName == TEXT("PropertyName"))
 	{
-		if ( const FText* TextToDisplay = DisplayTextForProperties.Find( AttributeName ) )
+		if (const FText* TextToDisplay = DisplayTextForProperties.Find(AttributeName))
 		{
-			SAssignNew( ColumnWidget, STextBlock )
-			.Text( *TextToDisplay )
-			.Font( FAppStyle::GetFontStyle( UE::SUsdIntergrationsPanel::Private::NormalFont ) )
-			.ToolTipText( *ToolTipText );
+			// clang-format off
+			SAssignNew(ColumnWidget, STextBlock)
+			.Text(*TextToDisplay)
+			.Font(FAppStyle::GetFontStyle(UE::SUsdIntergrationsPanel::Private::NormalFont))
+			.ToolTipText(*ToolTipText);
+			// clang-format on
 		}
 		else
 		{
@@ -106,71 +147,70 @@ TSharedRef< SWidget > SUsdIntegrationsPanelRow::GenerateWidgetForColumn( const F
 	{
 		bIsLeftColumn = false;
 
-		FName TypeName = Attribute->GetTypeName();
 		TSharedPtr<UE::FUsdAttribute> AttributeCopy = Attribute;
-		if ( !AttributeCopy || !*AttributeCopy )
+		if (!AttributeCopy || !*AttributeCopy)
 		{
 			return ColumnWidget;
 		}
 
-		if ( AttributeName == *UsdToUnreal::ConvertToken( UnrealIdentifiers::UnrealAnimBlueprintPath ) )
+		if (AttributeName == *UsdToUnreal::ConvertToken(UnrealIdentifiers::UnrealAnimBlueprintPath))
 		{
-			if ( UUsdIntegrationsPanelPropertyDummy* Dummy = GetMutableDefault<UUsdIntegrationsPanelPropertyDummy>() )
+			if (UUsdIntegrationsPanelPropertyDummy* Dummy = GetMutableDefault<UUsdIntegrationsPanelPropertyDummy>())
 			{
 				// Let the object picker row be as tall as it wants
 				RowHeight = FOptionalSize();
 
 				FSinglePropertyParams Params;
-				Params.Font = FAppStyle::GetFontStyle( UE::SUsdIntergrationsPanel::Private::NormalFont );
+				Params.Font = FAppStyle::GetFontStyle(UE::SUsdIntergrationsPanel::Private::NormalFont);
 				Params.NamePlacement = EPropertyNamePlacement::Hidden;
 
 				UE::FVtValue Value;
-				if ( AttributeCopy->Get( Value ) && !Value.IsEmpty() )
+				if (AttributeCopy->Get(Value) && !Value.IsEmpty())
 				{
-					Dummy->AnimBPProperty = FSoftObjectPath( UsdUtils::Stringify( Value ) ).TryLoad();
+					Dummy->AnimBPProperty = FSoftObjectPath(UsdUtils::Stringify(Value)).TryLoad();
 				}
 				else
 				{
 					Dummy->AnimBPProperty = nullptr;
 				}
 
-				FPropertyEditorModule& PropertyEditor = FModuleManager::LoadModuleChecked<FPropertyEditorModule>( TEXT( "PropertyEditor" ) );
+				FPropertyEditorModule& PropertyEditor = FModuleManager::LoadModuleChecked<FPropertyEditorModule>(TEXT("PropertyEditor"));
 				TSharedPtr<class ISinglePropertyView> PropertyView = PropertyEditor.CreateSingleProperty(
 					Dummy,
-					GET_MEMBER_NAME_CHECKED( UUsdIntegrationsPanelPropertyDummy, AnimBPProperty ),
+					GET_MEMBER_NAME_CHECKED(UUsdIntegrationsPanelPropertyDummy, AnimBPProperty),
 					Params
 				);
 
-				FSimpleDelegate PropertyChanged = FSimpleDelegate::CreateLambda( [AttributeCopy]()
-				{
-					UE::FUsdPrim Prim = AttributeCopy->GetPrim();
-					if (UsdUtils::NotifyIfInstanceProxy(Prim))
+				FSimpleDelegate PropertyChanged = FSimpleDelegate::CreateLambda(
+					[AttributeCopy]()
 					{
-						return;
-					}
+						UE::FUsdPrim Prim = AttributeCopy->GetPrim();
+						if (UsdUtils::NotifyIfInstanceProxy(Prim))
+						{
+							return;
+						}
 
-					UUsdIntegrationsPanelPropertyDummy* Dummy = GetMutableDefault<UUsdIntegrationsPanelPropertyDummy>();
+						UUsdIntegrationsPanelPropertyDummy* Dummy = GetMutableDefault<UUsdIntegrationsPanelPropertyDummy>();
 
-					FString Path = ( Dummy && Dummy->AnimBPProperty ) ? Dummy->AnimBPProperty->GetPathName() : FString{};
-					std::string UsdPath = UnrealToUsd::ConvertString(*Path).Get();
+						FString Path = (Dummy && Dummy->AnimBPProperty) ? Dummy->AnimBPProperty->GetPathName() : FString{};
+						std::string UsdPath = UnrealToUsd::ConvertString(*Path).Get();
 
-					// We need transactions for these because the notices emitted for them are upgraded to resyncs by the stage
-					// actor, which means they may generate new assets and components
-					FScopedTransaction Transaction(
-						FText::Format( LOCTEXT( "AnimBluepringTransaction", "Changed AnimBluepring path to '{0}'" ),
-							FText::FromString( Path )
-						)
-					);
+						// We need transactions for these because the notices emitted for them are upgraded to resyncs by the stage
+						// actor, which means they may generate new assets and components
+						FScopedTransaction Transaction(
+							FText::Format(LOCTEXT("AnimBluepringTransaction", "Changed AnimBluepring path to '{0}'"), FText::FromString(Path))
+						);
 
-					UE::FVtValue Value;
-					if ( AttributeCopy && *AttributeCopy && UsdUtils::SetUnderlyingValue( Value, UsdPath ) && !Value.IsEmpty() )
-					{
-						UE::FSdfChangeBlock Block;
-						AttributeCopy->Set( Value );
-						UsdUtils::NotifyIfOverriddenOpinion( *AttributeCopy );
+						UE::FVtValue Value;
+						if (AttributeCopy && *AttributeCopy && UsdUtils::SetUnderlyingValue(Value, UsdPath) && !Value.IsEmpty())
+						{
+							UE::FSdfChangeBlock Block;
+							AttributeCopy->Set(Value);
+							UsdUtils::NotifyIfOverriddenOpinion(*AttributeCopy);
+						}
 					}
-				} );
-				PropertyView->SetOnPropertyValueChanged( PropertyChanged );
+				);
+				PropertyView->SetOnPropertyValueChanged(PropertyChanged);
 
 				ColumnWidget = PropertyView.ToSharedRef();
 			}
@@ -179,75 +219,77 @@ TSharedRef< SWidget > SUsdIntegrationsPanelRow::GenerateWidgetForColumn( const F
 				ColumnWidget = SNullWidget::NullWidget;
 			}
 		}
-		else if ( AttributeName == *UsdToUnreal::ConvertToken( UnrealIdentifiers::UnrealControlRigPath ) )
+		else if (AttributeName == *UsdToUnreal::ConvertToken(UnrealIdentifiers::UnrealControlRigPath))
 		{
-			if ( UUsdIntegrationsPanelPropertyDummy* Dummy = GetMutableDefault<UUsdIntegrationsPanelPropertyDummy>() )
+			if (UUsdIntegrationsPanelPropertyDummy* Dummy = GetMutableDefault<UUsdIntegrationsPanelPropertyDummy>())
 			{
 				// Let the object picker row be as tall as it wants
 				RowHeight = FOptionalSize();
 
 				FSinglePropertyParams Params;
-				Params.Font = FAppStyle::GetFontStyle( UE::SUsdIntergrationsPanel::Private::NormalFont );
+				Params.Font = FAppStyle::GetFontStyle(UE::SUsdIntergrationsPanel::Private::NormalFont);
 				Params.NamePlacement = EPropertyNamePlacement::Hidden;
 
 				UE::FVtValue Value;
-				if ( AttributeCopy->Get( Value ) && !Value.IsEmpty() )
+				if (AttributeCopy->Get(Value) && !Value.IsEmpty())
 				{
-					Dummy->ControlRigProperty = FSoftObjectPath( UsdUtils::Stringify( Value ) ).TryLoad();
+					Dummy->ControlRigProperty = FSoftObjectPath(UsdUtils::Stringify(Value)).TryLoad();
 				}
 				else
 				{
 					Dummy->ControlRigProperty = nullptr;
 				}
 
-				FPropertyEditorModule& PropertyEditor = FModuleManager::LoadModuleChecked<FPropertyEditorModule>( TEXT( "PropertyEditor" ) );
+				FPropertyEditorModule& PropertyEditor = FModuleManager::LoadModuleChecked<FPropertyEditorModule>(TEXT("PropertyEditor"));
 				TSharedPtr<class ISinglePropertyView> PropertyView = PropertyEditor.CreateSingleProperty(
 					Dummy,
-					GET_MEMBER_NAME_CHECKED( UUsdIntegrationsPanelPropertyDummy, ControlRigProperty ),
+					GET_MEMBER_NAME_CHECKED(UUsdIntegrationsPanelPropertyDummy, ControlRigProperty),
 					Params
 				);
 
-				FSimpleDelegate PropertyChanged = FSimpleDelegate::CreateLambda( [AttributeCopy]()
-				{
-					UE::FUsdPrim Prim = AttributeCopy->GetPrim();
-					if (UsdUtils::NotifyIfInstanceProxy(Prim))
+				FSimpleDelegate PropertyChanged = FSimpleDelegate::CreateLambda(
+					[AttributeCopy]()
 					{
-						return;
-					}
+						UE::FUsdPrim Prim = AttributeCopy->GetPrim();
+						if (UsdUtils::NotifyIfInstanceProxy(Prim))
+						{
+							return;
+						}
 
-					UUsdIntegrationsPanelPropertyDummy* Dummy = GetMutableDefault<UUsdIntegrationsPanelPropertyDummy>();
+						UUsdIntegrationsPanelPropertyDummy* Dummy = GetMutableDefault<UUsdIntegrationsPanelPropertyDummy>();
 
-					FString Path = ( Dummy && Dummy->ControlRigProperty ) ? Dummy->ControlRigProperty->GetPathName() : FString{};
-					std::string UsdPath = UnrealToUsd::ConvertString( *Path ).Get();
+						FString Path = (Dummy && Dummy->ControlRigProperty) ? Dummy->ControlRigProperty->GetPathName() : FString{};
+						std::string UsdPath = UnrealToUsd::ConvertString(*Path).Get();
 
-					FScopedTransaction Transaction(
-						FText::Format( LOCTEXT( "ControlRigPathTransaction", "Changed ControlRig blueprint path to '{0}'" ),
-							FText::FromString( Path )
-						)
-					);
+						FScopedTransaction Transaction(
+							FText::Format(LOCTEXT("ControlRigPathTransaction", "Changed ControlRig blueprint path to '{0}'"), FText::FromString(Path))
+						);
 
-					UE::FVtValue Value;
-					if ( AttributeCopy && *AttributeCopy && UsdUtils::SetUnderlyingValue( Value, UsdPath ) && !Value.IsEmpty() )
-					{
-						UE::FSdfChangeBlock Block;
-						AttributeCopy->Set( Value );
-						UsdUtils::NotifyIfOverriddenOpinion( *AttributeCopy );
+						UE::FVtValue Value;
+						if (AttributeCopy && *AttributeCopy && UsdUtils::SetUnderlyingValue(Value, UsdPath) && !Value.IsEmpty())
+						{
+							UE::FSdfChangeBlock Block;
+							AttributeCopy->Set(Value);
+							UsdUtils::NotifyIfOverriddenOpinion(*AttributeCopy);
+						}
 					}
-				} );
-				PropertyView->SetOnPropertyValueChanged( PropertyChanged );
+				);
+				PropertyView->SetOnPropertyValueChanged(PropertyChanged);
 
 				ColumnWidget = PropertyView.ToSharedRef();
 
 				// Disable the widget if we're using an FKControlRig instead
-				if ( UE::FUsdAttribute FKEnabledAttr = Attribute->GetPrim().GetAttribute( *UsdToUnreal::ConvertToken( UnrealIdentifiers::UnrealUseFKControlRig ) ) )
+				if (UE::FUsdAttribute FKEnabledAttr = Attribute->GetPrim().GetAttribute(
+						*UsdToUnreal::ConvertToken(UnrealIdentifiers::UnrealUseFKControlRig)
+					))
 				{
 					UE::FVtValue FKEnabledValue;
-					if ( FKEnabledAttr.Get( FKEnabledValue ) )
+					if (FKEnabledAttr.Get(FKEnabledValue))
 					{
 						const bool bDefaultValue = false;
-						if ( UsdUtils::GetUnderlyingValue<bool>( FKEnabledValue ).Get( bDefaultValue ) )
+						if (UsdUtils::GetUnderlyingValue<bool>(FKEnabledValue).Get(bDefaultValue))
 						{
-							ColumnWidget->SetEnabled( false );
+							ColumnWidget->SetEnabled(false);
 						}
 					}
 				}
@@ -257,19 +299,20 @@ TSharedRef< SWidget > SUsdIntegrationsPanelRow::GenerateWidgetForColumn( const F
 				ColumnWidget = SNullWidget::NullWidget;
 			}
 		}
-		else if ( AttributeName == *UsdToUnreal::ConvertToken( UnrealIdentifiers::UnrealControlRigReduceKeys ) )
+		else if (AttributeName == *UsdToUnreal::ConvertToken(UnrealIdentifiers::UnrealControlRigReduceKeys))
 		{
-			SAssignNew( ColumnWidget, SBox )
-			.HeightOverride( FUsdStageEditorStyle::Get()->GetFloat( "UsdStageEditor.ListItemHeight" ) )
-			.VAlign( VAlign_Center )
+			// clang-format off
+			SAssignNew(ColumnWidget, SBox)
+			.HeightOverride(FUsdStageEditorStyle::Get()->GetFloat("UsdStageEditor.ListItemHeight"))
+			.VAlign(VAlign_Center)
 			[
 				SNew(SCheckBox)
-				.IsChecked_Lambda( [AttributeCopy]() ->ECheckBoxState
+				.IsChecked_Lambda([AttributeCopy]() ->ECheckBoxState
 				{
 					UE::FVtValue Value;
-					if ( AttributeCopy && *AttributeCopy && AttributeCopy->Get( Value ) )
+					if (AttributeCopy && *AttributeCopy && AttributeCopy->Get(Value))
 					{
-						if ( TOptional<bool> UnderlyingValue = UsdUtils::GetUnderlyingValue<bool>( Value ) )
+						if (TOptional<bool> UnderlyingValue = UsdUtils::GetUnderlyingValue<bool>(Value))
 						{
 							return UnderlyingValue.GetValue() ? ECheckBoxState::Checked : ECheckBoxState::Unchecked;
 						}
@@ -277,7 +320,7 @@ TSharedRef< SWidget > SUsdIntegrationsPanelRow::GenerateWidgetForColumn( const F
 
 					return ECheckBoxState::Undetermined;
 				})
-				.OnCheckStateChanged_Lambda( [AttributeCopy]( ECheckBoxState NewValue )
+				.OnCheckStateChanged_Lambda([AttributeCopy](ECheckBoxState NewValue)
 				{
 					UE::FUsdPrim Prim = AttributeCopy->GetPrim();
 					if (UsdUtils::NotifyIfInstanceProxy(Prim))
@@ -286,30 +329,32 @@ TSharedRef< SWidget > SUsdIntegrationsPanelRow::GenerateWidgetForColumn( const F
 					}
 
 					FScopedTransaction Transaction(
-						FText::Format( LOCTEXT( "ControlRigReduceKeysTransaction", "Changed ControlRig reduce keys property to '{0}'" ),
-							FText::FromString( NewValue == ECheckBoxState::Checked ? TEXT( "true" ) : TEXT( "false" ) )
+						FText::Format(LOCTEXT("ControlRigReduceKeysTransaction", "Changed ControlRig reduce keys property to '{0}'"),
+							FText::FromString(NewValue == ECheckBoxState::Checked ? TEXT("true") : TEXT("false"))
 						)
 					);
 
 					UE::FVtValue Value;
-					if ( AttributeCopy && *AttributeCopy && UsdUtils::SetUnderlyingValue( Value, NewValue == ECheckBoxState::Checked ) && !Value.IsEmpty() )
+					if (AttributeCopy && *AttributeCopy && UsdUtils::SetUnderlyingValue(Value, NewValue == ECheckBoxState::Checked) && !Value.IsEmpty())
 					{
 						UE::FSdfChangeBlock Block;
-						AttributeCopy->Set( Value );
-						UsdUtils::NotifyIfOverriddenOpinion( *AttributeCopy );
+						AttributeCopy->Set(Value);
+						UsdUtils::NotifyIfOverriddenOpinion(*AttributeCopy);
 					}
 				})
 			];
+			// clang-format on
 		}
-		else if ( AttributeName == *UsdToUnreal::ConvertToken( UnrealIdentifiers::UnrealControlRigReductionTolerance ) )
+		else if (AttributeName == *UsdToUnreal::ConvertToken(UnrealIdentifiers::UnrealControlRigReductionTolerance))
 		{
-			SAssignNew( ColumnWidget, SBox )
-			.HeightOverride( FUsdStageEditorStyle::Get()->GetFloat( "UsdStageEditor.ListItemHeight" ) )
-			.VAlign( VAlign_Center )
+			// clang-format off
+			SAssignNew(ColumnWidget, SBox)
+			.HeightOverride(FUsdStageEditorStyle::Get()->GetFloat("UsdStageEditor.ListItemHeight"))
+			.VAlign(VAlign_Center)
 			[
 				SNew(SNumericEntryBox<float>)
-				.Font( FAppStyle::GetFontStyle( UE::SUsdIntergrationsPanel::Private::NormalFont ) )
-				.OnValueCommitted_Lambda( [AttributeCopy]( float NewValue, ETextCommit::Type CommitType )
+				.Font(FAppStyle::GetFontStyle(UE::SUsdIntergrationsPanel::Private::NormalFont))
+				.OnValueCommitted_Lambda([AttributeCopy](float NewValue, ETextCommit::Type CommitType)
 				{
 					UE::FUsdPrim Prim = AttributeCopy->GetPrim();
 					if (UsdUtils::NotifyIfInstanceProxy(Prim))
@@ -318,25 +363,25 @@ TSharedRef< SWidget > SUsdIntegrationsPanelRow::GenerateWidgetForColumn( const F
 					}
 
 					FScopedTransaction Transaction(
-						FText::Format( LOCTEXT( "ControlRigReduceToleranceTransaction", "Changed ControlRig reduce tolerance to '{0}'" ),
+						FText::Format(LOCTEXT("ControlRigReduceToleranceTransaction", "Changed ControlRig reduce tolerance to '{0}'"),
 							NewValue
 						)
 					);
 
 					UE::FVtValue Value;
-					if ( AttributeCopy && *AttributeCopy && UsdUtils::SetUnderlyingValue( Value, NewValue ) && !Value.IsEmpty() )
+					if (AttributeCopy && *AttributeCopy && UsdUtils::SetUnderlyingValue(Value, NewValue) && !Value.IsEmpty())
 					{
 						UE::FSdfChangeBlock Block;
-						AttributeCopy->Set( Value );
-						UsdUtils::NotifyIfOverriddenOpinion( *AttributeCopy );
+						AttributeCopy->Set(Value);
+						UsdUtils::NotifyIfOverriddenOpinion(*AttributeCopy);
 					}
 				})
-				.Value_Lambda( [AttributeCopy]() -> float
+				.Value_Lambda([AttributeCopy]() -> float
 				{
 					UE::FVtValue Value;
-					if ( AttributeCopy && *AttributeCopy && AttributeCopy->Get( Value ) )
+					if (AttributeCopy && *AttributeCopy && AttributeCopy->Get(Value))
 					{
-						if ( TOptional<float> UnderlyingValue = UsdUtils::GetUnderlyingValue<float>( Value ) )
+						if (TOptional<float> UnderlyingValue = UsdUtils::GetUnderlyingValue<float>(Value))
 						{
 							return UnderlyingValue.GetValue();
 						}
@@ -345,20 +390,22 @@ TSharedRef< SWidget > SUsdIntegrationsPanelRow::GenerateWidgetForColumn( const F
 					return 0.0f;
 				})
 			];
+			// clang-format on
 		}
-		else if ( AttributeName == *UsdToUnreal::ConvertToken( UnrealIdentifiers::UnrealLiveLinkSubjectName ) )
+		else if (AttributeName == *UsdToUnreal::ConvertToken(UnrealIdentifiers::UnrealLiveLinkSubjectName))
 		{
 			// SLiveLinkSubjectRepresentationPicker is also a little bit larger than one of our rows, so let it expand a bit
 			RowHeight = FOptionalSize();
 
-			SAssignNew( ColumnWidget, SLiveLinkSubjectRepresentationPicker )
-			.ShowRole( false )
-			.ShowSource( false )
-			.Font( FAppStyle::GetFontStyle( UE::SUsdIntergrationsPanel::Private::NormalFont ) )
-			.Value_Lambda( [AttributeCopy]()
+			// clang-format off
+			SAssignNew(ColumnWidget, SLiveLinkSubjectRepresentationPicker)
+			.ShowRole(false)
+			.ShowSource(false)
+			.Font(FAppStyle::GetFontStyle(UE::SUsdIntergrationsPanel::Private::NormalFont))
+			.Value_Lambda([AttributeCopy]()
 			{
 				SLiveLinkSubjectRepresentationPicker::FLiveLinkSourceSubjectRole Result;
-				if ( AttributeCopy && *AttributeCopy )
+				if (AttributeCopy && *AttributeCopy)
 				{
 					const UE::FUsdPrim& Prim = AttributeCopy->GetPrim();
 					Result.Role = (Prim.IsA(TEXT("SkelRoot")) || Prim.IsA(TEXT("Skeleton")))
@@ -366,13 +413,13 @@ TSharedRef< SWidget > SUsdIntegrationsPanelRow::GenerateWidgetForColumn( const F
 									  : ULiveLinkTransformRole::StaticClass();
 
 					UE::FVtValue Value;
-					if ( AttributeCopy->Get( Value ) )
+					if (AttributeCopy->Get(Value))
 					{
-						if ( TOptional<std::string> UnderlyingValue = UsdUtils::GetUnderlyingValue<std::string>( Value ) )
+						if (TOptional<std::string> UnderlyingValue = UsdUtils::GetUnderlyingValue<std::string>(Value))
 						{
 							Result.Subject = FLiveLinkSubjectName{
 								FName{
-									*UsdToUnreal::ConvertString( UnderlyingValue.GetValue() )
+									*UsdToUnreal::ConvertString(UnderlyingValue.GetValue())
 								}
 							};
 						}
@@ -381,7 +428,7 @@ TSharedRef< SWidget > SUsdIntegrationsPanelRow::GenerateWidgetForColumn( const F
 
 				return Result;
 			})
-			.OnValueChanged_Lambda( [AttributeCopy]( SLiveLinkSubjectRepresentationPicker::FLiveLinkSourceSubjectRole NewValue )
+			.OnValueChanged_Lambda([AttributeCopy](SLiveLinkSubjectRepresentationPicker::FLiveLinkSourceSubjectRole NewValue)
 			{
 				UE::FUsdPrim Prim = AttributeCopy->GetPrim();
 				if (UsdUtils::NotifyIfInstanceProxy(Prim))
@@ -391,37 +438,38 @@ TSharedRef< SWidget > SUsdIntegrationsPanelRow::GenerateWidgetForColumn( const F
 
 				FScopedUsdAllocs UsdAllocs;
 
-				std::string UsdValue = UnrealToUsd::ConvertString( *( NewValue.Subject.ToString() ) ).Get();
+				std::string UsdValue = UnrealToUsd::ConvertString(*(NewValue.Subject.ToString())).Get();
 
 				FScopedTransaction Transaction(
-					FText::Format( LOCTEXT( "SubjectNameTransaction", "Changed Live Link subject name to '{0}'" ),
-						FText::FromString( *( NewValue.Subject.ToString() ) )
+					FText::Format(LOCTEXT("SubjectNameTransaction", "Changed Live Link subject name to '{0}'"),
+						FText::FromString(*(NewValue.Subject.ToString()))
 					)
 				);
 
 				UE::FVtValue Value;
-				if ( AttributeCopy && *AttributeCopy && UsdUtils::SetUnderlyingValue( Value, UsdValue ) && !Value.IsEmpty() )
+				if (AttributeCopy && *AttributeCopy && UsdUtils::SetUnderlyingValue(Value, UsdValue) && !Value.IsEmpty())
 				{
 					UE::FSdfChangeBlock Block;
-					AttributeCopy->Set( Value );
-					UsdUtils::NotifyIfOverriddenOpinion( *AttributeCopy );
+					AttributeCopy->Set(Value);
+					UsdUtils::NotifyIfOverriddenOpinion(*AttributeCopy);
 				}
 			});
+			// clang-format on
 		}
-		else if ( AttributeName == *UsdToUnreal::ConvertToken( UnrealIdentifiers::UnrealLiveLinkEnabled )
-				|| AttributeName == *UsdToUnreal::ConvertToken( UnrealIdentifiers::UnrealUseFKControlRig ) )
+		else if (AttributeName == *UsdToUnreal::ConvertToken(UnrealIdentifiers::UnrealLiveLinkEnabled) || AttributeName == *UsdToUnreal::ConvertToken(UnrealIdentifiers::UnrealUseFKControlRig))
 		{
-			SAssignNew( ColumnWidget, SBox )
-			.HeightOverride( FUsdStageEditorStyle::Get()->GetFloat( "UsdStageEditor.ListItemHeight" ) )
-			.VAlign( VAlign_Center )
+			// clang-format off
+			SAssignNew(ColumnWidget, SBox)
+			.HeightOverride(FUsdStageEditorStyle::Get()->GetFloat("UsdStageEditor.ListItemHeight"))
+			.VAlign(VAlign_Center)
 			[
-				SNew( SCheckBox )
-				.IsChecked_Lambda( [AttributeCopy]()
+				SNew(SCheckBox)
+				.IsChecked_Lambda([AttributeCopy]()
 				{
 					UE::FVtValue Value;
-					if ( AttributeCopy && *AttributeCopy && AttributeCopy->Get( Value ) )
+					if (AttributeCopy && *AttributeCopy && AttributeCopy->Get(Value))
 					{
-						if ( TOptional<bool> UnderlyingValue = UsdUtils::GetUnderlyingValue<bool>( Value ) )
+						if (TOptional<bool> UnderlyingValue = UsdUtils::GetUnderlyingValue<bool>(Value))
 						{
 							return UnderlyingValue.GetValue() ? ECheckBoxState::Checked : ECheckBoxState::Unchecked;
 						}
@@ -429,7 +477,7 @@ TSharedRef< SWidget > SUsdIntegrationsPanelRow::GenerateWidgetForColumn( const F
 
 					return ECheckBoxState::Undetermined;
 				})
-				.OnCheckStateChanged_Lambda( [AttributeCopy, AttributeName]( ECheckBoxState NewState )
+				.OnCheckStateChanged_Lambda([AttributeCopy, AttributeName](ECheckBoxState NewState)
 				{
 					UE::FUsdPrim Prim = AttributeCopy->GetPrim();
 					if (UsdUtils::NotifyIfInstanceProxy(Prim))
@@ -438,21 +486,22 @@ TSharedRef< SWidget > SUsdIntegrationsPanelRow::GenerateWidgetForColumn( const F
 					}
 
 					FScopedTransaction Transaction(
-						FText::Format( LOCTEXT( "IntegrationBoolTransaction", "Changed attribute '{0}' to " ),
-							FText::FromName( AttributeName ),
-							FText::FromString( NewState == ECheckBoxState::Checked ? TEXT( "true" ) : TEXT( "false" ) )
+						FText::Format(LOCTEXT("IntegrationBoolTransaction", "Changed attribute '{0}' to "),
+							FText::FromName(AttributeName),
+							FText::FromString(NewState == ECheckBoxState::Checked ? TEXT("true") : TEXT("false"))
 						)
 					);
 
 					UE::FVtValue Value;
-					if ( AttributeCopy && *AttributeCopy && UsdUtils::SetUnderlyingValue( Value, NewState == ECheckBoxState::Checked ) && !Value.IsEmpty() )
+					if (AttributeCopy && *AttributeCopy && UsdUtils::SetUnderlyingValue(Value, NewState == ECheckBoxState::Checked) && !Value.IsEmpty())
 					{
 						UE::FSdfChangeBlock Block;
-						AttributeCopy->Set( Value );
-						UsdUtils::NotifyIfOverriddenOpinion( *AttributeCopy );
+						AttributeCopy->Set(Value);
+						UsdUtils::NotifyIfOverriddenOpinion(*AttributeCopy);
 					}
 				})
 			];
+			// clang-format on
 		}
 		else
 		{
@@ -460,57 +509,61 @@ TSharedRef< SWidget > SUsdIntegrationsPanelRow::GenerateWidgetForColumn( const F
 		}
 	}
 
+	// clang-format off
 	return SNew(SBox)
-		.HeightOverride( RowHeight )
+		.HeightOverride(RowHeight)
 		[
-			SNew( SHorizontalBox )
+			SNew(SHorizontalBox)
 			+ SHorizontalBox::Slot()
-			.HAlign( HAlign_Left )
-			.VAlign( VAlign_Center )
-			.Padding( bIsLeftColumn ? UE::SUsdIntergrationsPanel::Private::LeftRowPadding : UE::SUsdIntergrationsPanel::Private::RightRowPadding )
+			.HAlign(HAlign_Left)
+			.VAlign(VAlign_Center)
+			.Padding(bIsLeftColumn ? UE::SUsdIntergrationsPanel::Private::LeftRowPadding : UE::SUsdIntergrationsPanel::Private::RightRowPadding)
 			.AutoWidth()
 			[
 				ColumnWidget
 			]
 		];
+	// clang-format on
 }
 
-void SUsdIntegrationsPanel::Construct( const FArguments& InArgs )
+void SUsdIntegrationsPanel::Construct(const FArguments& InArgs)
 {
-	SAssignNew( HeaderRowWidget, SHeaderRow )
+	// clang-format off
+	SAssignNew(HeaderRowWidget, SHeaderRow)
 
-	+SHeaderRow::Column( FName( TEXT("PropertyName") ) )
-	.DefaultLabel( NSLOCTEXT( "SUsdIntegrationsPanel", "Integrations", "Integrations" ) )
-	.FillWidth( 30.0f )
+	+SHeaderRow::Column(FName(TEXT("PropertyName")))
+	.DefaultLabel(NSLOCTEXT("SUsdIntegrationsPanel", "Integrations", "Integrations"))
+	.FillWidth(30.0f)
 
-	+SHeaderRow::Column( FName( TEXT("Value") ) )
-	.DefaultLabel( FText::GetEmpty() )
-	.FillWidth( 70.0f );
+	+SHeaderRow::Column(FName(TEXT("Value")))
+	.DefaultLabel(FText::GetEmpty())
+	.FillWidth(70.0f);
 
 	SListView::Construct
 	(
 		SListView::FArguments()
-		.ListItemsSource( &ViewModel.Attributes )
-		.OnGenerateRow( this, &SUsdIntegrationsPanel::OnGenerateRow )
-		.HeaderRow( HeaderRowWidget )
+		.ListItemsSource(&ViewModel.Attributes)
+		.OnGenerateRow(this, &SUsdIntegrationsPanel::OnGenerateRow)
+		.HeaderRow(HeaderRowWidget)
 	);
 
-	SetVisibility( EVisibility::Collapsed ); // Start hidden until SetPrimPath displays us
+	SetVisibility(EVisibility::Collapsed); // Start hidden until SetPrimPath displays us
+	// clang-format on
 }
 
-void SUsdIntegrationsPanel::SetPrimPath( const UE::FUsdStageWeak& InUsdStage, const TCHAR* InPrimPath )
+void SUsdIntegrationsPanel::SetPrimPath(const UE::FUsdStageWeak& InUsdStage, const TCHAR* InPrimPath)
 {
-	ViewModel.UpdateAttributes( InUsdStage, InPrimPath );
+	ViewModel.UpdateAttributes(InUsdStage, InPrimPath);
 
-	SetVisibility( ViewModel.Attributes.Num() > 0 ? EVisibility::Visible : EVisibility::Collapsed );
+	SetVisibility(ViewModel.Attributes.Num() > 0 ? EVisibility::Visible : EVisibility::Collapsed);
 
 	RequestListRefresh();
 }
 
-TSharedRef< ITableRow > SUsdIntegrationsPanel::OnGenerateRow( TSharedPtr<UE::FUsdAttribute> InAttr, const TSharedRef< STableViewBase >& OwnerTable )
+TSharedRef<ITableRow> SUsdIntegrationsPanel::OnGenerateRow(TSharedPtr<UE::FUsdAttribute> InAttr, const TSharedRef<STableViewBase>& OwnerTable)
 {
-	return SNew( SUsdIntegrationsPanelRow, InAttr, OwnerTable );
+	return SNew(SUsdIntegrationsPanelRow, InAttr, OwnerTable);
 }
 
 #undef LOCTEXT_NAMESPACE
-#endif // USE_USD_SDK
\ No newline at end of file
+#endif	  // USE_USD_SDK
diff --git a/Engine/Plugins/Importers/USDImporter/Source/USDStageEditor/Private/Widgets/SUSDIntegrationsPanel.h b/Engine/Plugins/Importers/USDImporter/Source/USDStageEditor/Private/Widgets/SUSDIntegrationsPanel.h
--- a/Engine/Plugins/Importers/USDImporter/Source/USDStageEditor/Private/Widgets/SUSDIntegrationsPanel.h
+++ b/Engine/Plugins/Importers/USDImporter/Source/USDStageEditor/Private/Widgets/SUSDIntegrationsPanel.h
@@ -3,7 +3,7 @@
 #pragma once
 
 #include "USDIntegrationsViewModel.h"
-#include "UsdWrappers/UsdStage.h"
+#include "UsdWrappers/ForwardDeclarations.h"
 
 #include "Templates/SharedPointer.h"
 #include "Widgets/Views/SListView.h"
@@ -26,25 +26,27 @@ public:
 
 	// Ideally these would be FSoftObjectPaths, but FPropertyEditorModule doesn't support generating
 	// widgets for those properties yet
-	UPROPERTY( EditAnywhere, Transient, Category = Dummy, meta = ( AllowedClasses = "/Script/Engine.AnimBlueprint" ) )
+	UPROPERTY(EditAnywhere, Transient, Category = Dummy, meta = (AllowedClasses = "/Script/Engine.AnimBlueprint"))
 	TObjectPtr<UObject> AnimBPProperty;
 
-	UPROPERTY( EditAnywhere, Transient, Category = Dummy, meta = ( AllowedClasses = "/Script/ControlRigDeveloper.ControlRigBlueprint" ) )
+	UPROPERTY(EditAnywhere, Transient, Category = Dummy, meta = (AllowedClasses = "/Script/ControlRigDeveloper.ControlRigBlueprint"))
 	TObjectPtr<UObject> ControlRigProperty;
 };
 
 #if USE_USD_SDK
 
-class SUsdIntegrationsPanelRow : public SMultiColumnTableRow< TSharedPtr<UE::FUsdAttribute> >
+class SUsdIntegrationsPanelRow : public SMultiColumnTableRow<TSharedPtr<UE::FUsdAttribute>>
 {
 public:
-	SLATE_BEGIN_ARGS( SUsdIntegrationsPanelRow ) {}
+	SLATE_BEGIN_ARGS(SUsdIntegrationsPanelRow)
+	{
+	}
 	SLATE_END_ARGS()
 
 public:
-	void Construct( const FArguments& InArgs, TSharedPtr<UE::FUsdAttribute> InAttr, const TSharedRef< STableViewBase >& OwnerTable );
+	void Construct(const FArguments& InArgs, TSharedPtr<UE::FUsdAttribute> InAttr, const TSharedRef<STableViewBase>& OwnerTable);
 
-	virtual TSharedRef< SWidget > GenerateWidgetForColumn( const FName& ColumnName ) override;
+	virtual TSharedRef<SWidget> GenerateWidgetForColumn(const FName& ColumnName) override;
 
 private:
 	TSharedPtr<UE::FUsdAttribute> Attribute;
@@ -53,21 +55,23 @@ private:
 // We don't really need a list view here since we'll mostly always know exactly what attributes are going to
 // be displayed here beforehand, but doing so is a simple way of ensuring a consistent look between this panel
 // and the variants/references panels, that *do* need to be lists
-class SUsdIntegrationsPanel : public SListView< TSharedPtr<UE::FUsdAttribute> >
+class SUsdIntegrationsPanel : public SListView<TSharedPtr<UE::FUsdAttribute>>
 {
-	SLATE_BEGIN_ARGS( SUsdIntegrationsPanel ) {}
+	SLATE_BEGIN_ARGS(SUsdIntegrationsPanel)
+	{
+	}
 	SLATE_END_ARGS()
 
 public:
-	void Construct( const FArguments& InArgs );
-	void SetPrimPath( const UE::FUsdStageWeak& UsdStage, const TCHAR* PrimPath );
+	void Construct(const FArguments& InArgs);
+	void SetPrimPath(const UE::FUsdStageWeak& UsdStage, const TCHAR* PrimPath);
 
 protected:
-	TSharedRef< ITableRow > OnGenerateRow( TSharedPtr<UE::FUsdAttribute> InAttr, const TSharedRef< STableViewBase >& OwnerTable );
+	TSharedRef<ITableRow> OnGenerateRow(TSharedPtr<UE::FUsdAttribute> InAttr, const TSharedRef<STableViewBase>& OwnerTable);
 
 private:
-	TSharedPtr< SHeaderRow > HeaderRowWidget;
+	TSharedPtr<SHeaderRow> HeaderRowWidget;
 	FUsdIntegrationsViewModel ViewModel;
 };
 
-#endif // USE_USD_SDK
\ No newline at end of file
+#endif	  // USE_USD_SDK
