USD: Reformat USDStageEditor module.

diff --git a/Engine/Plugins/Importers/USDImporter/Source/USDStageEditor/Private/Widgets/SUSDObjectFieldList.cpp b/Engine/Plugins/Importers/USDImporter/Source/USDStageEditor/Private/Widgets/SUSDObjectFieldList.cpp
--- a/Engine/Plugins/Importers/USDImporter/Source/USDStageEditor/Private/Widgets/SUSDObjectFieldList.cpp
+++ b/Engine/Plugins/Importers/USDImporter/Source/USDStageEditor/Private/Widgets/SUSDObjectFieldList.cpp
@@ -4,16 +4,10 @@
 
 #include "SUSDStageEditorStyle.h"
 #include "UnrealUSDWrapper.h"
-#include "USDLog.h"
-#include "USDMemory.h"
-#include "USDStageActor.h"
-#include "USDStageModule.h"
-#include "USDTypesConversion.h"
 
-#include "Styling/AppStyle.h"
-#include "Engine/World.h"
 #include "Modules/ModuleManager.h"
 #include "ScopedTransaction.h"
+#include "Styling/AppStyle.h"
 #include "Widgets/Input/SCheckBox.h"
 #include "Widgets/Input/SComboBox.h"
 #include "Widgets/Input/SEditableTextBox.h"
@@ -25,10 +19,10 @@
 
 namespace USDObjectFieldListConstants
 {
-	const FMargin LeftRowPadding( 6.0f, 2.5f, 2.0f, 2.5f );
-	const FMargin RightRowPadding( 3.0f, 2.5f, 2.0f, 2.5f );
-	const FMargin ComboBoxItemPadding( 3.0f, 0.0f, 2.0f, 0.0f );
-	const FMargin NumericEntryBoxItemPadding( 0.0f, 0.0f, 2.0f, 0.0f );
+	const FMargin LeftRowPadding(6.0f, 2.5f, 2.0f, 2.5f);
+	const FMargin RightRowPadding(3.0f, 2.5f, 2.0f, 2.5f);
+	const FMargin ComboBoxItemPadding(3.0f, 0.0f, 2.0f, 0.0f);
+	const FMargin NumericEntryBoxItemPadding(0.0f, 0.0f, 2.0f, 0.0f);
 	const float DesiredNumericEntryBoxWidth = 80.0f;
 	const float DesiredTextEntryBoxWidth = 80.0f;
 
@@ -63,42 +57,42 @@ namespace USDObjectFieldListImpl
 
 	static TMap<FString, TArray<TSharedPtr<FString>>> TokenDropdownOptions;
 
-	void ResetOptions( const FString& TokenName )
+	void ResetOptions(const FString& TokenName)
 	{
-		TokenDropdownOptions.Remove( TokenName );
+		TokenDropdownOptions.Remove(TokenName);
 	}
 
-	TArray< TSharedPtr< FString > >* GetTokenDropdownOptions( const FUsdObjectFieldViewModel& ViewModel )
+	TArray<TSharedPtr<FString>>* GetTokenDropdownOptions(const FUsdObjectFieldViewModel& ViewModel)
 	{
-		if ( TArray< TSharedPtr< FString> >* FoundOptions = TokenDropdownOptions.Find( ViewModel.Label ) )
+		if (TArray<TSharedPtr<FString>>* FoundOptions = TokenDropdownOptions.Find(ViewModel.Label))
 		{
 			return FoundOptions;
 		}
 
-		TArray< TSharedPtr< FString > > Options = ViewModel.GetDropdownOptions();
-		if ( Options.Num() == 0 )
+		TArray<TSharedPtr<FString>> Options = ViewModel.GetDropdownOptions();
+		if (Options.Num() == 0)
 		{
 			// We don't know the options for this property, so return nullptr so that it can become a regular text input box
 			return nullptr;
 		}
 
-		return &USDObjectFieldListImpl::TokenDropdownOptions.Add( ViewModel.Label, Options );
+		return &USDObjectFieldListImpl::TokenDropdownOptions.Add(ViewModel.Label, Options);
 	}
 
-	EWidgetLabelTypes GetLabelType( const UsdUtils::FConvertedVtValue& FieldValue, const FString& ValueRole )
+	EWidgetLabelTypes GetLabelType(const UsdUtils::FConvertedVtValue& FieldValue, const FString& ValueRole)
 	{
-		if ( FieldValue.Entries.Num() < 1 )
+		if (FieldValue.Entries.Num() < 1)
 		{
 			return EWidgetLabelTypes::NoLabel;
 		}
 
 		int32 NumComponents = FieldValue.Entries[0].Num();
 		EWidgetLabelTypes LabelType = EWidgetLabelTypes::NoLabel;
-		if ( ValueRole.StartsWith( TEXT( "color" ), ESearchCase::IgnoreCase ) )
+		if (ValueRole.StartsWith(TEXT("color"), ESearchCase::IgnoreCase))
 		{
 			LabelType = EWidgetLabelTypes::RGBA;
 		}
-		else if ( NumComponents > 1 && NumComponents <= 4 && FieldValue.SourceType != UsdUtils::EUsdBasicDataTypes::Matrix2d )
+		else if (NumComponents > 1 && NumComponents <= 4 && FieldValue.SourceType != UsdUtils::EUsdBasicDataTypes::Matrix2d)
 		{
 			LabelType = EWidgetLabelTypes::XYZW;
 		}
@@ -106,125 +100,130 @@ namespace USDObjectFieldListImpl
 		return LabelType;
 	}
 
-	EWidgetType GetWidgetType( UsdUtils::EUsdBasicDataTypes SourceType )
+	EWidgetType GetWidgetType(UsdUtils::EUsdBasicDataTypes SourceType)
 	{
 		using namespace UsdUtils;
 
-		switch ( SourceType )
+		switch (SourceType)
 		{
-		case EUsdBasicDataTypes::Bool:
-			return EWidgetType::Bool;
-			break;
-		case EUsdBasicDataTypes::Uchar:
-			return EWidgetType::U8;
-			break;
-		case EUsdBasicDataTypes::Int:
-		case EUsdBasicDataTypes::Int2:
-		case EUsdBasicDataTypes::Int3:
-		case EUsdBasicDataTypes::Int4:
-			return EWidgetType::I32;
-			break;
-		case EUsdBasicDataTypes::Uint:
-			return EWidgetType::U32;
-			break;
-		case EUsdBasicDataTypes::Int64:
-			return EWidgetType::I64;
-			break;
-		case EUsdBasicDataTypes::Uint64:
-			return EWidgetType::U64;
-			break;
-		case EUsdBasicDataTypes::Half:
-		case EUsdBasicDataTypes::Half2:
-		case EUsdBasicDataTypes::Half3:
-		case EUsdBasicDataTypes::Half4:
-		case EUsdBasicDataTypes::Quath:
-		case EUsdBasicDataTypes::Float:
-		case EUsdBasicDataTypes::Float2:
-		case EUsdBasicDataTypes::Float3:
-		case EUsdBasicDataTypes::Float4:
-		case EUsdBasicDataTypes::Quatf:
-			return EWidgetType::F32;
-			break;
-		case EUsdBasicDataTypes::Double:
-		case EUsdBasicDataTypes::Double2:
-		case EUsdBasicDataTypes::Double3:
-		case EUsdBasicDataTypes::Double4:
-		case EUsdBasicDataTypes::Timecode:
-		case EUsdBasicDataTypes::Matrix2d:
-		case EUsdBasicDataTypes::Matrix3d:
-		case EUsdBasicDataTypes::Matrix4d:
-		case EUsdBasicDataTypes::Quatd:
-			return EWidgetType::F64;
-			break;
-		case EUsdBasicDataTypes::Token:
-			return EWidgetType::Dropdown;
-			break;
-		case EUsdBasicDataTypes::String:
-		case EUsdBasicDataTypes::Asset:
-			return EWidgetType::Text;
-			break;
-		default:
-			break;
+			case EUsdBasicDataTypes::Bool:
+				return EWidgetType::Bool;
+				break;
+			case EUsdBasicDataTypes::Uchar:
+				return EWidgetType::U8;
+				break;
+			case EUsdBasicDataTypes::Int:
+			case EUsdBasicDataTypes::Int2:
+			case EUsdBasicDataTypes::Int3:
+			case EUsdBasicDataTypes::Int4:
+				return EWidgetType::I32;
+				break;
+			case EUsdBasicDataTypes::Uint:
+				return EWidgetType::U32;
+				break;
+			case EUsdBasicDataTypes::Int64:
+				return EWidgetType::I64;
+				break;
+			case EUsdBasicDataTypes::Uint64:
+				return EWidgetType::U64;
+				break;
+			case EUsdBasicDataTypes::Half:
+			case EUsdBasicDataTypes::Half2:
+			case EUsdBasicDataTypes::Half3:
+			case EUsdBasicDataTypes::Half4:
+			case EUsdBasicDataTypes::Quath:
+			case EUsdBasicDataTypes::Float:
+			case EUsdBasicDataTypes::Float2:
+			case EUsdBasicDataTypes::Float3:
+			case EUsdBasicDataTypes::Float4:
+			case EUsdBasicDataTypes::Quatf:
+				return EWidgetType::F32;
+				break;
+			case EUsdBasicDataTypes::Double:
+			case EUsdBasicDataTypes::Double2:
+			case EUsdBasicDataTypes::Double3:
+			case EUsdBasicDataTypes::Double4:
+			case EUsdBasicDataTypes::Timecode:
+			case EUsdBasicDataTypes::Matrix2d:
+			case EUsdBasicDataTypes::Matrix3d:
+			case EUsdBasicDataTypes::Matrix4d:
+			case EUsdBasicDataTypes::Quatd:
+				return EWidgetType::F64;
+				break;
+			case EUsdBasicDataTypes::Token:
+				return EWidgetType::Dropdown;
+				break;
+			case EUsdBasicDataTypes::String:
+			case EUsdBasicDataTypes::Asset:
+				return EWidgetType::Text;
+				break;
+			default:
+				break;
 		}
 
 		return EWidgetType::None;
 	}
 
 	/** Always returns `NumLabels` widgets, regardless of LabelTypes. Those may be the NullWidget though */
-	TArray<TSharedRef<SWidget>> GetNumericEntryBoxLabels( int32 NumLabels, EWidgetLabelTypes LabelType )
+	TArray<TSharedRef<SWidget>> GetNumericEntryBoxLabels(int32 NumLabels, EWidgetLabelTypes LabelType)
 	{
 		const static TArray<const FLinearColor*> Colors = {
 			&SNumericEntryBox<int32>::RedLabelBackgroundColor,
 			&SNumericEntryBox<int32>::GreenLabelBackgroundColor,
 			&SNumericEntryBox<int32>::BlueLabelBackgroundColor,
-			&FLinearColor::White
-		};
+			&FLinearColor::White};
 
-		if ( NumLabels > 4 )
+		if (NumLabels > 4)
 		{
 			LabelType = EWidgetLabelTypes::NoLabel;
 		}
 
 		TArray<TSharedRef<SWidget>> Labels;
-		Labels.Reserve( NumLabels );
+		Labels.Reserve(NumLabels);
 
-		switch ( LabelType )
+		switch (LabelType)
 		{
-		case EWidgetLabelTypes::RGBA:
-			for ( int32 Index = 0; Index < NumLabels; ++Index )
-			{
-				Labels.Add( SNumericEntryBox<int32>::BuildNarrowColorLabel( *Colors[ Index ] ) );
-			}
-			break;
-		case EWidgetLabelTypes::XYZW:
-			for ( int32 Index = 0; Index < NumLabels; ++Index )
-			{
-				Labels.Add( SNumericEntryBox<int32>::BuildNarrowColorLabel( *Colors[ Index ] ) );
-			}
-			break;
-		case EWidgetLabelTypes::NoLabel:
-		default:
-			for ( int32 Index = 0; Index < NumLabels; ++Index )
-			{
-				Labels.Add( SNullWidget::NullWidget );
-			}
-			break;
+			case EWidgetLabelTypes::RGBA:
+				for (int32 Index = 0; Index < NumLabels; ++Index)
+				{
+					Labels.Add(SNumericEntryBox<int32>::BuildNarrowColorLabel(*Colors[Index]));
+				}
+				break;
+			case EWidgetLabelTypes::XYZW:
+				for (int32 Index = 0; Index < NumLabels; ++Index)
+				{
+					Labels.Add(SNumericEntryBox<int32>::BuildNarrowColorLabel(*Colors[Index]));
+				}
+				break;
+			case EWidgetLabelTypes::NoLabel:
+			default:
+				for (int32 Index = 0; Index < NumLabels; ++Index)
+				{
+					Labels.Add(SNullWidget::NullWidget);
+				}
+				break;
 		}
 
 		return Labels;
 	}
-}
+}	 // namespace USDObjectFieldListImpl
 
 class SUsdObjectFieldRow : public SMultiColumnTableRow<TSharedPtr<FUsdObjectFieldViewModel>>
 {
-	SLATE_BEGIN_ARGS( SUsdObjectFieldRow ) {}
+	SLATE_BEGIN_ARGS(SUsdObjectFieldRow)
+	{
+	}
 	SLATE_END_ARGS()
 
 public:
-	void Construct(const FArguments& InArgs, const TSharedPtr<FUsdObjectFieldViewModel>& InUsdObjectField, const TSharedRef< STableViewBase >& OwnerTable);
-	virtual TSharedRef< SWidget > GenerateWidgetForColumn( const FName& ColumnName ) override;
+	void Construct(
+		const FArguments& InArgs,
+		const TSharedPtr<FUsdObjectFieldViewModel>& InUsdObjectField,
+		const TSharedRef<STableViewBase>& OwnerTable
+	);
+	virtual TSharedRef<SWidget> GenerateWidgetForColumn(const FName& ColumnName) override;
 
-	void SetUsdObjectField(const TSharedPtr< FUsdObjectFieldViewModel >& InUsdObjectField);
+	void SetUsdObjectField(const TSharedPtr<FUsdObjectFieldViewModel>& InUsdObjectField);
 
 protected:
 	FText GetLabel() const
@@ -233,16 +232,10 @@ protected:
 	}
 	FText GetType() const
 	{
-		const static FText TypeTexts[] = {
-			FText::FromString(TEXT("M")),
-			FText::FromString(TEXT("A")),
-			FText::FromString(TEXT("R"))
-		};
+		const static FText TypeTexts[] = {FText::FromString(TEXT("M")), FText::FromString(TEXT("A")), FText::FromString(TEXT("R"))};
 
 		const int32 Index = static_cast<int32>(Field->Type);
-		return Index >= 0 && Index < static_cast<int32>(EObjectFieldType::MAX)
-			? TypeTexts[Index]
-			: FText::GetEmpty();
+		return Index >= 0 && Index < static_cast<int32>(EObjectFieldType::MAX) ? TypeTexts[Index] : FText::GetEmpty();
 	}
 
 	FText GetTypeToolTip() const
@@ -250,24 +243,21 @@ protected:
 		const static FText ToolTipTexts[] = {
 			LOCTEXT("MetadataTooltip", "Metadata"),
 			LOCTEXT("AttributeTooltip", "Attribute"),
-			LOCTEXT("RelationshipTooltip", "Relationship")
-		};
+			LOCTEXT("RelationshipTooltip", "Relationship")};
 
 		const int32 Index = static_cast<int32>(Field->Type);
-		return Index >= 0 && Index < static_cast<int32>(EObjectFieldType::MAX)
-			? ToolTipTexts[Index]
-			: FText::GetEmpty();
+		return Index >= 0 && Index < static_cast<int32>(EObjectFieldType::MAX) ? ToolTipTexts[Index] : FText::GetEmpty();
 	}
 
 	// Optional here because SNumericEntryBox uses optional values as input
 	template<typename T>
-	TOptional<T> GetValue( int32 ComponentIndex ) const
+	TOptional<T> GetValue(int32 ComponentIndex) const
 	{
-		if ( Field->Value.Entries.Num() == 1 ) // Ignore arrays for now
+		if (Field->Value.Entries.Num() == 1)	// Ignore arrays for now
 		{
-			if ( Field->Value.Entries[0].IsValidIndex( ComponentIndex ) )
+			if (Field->Value.Entries[0].IsValidIndex(ComponentIndex))
 			{
-				if ( T* Value = Field->Value.Entries[0][ ComponentIndex ].TryGet<T>() )
+				if (T* Value = Field->Value.Entries[0][ComponentIndex].TryGet<T>())
 				{
 					return *Value;
 				}
@@ -280,13 +270,13 @@ protected:
 	// Other overloads as checkbox/text widgets don't use optional values as input, and these should only ever have one component anyway
 	FText GetValueText() const
 	{
-		if ( Field->Value.Entries.Num() == 1 )
+		if (Field->Value.Entries.Num() == 1)
 		{
-			if ( Field->Value.Entries[ 0 ].Num() > 0 )
+			if (Field->Value.Entries[0].Num() > 0)
 			{
-				if ( FString* Value = Field->Value.Entries[ 0 ][ 0 ].TryGet<FString>() )
+				if (FString* Value = Field->Value.Entries[0][0].TryGet<FString>())
 				{
-					return FText::FromString( *Value );
+					return FText::FromString(*Value);
 				}
 			}
 		}
@@ -296,11 +286,11 @@ protected:
 
 	ECheckBoxState GetValueBool() const
 	{
-		if ( Field->Value.Entries.Num() == 1 )
+		if (Field->Value.Entries.Num() == 1)
 		{
-			if ( Field->Value.Entries[ 0 ].Num() > 0 )
+			if (Field->Value.Entries[0].Num() > 0)
 			{
-				if ( bool* Value = Field->Value.Entries[ 0 ][ 0 ].TryGet<bool>() )
+				if (bool* Value = Field->Value.Entries[0][0].TryGet<bool>())
 				{
 					return *Value ? ECheckBoxState::Checked : ECheckBoxState::Unchecked;
 				}
@@ -314,51 +304,58 @@ private:
 	template<typename T>
 	TArray<TAttribute<TOptional<T>>> GetAttributeArray();
 
-	TSharedRef< SWidget > GenerateTextWidget(const TAttribute<FText>& Attribute, const TAttribute<FText>& ToolTipAttribute = {}, ETextJustify::Type TextJustify = ETextJustify::Left);
-	TSharedRef< SWidget > GenerateEditableTextWidget( const TAttribute<FText>& Attribute, bool bIsReadOnly );
-	TSharedRef< SWidget > GenerateCheckboxWidget( const TAttribute<ECheckBoxState>& Attribute );
+	TSharedRef<SWidget> GenerateTextWidget(
+		const TAttribute<FText>& Attribute,
+		const TAttribute<FText>& ToolTipAttribute = {},
+		ETextJustify::Type TextJustify = ETextJustify::Left
+	);
+	TSharedRef<SWidget> GenerateEditableTextWidget(const TAttribute<FText>& Attribute, bool bIsReadOnly);
+	TSharedRef<SWidget> GenerateCheckboxWidget(const TAttribute<ECheckBoxState>& Attribute);
 
 	template<typename T>
-	TSharedRef< SWidget > GenerateSpinboxWidgets(
+	TSharedRef<SWidget> GenerateSpinboxWidgets(
 		const TArray<TAttribute<TOptional<T>>>& Attribute,
 		USDObjectFieldListImpl::EWidgetLabelTypes LabelType = USDObjectFieldListImpl::EWidgetLabelTypes::NoLabel
 	);
 
 private:
 	template<typename T>
-	void OnSpinboxValueChanged( T NewValue, int32 ComponentIndex );
+	void OnSpinboxValueChanged(T NewValue, int32 ComponentIndex);
 
 	template<typename T>
-	void OnSpinboxValueCommitted( T InNewValue, ETextCommit::Type CommitType, int32 ComponentIndex );
+	void OnSpinboxValueCommitted(T InNewValue, ETextCommit::Type CommitType, int32 ComponentIndex);
 
-	void OnComboBoxSelectionChanged( TSharedPtr<FString> ChosenOption, ESelectInfo::Type SelectInfo );
+	void OnComboBoxSelectionChanged(TSharedPtr<FString> ChosenOption, ESelectInfo::Type SelectInfo);
 	void OnTextBoxTextCommitted(const FText& NewText, ETextCommit::Type CommitType);
-	void OnCheckBoxCheckStateChanged( ECheckBoxState NewState );
+	void OnCheckBoxCheckStateChanged(ECheckBoxState NewState);
 
 private:
 	TSharedPtr<FUsdObjectFieldViewModel> Field;
 };
 
-void SUsdObjectFieldRow::Construct(const FArguments& InArgs, const TSharedPtr<FUsdObjectFieldViewModel>& InUsdObjectField, const TSharedRef<STableViewBase>& OwnerTable)
+void SUsdObjectFieldRow::Construct(
+	const FArguments& InArgs,
+	const TSharedPtr<FUsdObjectFieldViewModel>& InUsdObjectField,
+	const TSharedRef<STableViewBase>& OwnerTable
+)
 {
 	SetUsdObjectField(InUsdObjectField);
 
-	SMultiColumnTableRow<TSharedPtr<FUsdObjectFieldViewModel>>::Construct(SMultiColumnTableRow<TSharedPtr<FUsdObjectFieldViewModel>>::FArguments(), OwnerTable);
+	SMultiColumnTableRow<TSharedPtr<FUsdObjectFieldViewModel>>::Construct(
+		SMultiColumnTableRow<TSharedPtr<FUsdObjectFieldViewModel>>::FArguments(),
+		OwnerTable
+	);
 }
 
-TSharedRef< SWidget > SUsdObjectFieldRow::GenerateWidgetForColumn( const FName& ColumnName )
+TSharedRef<SWidget> SUsdObjectFieldRow::GenerateWidgetForColumn(const FName& ColumnName)
 {
 	using namespace USDObjectFieldListImpl;
 
-	TSharedRef< SWidget > ColumnWidget = SNullWidget::NullWidget;
+	TSharedRef<SWidget> ColumnWidget = SNullWidget::NullWidget;
 
 	if (ColumnName == ObjectFieldColumnIds::TypeColumn)
 	{
-		ColumnWidget = GenerateTextWidget(
-			{this, &SUsdObjectFieldRow::GetType},
-			{this, &SUsdObjectFieldRow::GetTypeToolTip},
-			ETextJustify::Center
-		);
+		ColumnWidget = GenerateTextWidget({this, &SUsdObjectFieldRow::GetType}, {this, &SUsdObjectFieldRow::GetTypeToolTip}, ETextJustify::Center);
 	}
 	else if (ColumnName == ObjectFieldColumnIds::NameColumn)
 	{
@@ -366,90 +363,94 @@ TSharedRef< SWidget > SUsdObjectFieldRow::GenerateWidgetForColumn( const FName&
 	}
 	else
 	{
-		EWidgetType WidgetType = GetWidgetType( Field->Value.SourceType );
-		EWidgetLabelTypes LabelType = GetLabelType( Field->Value, Field->ValueRole );
+		EWidgetType WidgetType = GetWidgetType(Field->Value.SourceType);
+		EWidgetLabelTypes LabelType = GetLabelType(Field->Value, Field->ValueRole);
 
-		switch ( WidgetType )
-		{
-		case EWidgetType::Bool:
-			ColumnWidget = GenerateCheckboxWidget( { this, &SUsdObjectFieldRow::GetValueBool } );
-			break;
-		case EWidgetType::U8:
-			ColumnWidget = GenerateSpinboxWidgets( GetAttributeArray<uint8>(), LabelType );
-			break;
-		case EWidgetType::I32:
-			ColumnWidget = GenerateSpinboxWidgets( GetAttributeArray<int32>(), LabelType );
-			break;
-		case EWidgetType::U32:
-			ColumnWidget = GenerateSpinboxWidgets( GetAttributeArray<uint32>(), LabelType );
-			break;
-		case EWidgetType::I64:
-			ColumnWidget = GenerateSpinboxWidgets( GetAttributeArray<int64>(), LabelType );
-			break;
-		case EWidgetType::U64:
-			ColumnWidget = GenerateSpinboxWidgets( GetAttributeArray<uint64>(), LabelType );
-			break;
-		case EWidgetType::F32:
-			ColumnWidget = GenerateSpinboxWidgets( GetAttributeArray<float>(), LabelType );
-			break;
-		case EWidgetType::F64:
-			ColumnWidget = GenerateSpinboxWidgets( GetAttributeArray<double>(), LabelType );
-			break;
-		case EWidgetType::Dropdown:
+		switch (WidgetType)
 		{
-			TArray< TSharedPtr< FString > >* Options = GetTokenDropdownOptions( *Field );
-
-			// Show a dropdown if we know the available options for that token
-			if ( Options )
+			case EWidgetType::Bool:
+				ColumnWidget = GenerateCheckboxWidget({this, &SUsdObjectFieldRow::GetValueBool});
+				break;
+			case EWidgetType::U8:
+				ColumnWidget = GenerateSpinboxWidgets(GetAttributeArray<uint8>(), LabelType);
+				break;
+			case EWidgetType::I32:
+				ColumnWidget = GenerateSpinboxWidgets(GetAttributeArray<int32>(), LabelType);
+				break;
+			case EWidgetType::U32:
+				ColumnWidget = GenerateSpinboxWidgets(GetAttributeArray<uint32>(), LabelType);
+				break;
+			case EWidgetType::I64:
+				ColumnWidget = GenerateSpinboxWidgets(GetAttributeArray<int64>(), LabelType);
+				break;
+			case EWidgetType::U64:
+				ColumnWidget = GenerateSpinboxWidgets(GetAttributeArray<uint64>(), LabelType);
+				break;
+			case EWidgetType::F32:
+				ColumnWidget = GenerateSpinboxWidgets(GetAttributeArray<float>(), LabelType);
+				break;
+			case EWidgetType::F64:
+				ColumnWidget = GenerateSpinboxWidgets(GetAttributeArray<double>(), LabelType);
+				break;
+			case EWidgetType::Dropdown:
 			{
-				SAssignNew( ColumnWidget, SBox )
-				.HeightOverride( FUsdStageEditorStyle::Get()->GetFloat( "UsdStageEditor.ListItemHeight" ) )
-				.VAlign( VAlign_Center )
-				.HAlign( HAlign_Left )
-				[
-					SNew( SComboBox< TSharedPtr< FString > > )
-					.OptionsSource( Options )
-					.OnGenerateWidget_Lambda( [ & ]( TSharedPtr<FString> Option )
-					{
-						return SUsdObjectFieldRow::GenerateTextWidget( FText::FromString( *Option ) );
-					})
-					.OnSelectionChanged( this, &SUsdObjectFieldRow::OnComboBoxSelectionChanged )
+				TArray<TSharedPtr<FString>>* Options = GetTokenDropdownOptions(*Field);
+
+				// Show a dropdown if we know the available options for that token
+				if (Options)
+				{
+					// clang-format off
+					SAssignNew(ColumnWidget, SBox)
+					.HeightOverride(FUsdStageEditorStyle::Get()->GetFloat("UsdStageEditor.ListItemHeight"))
+					.VAlign(VAlign_Center)
+					.HAlign(HAlign_Left)
 					[
-						// Having an editable text box inside the combobox allows the user to pick through the most common ones but to
-						// also specify a custom kind/purpose/etc. if they want to
-						SNew( SEditableTextBox )
-						.Text( this, &SUsdObjectFieldRow::GetValueText )
-						.Font( FAppStyle::GetFontStyle( "PropertyWindow.NormalFont" ) )
-						.Padding( FMargin( 3.0f ) )
-						// Fixed foreground color or else it will flip when our row is selected, and we already have a background
-						.ForegroundColor( FAppStyle::GetSlateColor( TEXT( "Colors.ForegroundHover" ) ) )
-						.OnTextCommitted( this, &SUsdObjectFieldRow::OnTextBoxTextCommitted )
-					]
-				];
+						SNew(SComboBox<TSharedPtr<FString>>)
+						.OptionsSource(Options)
+						.OnGenerateWidget_Lambda([&](TSharedPtr<FString> Option)
+						{
+							return SUsdObjectFieldRow::GenerateTextWidget(FText::FromString(*Option));
+						})
+						.OnSelectionChanged(this, &SUsdObjectFieldRow::OnComboBoxSelectionChanged)
+						[
+							// Having an editable text box inside the combobox allows the user to pick through the most common ones but to
+							// also specify a custom kind/purpose/etc. if they want to
+							SNew(SEditableTextBox)
+							.Text(this, &SUsdObjectFieldRow::GetValueText)
+							.Font(FAppStyle::GetFontStyle("PropertyWindow.NormalFont"))
+							.Padding(FMargin(3.0f))
+							// Fixed foreground color or else it will flip when our row is selected, and we already have a background
+							.ForegroundColor(FAppStyle::GetSlateColor(TEXT("Colors.ForegroundHover")))
+							.OnTextCommitted(this, &SUsdObjectFieldRow::OnTextBoxTextCommitted)
+						]
+					];
+					// clang-format on
+
+					break;
+				}
 
+				// If we don't have any options we intentionally fall down into the 'Text' case
+			}
+			case EWidgetType::Text:
+			{
+				ColumnWidget = SUsdObjectFieldRow::GenerateEditableTextWidget({this, &SUsdObjectFieldRow::GetValueText}, Field->bReadOnly);
 				break;
 			}
-
-			// If we don't have any options we intentionally fall down into the 'Text' case
-		}
-		case EWidgetType::Text:
-		{
-			ColumnWidget = SUsdObjectFieldRow::GenerateEditableTextWidget( { this, &SUsdObjectFieldRow::GetValueText }, Field->bReadOnly );
-			break;
-		}
-		case EWidgetType::None:
-		default:
-			ColumnWidget = SNullWidget::NullWidget;
-			break;
+			case EWidgetType::None:
+			default:
+				ColumnWidget = SNullWidget::NullWidget;
+				break;
 		}
 	}
 
+	// clang-format off
 	return SNew(SBox)
 		.VAlign(VAlign_Center)
 		.HAlign(HAlign_Fill)
 		[
 			ColumnWidget
 		];
+	// clang-format on
 }
 
 void SUsdObjectFieldRow::SetUsdObjectField(const TSharedPtr<FUsdObjectFieldViewModel>& InUsdObjectField)
@@ -463,123 +464,130 @@ TArray<TAttribute<TOptional<T>>> SUsdObjectFieldRow::GetAttributeArray()
 	using AttrType = TAttribute<TOptional<T>>;
 
 	TArray<AttrType> Attributes;
-	if ( Field->Value.Entries.Num() == 1 )
+	if (Field->Value.Entries.Num() == 1)
 	{
-		for ( int32 ComponentIndex = 0; ComponentIndex < Field->Value.Entries[0].Num(); ++ComponentIndex )
+		for (int32 ComponentIndex = 0; ComponentIndex < Field->Value.Entries[0].Num(); ++ComponentIndex)
 		{
-			Attributes.Add( AttrType::Create( AttrType::FGetter::CreateSP( this, &SUsdObjectFieldRow::template GetValue<T>, ComponentIndex ) ) );
+			Attributes.Add(AttrType::Create(AttrType::FGetter::CreateSP(this, &SUsdObjectFieldRow::template GetValue<T>, ComponentIndex)));
 		}
 	}
 
 	return Attributes;
 }
 
-TSharedRef< SWidget > SUsdObjectFieldRow::GenerateTextWidget(
+TSharedRef<SWidget> SUsdObjectFieldRow::GenerateTextWidget(
 	const TAttribute<FText>& Attribute,
 	const TAttribute<FText>& ToolTipAttribute,
 	ETextJustify::Type TextJustify
 )
 {
+	// clang-format off
 	return SNew(STextBlock)
 		.Text(Attribute)
 		.Justification(TextJustify)
 		.ToolTipText(ToolTipAttribute)
 		.Font(FAppStyle::GetFontStyle("PropertyWindow.NormalFont"))
 		.Margin(USDObjectFieldListConstants::RightRowPadding);
+	// clang-format on
 }
 
-TSharedRef< SWidget > SUsdObjectFieldRow::GenerateEditableTextWidget( const TAttribute<FText>& Attribute, bool bIsReadOnly )
+TSharedRef<SWidget> SUsdObjectFieldRow::GenerateEditableTextWidget(const TAttribute<FText>& Attribute, bool bIsReadOnly)
 {
-	return SNew( SBox )
-		.HeightOverride( FUsdStageEditorStyle::Get()->GetFloat( "UsdStageEditor.ListItemHeight" ) )
-		.MinDesiredWidth( USDObjectFieldListConstants::DesiredTextEntryBoxWidth )
+	// clang-format off
+	return SNew(SBox)
+		.HeightOverride(FUsdStageEditorStyle::Get()->GetFloat("UsdStageEditor.ListItemHeight"))
+		.MinDesiredWidth(USDObjectFieldListConstants::DesiredTextEntryBoxWidth)
 		.VAlign(VAlign_Center)
 		.HAlign(HAlign_Left)
 		[
-			SNew( SEditableTextBox )
-			.MinDesiredWidth( USDObjectFieldListConstants::DesiredTextEntryBoxWidth )
-			.Text( Attribute )
-			.IsReadOnly( bIsReadOnly )
-			.Font( FAppStyle::GetFontStyle( "PropertyWindow.NormalFont" ) )
-			.OnTextCommitted( this, &SUsdObjectFieldRow::OnTextBoxTextCommitted )
+			SNew(SEditableTextBox)
+			.MinDesiredWidth(USDObjectFieldListConstants::DesiredTextEntryBoxWidth)
+			.Text(Attribute)
+			.IsReadOnly(bIsReadOnly)
+			.Font(FAppStyle::GetFontStyle("PropertyWindow.NormalFont"))
+			.OnTextCommitted(this, &SUsdObjectFieldRow::OnTextBoxTextCommitted)
 			// Fixed foreground color or else it will flip when our row is selected, and we already have a background
-			.ForegroundColor( FAppStyle::GetSlateColor( bIsReadOnly ? TEXT( "Colors.Foreground" ) : TEXT( "Colors.ForegroundHover" ) ) )
+			.ForegroundColor(FAppStyle::GetSlateColor(bIsReadOnly ? TEXT("Colors.Foreground") : TEXT("Colors.ForegroundHover")))
 		];
+	// clang-format on
 }
 
-TSharedRef< SWidget > SUsdObjectFieldRow::GenerateCheckboxWidget( const TAttribute<ECheckBoxState>& Attribute )
+TSharedRef<SWidget> SUsdObjectFieldRow::GenerateCheckboxWidget(const TAttribute<ECheckBoxState>& Attribute)
 {
-	return SNew( SBox )
-		.HeightOverride( FUsdStageEditorStyle::Get()->GetFloat( "UsdStageEditor.ListItemHeight" ) )
-		.VAlign( VAlign_Center )
+	// clang-format off
+	return SNew(SBox)
+		.HeightOverride(FUsdStageEditorStyle::Get()->GetFloat("UsdStageEditor.ListItemHeight"))
+		.VAlign(VAlign_Center)
 		.HAlign(HAlign_Left)
 		[
-			SNew( SCheckBox )
-			.IsChecked( Attribute )
-			.OnCheckStateChanged( this, &SUsdObjectFieldRow::OnCheckBoxCheckStateChanged )
+			SNew(SCheckBox)
+			.IsChecked(Attribute)
+			.OnCheckStateChanged(this, &SUsdObjectFieldRow::OnCheckBoxCheckStateChanged)
 		];
+	// clang-format on
 }
 
 template<typename T>
-TSharedRef< SWidget > SUsdObjectFieldRow::GenerateSpinboxWidgets( const TArray<TAttribute<TOptional<T>>>& Attributes, USDObjectFieldListImpl::EWidgetLabelTypes LabelType )
+TSharedRef<SWidget> SUsdObjectFieldRow::GenerateSpinboxWidgets(
+	const TArray<TAttribute<TOptional<T>>>& Attributes,
+	USDObjectFieldListImpl::EWidgetLabelTypes LabelType
+)
 {
-	if ( Attributes.Num() == 0 )
+	if (Attributes.Num() == 0)
 	{
 		return SNullWidget::NullWidget;
 	}
 
-	TArray<TSharedRef<SWidget>> Labels = USDObjectFieldListImpl::GetNumericEntryBoxLabels( Attributes.Num(), LabelType );
+	TArray<TSharedRef<SWidget>> Labels = USDObjectFieldListImpl::GetNumericEntryBoxLabels(Attributes.Num(), LabelType);
 
 	// If we're some type of matrix attribute, we'll want to display multiple rows
-	const int32 NumRows = LabelType != USDObjectFieldListImpl::EWidgetLabelTypes::NoLabel
-		? 1 // If we have labels we're definitely not a matrix type
-		: Attributes.Num() == 4
-			? 2
-			: Attributes.Num() == 9
-				? 3
-				: Attributes.Num() == 16
-					? 4
-					: 1;
+	const int32 NumRows = LabelType != USDObjectFieldListImpl::EWidgetLabelTypes::NoLabel ? 1	 // If we have labels we're definitely not a matrix
+																								 // type
+						  : Attributes.Num() == 4  ? 2
+						  : Attributes.Num() == 9  ? 3
+						  : Attributes.Num() == 16 ? 4
+												   : 1;
 
 	const int32 NumColumns = Attributes.Num() / NumRows;
 
-	TSharedPtr<SVerticalBox> VertBox = SNew( SVerticalBox );
-	for ( int32 RowIndex = 0; RowIndex < NumRows; ++RowIndex )
+	// clang-format off
+	TSharedPtr<SVerticalBox> VertBox = SNew(SVerticalBox);
+	for (int32 RowIndex = 0; RowIndex <NumRows; ++RowIndex)
 	{
-		TSharedPtr<SHorizontalBox> HorizBox = SNew( SHorizontalBox );
-		for ( int32 ColumnIndex = 0; ColumnIndex < NumColumns; ++ColumnIndex )
+		TSharedPtr<SHorizontalBox> HorizBox = SNew(SHorizontalBox);
+		for (int32 ColumnIndex = 0; ColumnIndex <NumColumns; ++ColumnIndex)
 		{
 			const int32 ComponentIndex = RowIndex * NumColumns + ColumnIndex;
-			const TAttribute<TOptional<T>>& Attribute = Attributes[ ComponentIndex ];
-
-			TSharedRef<SWidget> EntryBox = SNew( SNumericEntryBox<T> )
-				.AllowSpin( true )
-				.Font( FAppStyle::GetFontStyle( "PropertyWindow.NormalFont" ) )
-				.ShiftMouseMovePixelPerDelta( 1 )
-				.SupportDynamicSliderMaxValue( true )
-				.SupportDynamicSliderMinValue( true )
-				.OnValueChanged( this, &SUsdObjectFieldRow::OnSpinboxValueChanged<T>, ComponentIndex )
-				.OnValueCommitted( this, &SUsdObjectFieldRow::OnSpinboxValueCommitted<T>, ComponentIndex )
-				.Value( Attribute )
-				.MinValue( TOptional<T>() )
-				.MaxValue( TOptional<T>() )
-				.MinSliderValue( TOptional<T>() )
-				.MaxSliderValue( TOptional<T>() )
-				.SliderExponent( T( 1 ) )
-				.Delta( T( 0 ) )
-				.LabelPadding( FMargin(3) )
-				.LabelLocation( SNumericEntryBox<T>::ELabelLocation::Inside )
+			const TAttribute<TOptional<T>>& Attribute = Attributes[ComponentIndex];
+
+			TSharedRef<SWidget> EntryBox = SNew(SNumericEntryBox<T>)
+				.AllowSpin(true)
+				.Font(FAppStyle::GetFontStyle("PropertyWindow.NormalFont"))
+				.ShiftMouseMovePixelPerDelta(1)
+				.SupportDynamicSliderMaxValue(true)
+				.SupportDynamicSliderMinValue(true)
+				.OnValueChanged(this, &SUsdObjectFieldRow::OnSpinboxValueChanged<T>, ComponentIndex)
+				.OnValueCommitted(this, &SUsdObjectFieldRow::OnSpinboxValueCommitted<T>, ComponentIndex)
+				.Value(Attribute)
+				.MinValue(TOptional<T>())
+				.MaxValue(TOptional<T>())
+				.MinSliderValue(TOptional<T>())
+				.MaxSliderValue(TOptional<T>())
+				.SliderExponent(T(1))
+				.Delta(T(0))
+				.LabelPadding(FMargin(3))
+				.LabelLocation(SNumericEntryBox<T>::ELabelLocation::Inside)
 				.Label()
 				[
-					Labels[ ComponentIndex ]
+					Labels[ComponentIndex]
 				];
 
 			HorizBox->AddSlot()
-				.Padding( USDObjectFieldListConstants::NumericEntryBoxItemPadding )
+				.Padding(USDObjectFieldListConstants::NumericEntryBoxItemPadding)
 				[
-					SNew( SBox )
-					.MinDesiredWidth( USDObjectFieldListConstants::DesiredNumericEntryBoxWidth )
-					.MaxDesiredWidth( USDObjectFieldListConstants::DesiredNumericEntryBoxWidth )
+					SNew(SBox)
+					.MinDesiredWidth(USDObjectFieldListConstants::DesiredNumericEntryBoxWidth)
+					.MaxDesiredWidth(USDObjectFieldListConstants::DesiredNumericEntryBoxWidth)
 					[
 						EntryBox
 					]
@@ -588,9 +596,9 @@ TSharedRef< SWidget > SUsdObjectFieldRow::GenerateSpinboxWidgets( const TArray<T
 
 		VertBox->AddSlot()
 			[
-				SNew( SBox )
-				.HeightOverride( FUsdStageEditorStyle::Get()->GetFloat( "UsdStageEditor.ListItemHeight" ) )
-				.VAlign( VAlign_Center )
+				SNew(SBox)
+				.HeightOverride(FUsdStageEditorStyle::Get()->GetFloat("UsdStageEditor.ListItemHeight"))
+				.VAlign(VAlign_Center)
 				[
 					HorizBox.ToSharedRef()
 				]
@@ -604,101 +612,83 @@ TSharedRef< SWidget > SUsdObjectFieldRow::GenerateSpinboxWidgets( const TArray<T
 		[
 			VertBox.ToSharedRef()
 		];
+	// clang-format on
 }
 
 template<typename T>
-void SUsdObjectFieldRow::OnSpinboxValueChanged( T NewValue, int32 ComponentIndex )
+void SUsdObjectFieldRow::OnSpinboxValueChanged(T NewValue, int32 ComponentIndex)
 {
-	Field->Value.Entries[ 0 ][ ComponentIndex ].Set<T>( NewValue );
+	Field->Value.Entries[0][ComponentIndex].Set<T>(NewValue);
 }
 
 template<typename T>
-void SUsdObjectFieldRow::OnSpinboxValueCommitted( T InNewValue, ETextCommit::Type CommitType, int32 ComponentIndex )
+void SUsdObjectFieldRow::OnSpinboxValueCommitted(T InNewValue, ETextCommit::Type CommitType, int32 ComponentIndex)
 {
-	if ( CommitType == ETextCommit::OnCleared || Field->bReadOnly )
+	if (CommitType == ETextCommit::OnCleared || Field->bReadOnly)
 	{
 		return;
 	}
 
-	FScopedTransaction Transaction(
-		FText::Format(
-			LOCTEXT( "SetUsdAttributeTransaction", "Set attribute '{0}'" ),
-			FText::FromString( Field->Label )
-		)
-	);
+	FScopedTransaction Transaction(FText::Format(LOCTEXT("SetUsdAttributeTransaction", "Set attribute '{0}'"), FText::FromString(Field->Label)));
 
 	UsdUtils::FConvertedVtValue NewValue = Field->Value;
-	NewValue.Entries[ 0 ][ ComponentIndex ].Set<T>( InNewValue );
+	NewValue.Entries[0][ComponentIndex].Set<T>(InNewValue);
 
-	Field->SetAttributeValue( NewValue );
+	Field->SetAttributeValue(NewValue);
 }
 
-void SUsdObjectFieldRow::OnComboBoxSelectionChanged( TSharedPtr<FString> ChosenOption, ESelectInfo::Type SelectInfo )
+void SUsdObjectFieldRow::OnComboBoxSelectionChanged(TSharedPtr<FString> ChosenOption, ESelectInfo::Type SelectInfo)
 {
-	if ( Field->bReadOnly )
+	if (Field->bReadOnly)
 	{
 		return;
 	}
 
-	FScopedTransaction Transaction(
-		FText::Format(
-			LOCTEXT( "SetUsdAttributeTransaction", "Set attribute '{0}'" ),
-			FText::FromString( Field->Label )
-		)
-	);
+	FScopedTransaction Transaction(FText::Format(LOCTEXT("SetUsdAttributeTransaction", "Set attribute '{0}'"), FText::FromString(Field->Label)));
 
 	UsdUtils::FConvertedVtValue NewValue = Field->Value;
-	NewValue.Entries[ 0 ][ 0 ] = UsdUtils::FConvertedVtValueComponent( TInPlaceType<FString>(), *ChosenOption );
+	NewValue.Entries[0][0] = UsdUtils::FConvertedVtValueComponent(TInPlaceType<FString>(), *ChosenOption);
 
-	Field->SetAttributeValue( NewValue );
+	Field->SetAttributeValue(NewValue);
 }
 
-void SUsdObjectFieldRow::OnTextBoxTextCommitted( const FText& NewText, ETextCommit::Type CommitType )
+void SUsdObjectFieldRow::OnTextBoxTextCommitted(const FText& NewText, ETextCommit::Type CommitType)
 {
-	if ( CommitType == ETextCommit::OnCleared || Field->bReadOnly )
+	if (CommitType == ETextCommit::OnCleared || Field->bReadOnly)
 	{
 		return;
 	}
 
-	FScopedTransaction Transaction(
-		FText::Format(
-			LOCTEXT( "SetUsdAttributeTransaction", "Set attribute '{0}'" ),
-			FText::FromString( Field->Label )
-		)
-	);
+	FScopedTransaction Transaction(FText::Format(LOCTEXT("SetUsdAttributeTransaction", "Set attribute '{0}'"), FText::FromString(Field->Label)));
 
 	UsdUtils::FConvertedVtValue NewValue = Field->Value;
-	NewValue.Entries[ 0 ][ 0 ] = UsdUtils::FConvertedVtValueComponent( TInPlaceType<FString>(), NewText.ToString() );
+	NewValue.Entries[0][0] = UsdUtils::FConvertedVtValueComponent(TInPlaceType<FString>(), NewText.ToString());
 
-	Field->SetAttributeValue( NewValue );
+	Field->SetAttributeValue(NewValue);
 }
 
-void SUsdObjectFieldRow::OnCheckBoxCheckStateChanged( ECheckBoxState NewState )
+void SUsdObjectFieldRow::OnCheckBoxCheckStateChanged(ECheckBoxState NewState)
 {
-	if ( NewState == ECheckBoxState::Undetermined || Field->bReadOnly )
+	if (NewState == ECheckBoxState::Undetermined || Field->bReadOnly)
 	{
 		return;
 	}
 
-	FScopedTransaction Transaction(
-		FText::Format(
-			LOCTEXT( "SetUsdAttributeTransaction", "Set attribute '{0}'" ),
-			FText::FromString( Field->Label )
-		)
-	);
+	FScopedTransaction Transaction(FText::Format(LOCTEXT("SetUsdAttributeTransaction", "Set attribute '{0}'"), FText::FromString(Field->Label)));
 
 	UsdUtils::FConvertedVtValue NewValue = Field->Value;
-	NewValue.Entries[ 0 ][ 0 ] = UsdUtils::FConvertedVtValueComponent( TInPlaceType<bool>(), NewState == ECheckBoxState::Checked );
+	NewValue.Entries[0][0] = UsdUtils::FConvertedVtValueComponent(TInPlaceType<bool>(), NewState == ECheckBoxState::Checked);
 
-	Field->SetAttributeValue( NewValue );
+	Field->SetAttributeValue(NewValue);
 }
 
-void SUsdObjectFieldList::Construct( const FArguments& InArgs )
+void SUsdObjectFieldList::Construct(const FArguments& InArgs)
 {
 	// Clear map as usd file may have additional Kinds now
 	USDObjectFieldListImpl::ResetOptions(TEXT("Kind"));
 
-	SAssignNew( HeaderRowWidget, SHeaderRow )
+	// clang-format off
+	SAssignNew(HeaderRowWidget, SHeaderRow)
 
 	+SHeaderRow::Column(ObjectFieldColumnIds::TypeColumn)
 	.DefaultLabel(FText::GetEmpty())
@@ -708,28 +698,32 @@ void SUsdObjectFieldList::Construct( const FArguments& InArgs )
 
 	+SHeaderRow::Column(ObjectFieldColumnIds::NameColumn)
 	.DefaultLabel(InArgs._NameColumnText)
-	.FillWidth( 40.f )
+	.FillWidth(40.f)
 	.SortMode(this, &SUsdObjectFieldList::GetColumnSortMode, ObjectFieldColumnIds::NameColumn)
 	.OnSort(this, &SUsdObjectFieldList::Sort)
 
 	+SHeaderRow::Column(ObjectFieldColumnIds::ValueColumn)
-	.DefaultLabel( LOCTEXT( "ValueColumnText", "Value" ) )
-	.FillWidth( 60.f );
+	.DefaultLabel(LOCTEXT("ValueColumnText", "Value"))
+	.FillWidth(60.f);
 
 	SListView::Construct
 	(
 		SListView::FArguments()
-		.ListItemsSource( &ViewModel.Fields )
-		.OnGenerateRow( this, &SUsdObjectFieldList::OnGenerateRow )
-		.HeaderRow( HeaderRowWidget )
+		.ListItemsSource(&ViewModel.Fields)
+		.OnGenerateRow(this, &SUsdObjectFieldList::OnGenerateRow)
+		.HeaderRow(HeaderRowWidget)
 	);
+	// clang-format on
 
 	OnSelectionChanged = InArgs._OnSelectionChanged;
 }
 
-TSharedRef< ITableRow > SUsdObjectFieldList::OnGenerateRow( TSharedPtr< FUsdObjectFieldViewModel > InDisplayNode, const TSharedRef< STableViewBase >& OwnerTable )
+TSharedRef<ITableRow> SUsdObjectFieldList::OnGenerateRow(
+	TSharedPtr<FUsdObjectFieldViewModel> InDisplayNode,
+	const TSharedRef<STableViewBase>& OwnerTable
+)
 {
-	return SNew( SUsdObjectFieldRow, InDisplayNode, OwnerTable );
+	return SNew(SUsdObjectFieldRow, InDisplayNode, OwnerTable);
 }
 
 void SUsdObjectFieldList::GenerateFieldList(const UE::FUsdStageWeak& UsdStage, const TCHAR* InObjectPath)
@@ -738,11 +732,7 @@ void SUsdObjectFieldList::GenerateFieldList(const UE::FUsdStageWeak& UsdStage, c
 	ViewModel.Refresh(UsdStage, InObjectPath, TimeCode);
 }
 
-void SUsdObjectFieldList::Sort(
-	const EColumnSortPriority::Type SortPriority,
-	const FName& ColumnId,
-	const EColumnSortMode::Type NewSortMode
-)
+void SUsdObjectFieldList::Sort(const EColumnSortPriority::Type SortPriority, const FName& ColumnId, const EColumnSortMode::Type NewSortMode)
 {
 	ViewModel.CurrentSortColumn = ColumnId;
 	ViewModel.CurrentSortMode = NewSortMode;
@@ -772,39 +762,39 @@ TArray<FString> SUsdObjectFieldList::GetSelectedFieldNames() const
 {
 	TArray<FString> SelectedProperties;
 
-	TArray< TSharedPtr< FUsdObjectFieldViewModel > > SelectedViewModels = GetSelectedItems();
-	SelectedProperties.Reserve( SelectedViewModels.Num() );
+	TArray<TSharedPtr<FUsdObjectFieldViewModel>> SelectedViewModels = GetSelectedItems();
+	SelectedProperties.Reserve(SelectedViewModels.Num());
 
-	for ( const TSharedPtr< FUsdObjectFieldViewModel >& SelectedItem : SelectedViewModels )
+	for (const TSharedPtr<FUsdObjectFieldViewModel>& SelectedItem : SelectedViewModels)
 	{
-		if ( SelectedItem )
+		if (SelectedItem)
 		{
-			SelectedProperties.Add( SelectedItem->Label );
+			SelectedProperties.Add(SelectedItem->Label);
 		}
 	}
 
 	return SelectedProperties;
 }
 
-void SUsdObjectFieldList::SetSelectedFieldNames( const TArray<FString>& NewSelection )
+void SUsdObjectFieldList::SetSelectedFieldNames(const TArray<FString>& NewSelection)
 {
-	TSet<FString> NewSelectionSet{ NewSelection };
+	TSet<FString> NewSelectionSet{NewSelection};
 
 	Private_ClearSelection();
 
-	TArray< TSharedPtr< FUsdObjectFieldViewModel > > NewItemSelection;
-	NewItemSelection.Reserve( NewSelection.Num() );
+	TArray<TSharedPtr<FUsdObjectFieldViewModel>> NewItemSelection;
+	NewItemSelection.Reserve(NewSelection.Num());
 
-	for ( const TSharedPtr< FUsdObjectFieldViewModel >& Item : ViewModel.Fields )
+	for (const TSharedPtr<FUsdObjectFieldViewModel>& Item : ViewModel.Fields)
 	{
-		if ( Item && NewSelectionSet.Contains( Item->Label ) )
+		if (Item && NewSelectionSet.Contains(Item->Label))
 		{
-			NewItemSelection.Add( Item );
+			NewItemSelection.Add(Item);
 		}
 	}
 
 	const bool bSelected = true;
-	SetItemSelection( NewItemSelection, bSelected );
+	SetItemSelection(NewItemSelection, bSelected);
 }
 
 UE::FUsdStageWeak SUsdObjectFieldList::GetUsdStage() const
@@ -819,4 +809,4 @@ FString SUsdObjectFieldList::GetObjectPath() const
 
 #undef LOCTEXT_NAMESPACE
 
-#endif // #if USE_USD_SDK
+#endif	  // #if USE_USD_SDK
diff --git a/Engine/Plugins/Importers/USDImporter/Source/USDStageEditor/Private/Widgets/SUSDReferencesList.cpp b/Engine/Plugins/Importers/USDImporter/Source/USDStageEditor/Private/Widgets/SUSDReferencesList.cpp
--- a/Engine/Plugins/Importers/USDImporter/Source/USDStageEditor/Private/Widgets/SUSDReferencesList.cpp
+++ b/Engine/Plugins/Importers/USDImporter/Source/USDStageEditor/Private/Widgets/SUSDReferencesList.cpp
@@ -3,69 +3,65 @@
 #include "SUSDReferencesList.h"
 
 #include "SUSDStageEditorStyle.h"
-#include "USDStageActor.h"
-#include "USDStageModule.h"
-#include "USDTypesConversion.h"
 
 #include "Styling/AppStyle.h"
 
-#include "Engine/World.h"
-#include "Modules/ModuleManager.h"
-#include "Widgets/Input/STextComboBox.h"
-
 #if USE_USD_SDK
 
 #define LOCTEXT_NAMESPACE "USDReferencesList"
 
 namespace UsdReferencesListConstants
 {
-	const FMargin RowPadding( 6.0f, 2.5f, 2.0f, 2.5f );
+	const FMargin RowPadding(6.0f, 2.5f, 2.0f, 2.5f);
 
 	const TCHAR* NormalFont = TEXT("PropertyWindow.NormalFont");
 }
 
-void SUsdReferenceRow::Construct( const FArguments& InArgs, TSharedPtr< FUsdReference > InReference, const TSharedRef< STableViewBase >& OwnerTable )
+void SUsdReferenceRow::Construct(const FArguments& InArgs, TSharedPtr<FUsdReference> InReference, const TSharedRef<STableViewBase>& OwnerTable)
 {
 	Reference = InReference;
 
-	SMultiColumnTableRow< TSharedPtr< FUsdReference > >::Construct( SMultiColumnTableRow< TSharedPtr< FUsdReference > >::FArguments(), OwnerTable );
+	SMultiColumnTableRow<TSharedPtr<FUsdReference>>::Construct(SMultiColumnTableRow<TSharedPtr<FUsdReference>>::FArguments(), OwnerTable);
 }
 
-TSharedRef< SWidget > SUsdReferenceRow::GenerateWidgetForColumn( const FName& ColumnName )
+TSharedRef<SWidget> SUsdReferenceRow::GenerateWidgetForColumn(const FName& ColumnName)
 {
-	TSharedRef< SWidget > ColumnWidget = SNullWidget::NullWidget;
+	TSharedRef<SWidget> ColumnWidget = SNullWidget::NullWidget;
 
-	if ( ColumnName == TEXT("AssetPath") )
+	// clang-format off
+	if (ColumnName == TEXT("AssetPath"))
 	{
 		SAssignNew(ColumnWidget, STextBlock)
 		.Text(FText::FromString(Reference->AssetPath.IsEmpty() ? TEXT("(internal reference)") : Reference->AssetPath))
 		.Font(FAppStyle::GetFontStyle(UsdReferencesListConstants::NormalFont));
 	}
 
-	if ( ColumnName == TEXT("PrimPath") )
+	if (ColumnName == TEXT("PrimPath"))
 	{
 		SAssignNew(ColumnWidget, STextBlock)
 		.Text(FText::FromString(Reference->PrimPath.IsEmpty() ? TEXT("(default prim)") : Reference->PrimPath))
 		.Font(FAppStyle::GetFontStyle(UsdReferencesListConstants::NormalFont));
 	}
 
-	return SNew( SBox )
-		.HeightOverride( FUsdStageEditorStyle::Get()->GetFloat( "UsdStageEditor.ListItemHeight" ) )
+	return SNew(SBox)
+		.HeightOverride(FUsdStageEditorStyle::Get()->GetFloat("UsdStageEditor.ListItemHeight"))
 		[
-			SNew( SHorizontalBox )
+			SNew(SHorizontalBox)
 			+SHorizontalBox::Slot()
-			.HAlign( HAlign_Left )
-			.VAlign( VAlign_Center )
-			.Padding( UsdReferencesListConstants::RowPadding )
+			.HAlign(HAlign_Left)
+			.VAlign(VAlign_Center)
+			.Padding(UsdReferencesListConstants::RowPadding)
 			.AutoWidth()
 			[
 				ColumnWidget
 			]
 		];
+	// clang-format on
 }
 
-void SUsdReferencesList::Construct( const FArguments& InArgs )
+void SUsdReferencesList::Construct(const FArguments& InArgs)
 {
+	// clang-format off
 	SAssignNew(HeaderRowWidget, SHeaderRow)
 
 	+SHeaderRow::Column(FName(TEXT("AssetPath")))
@@ -79,28 +75,29 @@ void SUsdReferencesList::Construct( const FArguments& InArgs )
 	SListView::Construct
 	(
 		SListView::FArguments()
-		.ListItemsSource( &ViewModel.References )
-		.OnGenerateRow( this, &SUsdReferencesList::OnGenerateRow )
-		.HeaderRow( HeaderRowWidget )
+		.ListItemsSource(&ViewModel.References)
+		.OnGenerateRow(this, &SUsdReferencesList::OnGenerateRow)
+		.HeaderRow(HeaderRowWidget)
 	);
 
-	SetVisibility( EVisibility::Collapsed ); // Start hidden until SetPrimPath displays us
+	SetVisibility(EVisibility::Collapsed); // Start hidden until SetPrimPath displays us
+	// clang-format on
 }
 
-TSharedRef< ITableRow > SUsdReferencesList::OnGenerateRow( TSharedPtr< FUsdReference > InDisplayNode, const TSharedRef< STableViewBase >& OwnerTable )
+TSharedRef<ITableRow> SUsdReferencesList::OnGenerateRow(TSharedPtr<FUsdReference> InDisplayNode, const TSharedRef<STableViewBase>& OwnerTable)
 {
-	return SNew( SUsdReferenceRow, InDisplayNode, OwnerTable );
+	return SNew(SUsdReferenceRow, InDisplayNode, OwnerTable);
 }
 
-void SUsdReferencesList::SetPrimPath( const UE::FUsdStageWeak& UsdStage, const TCHAR* PrimPath )
+void SUsdReferencesList::SetPrimPath(const UE::FUsdStageWeak& UsdStage, const TCHAR* PrimPath)
 {
-	ViewModel.UpdateReferences( UsdStage, PrimPath );
+	ViewModel.UpdateReferences(UsdStage, PrimPath);
 
-	SetVisibility( ViewModel.References.Num() > 0 ? EVisibility::Visible : EVisibility::Collapsed );
+	SetVisibility(ViewModel.References.Num() > 0 ? EVisibility::Visible : EVisibility::Collapsed);
 
 	RequestListRefresh();
 }
 
 #undef LOCTEXT_NAMESPACE
 
-#endif // #if USE_USD_SDK
+#endif	  // #if USE_USD_SDK
diff --git a/Engine/Plugins/Importers/USDImporter/Source/USDStageEditor/Private/Widgets/SUSDVariantSetsList.cpp b/Engine/Plugins/Importers/USDImporter/Source/USDStageEditor/Private/Widgets/SUSDVariantSetsList.cpp
--- a/Engine/Plugins/Importers/USDImporter/Source/USDStageEditor/Private/Widgets/SUSDVariantSetsList.cpp
+++ b/Engine/Plugins/Importers/USDImporter/Source/USDStageEditor/Private/Widgets/SUSDVariantSetsList.cpp
@@ -2,129 +2,136 @@
 
 #include "SUSDVariantSetsList.h"
 
-#include "SUSDStageEditorStyle.h"
-#include "USDStageActor.h"
-#include "USDStageModule.h"
-#include "USDTypesConversion.h"
-
 #include "Styling/AppStyle.h"
-
-#include "Engine/World.h"
-#include "Modules/ModuleManager.h"
-#include "ScopedTransaction.h"
+#include "SUSDStageEditorStyle.h"
 #include "Widgets/Input/STextComboBox.h"
 
 #if USE_USD_SDK
 
 namespace UsdVariantSetsListConstants
 {
-	const FMargin LeftRowPadding( 6.0f, 0.0f, 2.0f, 0.0f );
-	const FMargin RightRowPadding( 3.0f, 0.0f, 2.0f, 0.0f );
+	const FMargin LeftRowPadding(6.0f, 0.0f, 2.0f, 0.0f);
+	const FMargin RightRowPadding(3.0f, 0.0f, 2.0f, 0.0f);
 
 	const TCHAR* NormalFont = TEXT("PropertyWindow.NormalFont");
 }
 
-void SUsdVariantRow::Construct( const FArguments& InArgs, TSharedPtr< FUsdVariantSetViewModel > InVariantSet, const TSharedRef< STableViewBase >& OwnerTable )
+void SUsdVariantRow::Construct(
+	const FArguments& InArgs,
+	TSharedPtr<FUsdVariantSetViewModel> InVariantSet,
+	const TSharedRef<STableViewBase>& OwnerTable
+)
 {
 	OnVariantSelectionChanged = InArgs._OnVariantSelectionChanged;
 
 	VariantSet = InVariantSet;
 
-	SMultiColumnTableRow< TSharedPtr< FUsdVariantSetViewModel > >::Construct( SMultiColumnTableRow< TSharedPtr< FUsdVariantSetViewModel > >::FArguments(), OwnerTable );
+	SMultiColumnTableRow<TSharedPtr<FUsdVariantSetViewModel>>::Construct(
+		SMultiColumnTableRow<TSharedPtr<FUsdVariantSetViewModel>>::FArguments(),
+		OwnerTable
+	);
 }
 
-TSharedRef< SWidget > SUsdVariantRow::GenerateWidgetForColumn( const FName& ColumnName )
+TSharedRef<SWidget> SUsdVariantRow::GenerateWidgetForColumn(const FName& ColumnName)
 {
-	TSharedRef< SWidget > ColumnWidget = SNullWidget::NullWidget;
+	TSharedRef<SWidget> ColumnWidget = SNullWidget::NullWidget;
 
 	bool bIsLeftRow = true;
 
-	if ( ColumnName == TEXT("VariantSetName") )
+	if (ColumnName == TEXT("VariantSetName"))
 	{
-		SAssignNew( ColumnWidget, STextBlock )
-		.Text( FText::FromString( VariantSet->SetName ) )
-		.Font( FAppStyle::GetFontStyle( UsdVariantSetsListConstants::NormalFont ) );
+		// clang-format off
+		SAssignNew(ColumnWidget, STextBlock)
+		.Text(FText::FromString(VariantSet->SetName))
+		.Font(FAppStyle::GetFontStyle(UsdVariantSetsListConstants::NormalFont));
+		// clang-format on
 	}
 	else
 	{
 		bIsLeftRow = false;
 
-		TSharedPtr< FString >* InitialSelectionPtr = VariantSet->Variants.FindByPredicate(
-			[ VariantSelection = VariantSet->VariantSelection ]( const TSharedPtr< FString >& A )
+		TSharedPtr<FString>* InitialSelectionPtr = VariantSet->Variants.FindByPredicate(
+			[VariantSelection = VariantSet->VariantSelection](const TSharedPtr<FString>& A)
 			{
-				return A->Equals( *VariantSelection, ESearchCase::IgnoreCase );
-			} );
-
-		TSharedPtr< FString > InitialSelection = InitialSelectionPtr ? *InitialSelectionPtr : TSharedPtr< FString >();
-
-		SAssignNew( ColumnWidget, STextComboBox )
-				.OptionsSource( &VariantSet->Variants )
-				.InitiallySelectedItem( InitialSelection )
-				.OnSelectionChanged( this, &SUsdVariantRow::OnSelectionChanged )
-				.Font( FAppStyle::GetFontStyle( UsdVariantSetsListConstants::NormalFont ) );
+				return A->Equals(*VariantSelection, ESearchCase::IgnoreCase);
+			}
+		);
+
+		TSharedPtr<FString> InitialSelection = InitialSelectionPtr ? *InitialSelectionPtr : TSharedPtr<FString>();
+
+		// clang-format off
+		SAssignNew(ColumnWidget, STextComboBox)
+		.OptionsSource(&VariantSet->Variants)
+		.InitiallySelectedItem(InitialSelection)
+		.OnSelectionChanged(this, &SUsdVariantRow::OnSelectionChanged)
+		.Font(FAppStyle::GetFontStyle(UsdVariantSetsListConstants::NormalFont));
+		// clang-format on
 	}
 
+	// clang-format off
 	return SNew(SBox)
-		.HeightOverride( FUsdStageEditorStyle::Get()->GetFloat( "UsdStageEditor.ListItemHeight" ) )
+		.HeightOverride(FUsdStageEditorStyle::Get()->GetFloat("UsdStageEditor.ListItemHeight"))
 		[
-			SNew( SHorizontalBox )
+			SNew(SHorizontalBox)
 			+ SHorizontalBox::Slot()
-			.HAlign( HAlign_Left )
-			.VAlign( VAlign_Center )
-			.Padding( bIsLeftRow ? UsdVariantSetsListConstants::LeftRowPadding : UsdVariantSetsListConstants::RightRowPadding )
+			.HAlign(HAlign_Left)
+			.VAlign(VAlign_Center)
+			.Padding(bIsLeftRow ? UsdVariantSetsListConstants::LeftRowPadding : UsdVariantSetsListConstants::RightRowPadding)
 			.AutoWidth()
 			[
 				ColumnWidget
 			]
 		];
+	// clang-format on
 }
 
-void SUsdVariantRow::OnSelectionChanged( TSharedPtr< FString > NewValue, ESelectInfo::Type SelectInfo )
+void SUsdVariantRow::OnSelectionChanged(TSharedPtr<FString> NewValue, ESelectInfo::Type SelectInfo)
 {
-	OnVariantSelectionChanged.ExecuteIfBound( VariantSet.ToSharedRef(), NewValue );
+	OnVariantSelectionChanged.ExecuteIfBound(VariantSet.ToSharedRef(), NewValue);
 }
 
-void SVariantsList::Construct( const FArguments& InArgs )
+void SVariantsList::Construct(const FArguments& InArgs)
 {
-	SAssignNew( HeaderRowWidget, SHeaderRow )
+	// clang-format off
+	SAssignNew(HeaderRowWidget, SHeaderRow)
 
-	+SHeaderRow::Column( FName( TEXT("VariantSetName") ) )
-	.DefaultLabel( NSLOCTEXT( "USDVariantSetsList", "VariantSetName", "Variants" ) )
-	.FillWidth( 25.f )
+	+SHeaderRow::Column(FName(TEXT("VariantSetName")))
+	.DefaultLabel(NSLOCTEXT("USDVariantSetsList", "VariantSetName", "Variants"))
+	.FillWidth(25.f)
 
-	+SHeaderRow::Column( FName( TEXT("VariantSetSelection") ) )
-	.DefaultLabel( FText::GetEmpty() )
-	.FillWidth( 75.f );
+	+SHeaderRow::Column(FName(TEXT("VariantSetSelection")))
+	.DefaultLabel(FText::GetEmpty())
+	.FillWidth(75.f);
 
 	SListView::Construct
 	(
 		SListView::FArguments()
-		.ListItemsSource( &ViewModel.VariantSets )
-		.OnGenerateRow( this, &SVariantsList::OnGenerateRow )
-		.HeaderRow( HeaderRowWidget )
+		.ListItemsSource(&ViewModel.VariantSets)
+		.OnGenerateRow(this, &SVariantsList::OnGenerateRow)
+		.HeaderRow(HeaderRowWidget)
 	);
 
-	SetVisibility( EVisibility::Collapsed ); // Start hidden until SetPrimPath displays us
+	SetVisibility(EVisibility::Collapsed); // Start hidden until SetPrimPath displays us
+	// clang-format on
 }
 
-TSharedRef< ITableRow > SVariantsList::OnGenerateRow( TSharedPtr< FUsdVariantSetViewModel > InDisplayNode, const TSharedRef< STableViewBase >& OwnerTable )
+TSharedRef<ITableRow> SVariantsList::OnGenerateRow(TSharedPtr<FUsdVariantSetViewModel> InDisplayNode, const TSharedRef<STableViewBase>& OwnerTable)
 {
-	return SNew( SUsdVariantRow, InDisplayNode, OwnerTable )
-			.OnVariantSelectionChanged( this, &SVariantsList::OnVariantSelectionChanged );
+	return SNew(SUsdVariantRow, InDisplayNode, OwnerTable).OnVariantSelectionChanged(this, &SVariantsList::OnVariantSelectionChanged);
 }
 
-void SVariantsList::SetPrimPath( const UE::FUsdStageWeak& UsdStage, const TCHAR* InPrimPath )
+void SVariantsList::SetPrimPath(const UE::FUsdStageWeak& UsdStage, const TCHAR* InPrimPath)
 {
-	ViewModel.UpdateVariantSets( UsdStage, InPrimPath );
+	ViewModel.UpdateVariantSets(UsdStage, InPrimPath);
 
-	SetVisibility( ViewModel.VariantSets.Num() > 0 ? EVisibility::Visible : EVisibility::Collapsed );
+	SetVisibility(ViewModel.VariantSets.Num() > 0 ? EVisibility::Visible : EVisibility::Collapsed);
 
 	RequestListRefresh();
 }
 
-void SVariantsList::OnVariantSelectionChanged( const TSharedRef< FUsdVariantSetViewModel >& VariantSet, const TSharedPtr< FString >& NewValue )
+void SVariantsList::OnVariantSelectionChanged(const TSharedRef<FUsdVariantSetViewModel>& VariantSet, const TSharedPtr<FString>& NewValue)
 {
-	VariantSet->SetVariantSelection( NewValue );
+	VariantSet->SetVariantSelection(NewValue);
 }
 
-#endif // #if USE_USD_SDK
+#endif	  // #if USE_USD_SDK
