USD: Reformat USDStageEditorViewModels module.

#rnx

diff --git a/Engine/Plugins/Importers/USDImporter/Source/USDStageEditorViewModels/Private/USDLayersViewModel.cpp b/Engine/Plugins/Importers/USDImporter/Source/USDStageEditorViewModels/Private/USDLayersViewModel.cpp
--- a/Engine/Plugins/Importers/USDImporter/Source/USDStageEditorViewModels/Private/USDLayersViewModel.cpp
+++ b/Engine/Plugins/Importers/USDImporter/Source/USDStageEditorViewModels/Private/USDLayersViewModel.cpp
@@ -9,19 +9,14 @@
 #include "UsdWrappers/SdfLayer.h"
 
 #if USE_USD_SDK
-
 #include "USDIncludesStart.h"
-	#include "pxr/pxr.h"
-	#include "pxr/usd/pcp/cache.h"
-	#include "pxr/usd/pcp/layerStack.h"
-	#include "pxr/usd/sdf/layer.h"
-	#include "pxr/usd/sdf/layerUtils.h"
-	#include "pxr/usd/sdf/layerTree.h"
-	#include "pxr/usd/usd/stage.h"
-	#include "pxr/usd/usdUtils/dependencies.h"
+#include "pxr/usd/pcp/layerStack.h"
+#include "pxr/usd/sdf/layer.h"
+#include "pxr/usd/sdf/layerTree.h"
+#include "pxr/usd/sdf/layerUtils.h"
+#include "pxr/usd/usd/stage.h"
 #include "USDIncludesEnd.h"
-
-#endif // #if USE_USD_SDK
+#endif	  // #if USE_USD_SDK
 
 FUsdLayerViewModel::FUsdLayerViewModel(
 	FUsdLayerViewModel* InParentItem,
@@ -29,23 +24,23 @@ FUsdLayerViewModel::FUsdLayerViewModel(
 	const UE::FUsdStageWeak& InIsolatedStage,
 	const FString& InLayerIdentifier
 )
-	: LayerModel( MakeShared< FUsdLayerModel >() )
-	, ParentItem( InParentItem )
-	, UsdStage( InUsdStage )
-	, IsolatedStage( InIsolatedStage )
-	, LayerIdentifier( InLayerIdentifier )
+	: LayerModel(MakeShared<FUsdLayerModel>())
+	, ParentItem(InParentItem)
+	, UsdStage(InUsdStage)
+	, IsolatedStage(InIsolatedStage)
+	, LayerIdentifier(InLayerIdentifier)
 {
 	RefreshData();
 }
 
 bool FUsdLayerViewModel::IsValid() const
 {
-	return ( bool ) UsdStage && ( !ParentItem || !ParentItem->LayerIdentifier.Equals( LayerIdentifier, ESearchCase::Type::IgnoreCase ) );
+	return (bool)UsdStage && (!ParentItem || !ParentItem->LayerIdentifier.Equals(LayerIdentifier, ESearchCase::Type::IgnoreCase));
 }
 
-TArray< TSharedRef< FUsdLayerViewModel > > FUsdLayerViewModel::GetChildren()
+TArray<TSharedRef<FUsdLayerViewModel>> FUsdLayerViewModel::GetChildren()
 {
-	if ( !IsValid() )
+	if (!IsValid())
 	{
 		return {};
 	}
@@ -56,16 +51,17 @@ TArray< TSharedRef< FUsdLayerViewModel > > FUsdLayerViewModel::GetChildren()
 	{
 		FScopedUsdAllocs UsdAllocs;
 
-		pxr::SdfLayerRefPtr UsdLayer( GetLayer() );
+		pxr::SdfLayerRefPtr UsdLayer(GetLayer());
 
-		if ( UsdLayer )
+		if (UsdLayer)
 		{
 			int32 SubLayerIndex = 0;
-			for ( const std::string& SubLayerPath : UsdLayer->GetSubLayerPaths() )
+			for (const std::string& SubLayerPath : UsdLayer->GetSubLayerPaths())
 			{
-				const FString SubLayerIdentifier = UsdToUnreal::ConvertString( pxr::SdfComputeAssetPathRelativeToLayer( UsdLayer, SubLayerPath ) );
+				const FString SubLayerIdentifier = UsdToUnreal::ConvertString(pxr::SdfComputeAssetPathRelativeToLayer(UsdLayer, SubLayerPath));
 
-				if ( !Children.IsValidIndex( SubLayerIndex ) || !Children[ SubLayerIndex ]->LayerIdentifier.Equals( SubLayerIdentifier, ESearchCase::Type::IgnoreCase ) )
+				if (!Children.IsValidIndex(SubLayerIndex)
+					|| !Children[SubLayerIndex]->LayerIdentifier.Equals(SubLayerIdentifier, ESearchCase::Type::IgnoreCase))
 				{
 					Children.Reset();
 					bNeedsRefresh = true;
@@ -75,16 +71,16 @@ TArray< TSharedRef< FUsdLayerViewModel > > FUsdLayerViewModel::GetChildren()
 				++SubLayerIndex;
 			}
 
-			if ( !bNeedsRefresh && SubLayerIndex < Children.Num() )
+			if (!bNeedsRefresh && SubLayerIndex < Children.Num())
 			{
 				Children.Reset();
 				bNeedsRefresh = true;
 			}
 		}
 	}
-#endif // #if USE_USD_SDK
+#endif	  // #if USE_USD_SDK
 
-	if ( bNeedsRefresh )
+	if (bNeedsRefresh)
 	{
 		FillChildren();
 	}
@@ -96,7 +92,7 @@ void FUsdLayerViewModel::FillChildren()
 {
 	Children.Reset();
 
-	if ( !IsValid() )
+	if (!IsValid())
 	{
 		return;
 	}
@@ -105,37 +101,37 @@ void FUsdLayerViewModel::FillChildren()
 	{
 		FScopedUsdAllocs UsdAllocs;
 
-		pxr::SdfLayerRefPtr UsdLayer( GetLayer() );
+		pxr::SdfLayerRefPtr UsdLayer(GetLayer());
 
-		if ( UsdLayer )
+		if (UsdLayer)
 		{
-			TSet< FString > AllLayerIdentifiers;
+			TSet<FString> AllLayerIdentifiers;
 
 			FUsdLayerViewModel* CurrentItem = this;
-			while ( CurrentItem )
+			while (CurrentItem)
 			{
-				AllLayerIdentifiers.Add( CurrentItem->LayerIdentifier );
+				AllLayerIdentifiers.Add(CurrentItem->LayerIdentifier);
 				CurrentItem = CurrentItem->ParentItem;
 			}
 
-			for ( std::string SubLayerPath : UsdLayer->GetSubLayerPaths() )
+			for (std::string SubLayerPath : UsdLayer->GetSubLayerPaths())
 			{
-				FString AssetPathRelativeToLayer = UsdToUnreal::ConvertString( pxr::SdfComputeAssetPathRelativeToLayer( UsdLayer, SubLayerPath ) );
+				FString AssetPathRelativeToLayer = UsdToUnreal::ConvertString(pxr::SdfComputeAssetPathRelativeToLayer(UsdLayer, SubLayerPath));
 
 				// Prevent infinite recursions if a sublayer refers to a parent of the same hierarchy
-				if ( !AllLayerIdentifiers.Contains( AssetPathRelativeToLayer ) )
+				if (!AllLayerIdentifiers.Contains(AssetPathRelativeToLayer))
 				{
-					Children.Add( MakeShared< FUsdLayerViewModel >( this, UsdStage, IsolatedStage, AssetPathRelativeToLayer ) );
+					Children.Add(MakeShared<FUsdLayerViewModel>(this, UsdStage, IsolatedStage, AssetPathRelativeToLayer));
 				}
 			}
 		}
 	}
-#endif // #if USE_USD_SDK
+#endif	  // #if USE_USD_SDK
 }
 
 void FUsdLayerViewModel::RefreshData()
 {
-	if ( !IsValid() )
+	if (!IsValid())
 	{
 		return;
 	}
@@ -145,17 +141,15 @@ void FUsdLayerViewModel::RefreshData()
 #if USE_USD_SDK
 	FScopedUsdAllocs UsdAllocs;
 
-	pxr::UsdStageRefPtr UsdStageRef( UsdStage );
-	pxr::UsdStageRefPtr IsolatedStageRef( IsolatedStage );
+	pxr::UsdStageRefPtr UsdStageRef(UsdStage);
+	pxr::UsdStageRefPtr IsolatedStageRef(IsolatedStage);
 
-	const std::string UsdLayerIdentifier = UnrealToUsd::ConvertString( *LayerIdentifier ).Get();
+	const std::string UsdLayerIdentifier = UnrealToUsd::ConvertString(*LayerIdentifier).Get();
 
-	LayerModel->DisplayName = UsdToUnreal::ConvertString(
-		pxr::SdfLayer::GetDisplayNameFromIdentifier( UsdLayerIdentifier )
-	);
+	LayerModel->DisplayName = UsdToUnreal::ConvertString(pxr::SdfLayer::GetDisplayNameFromIdentifier(UsdLayerIdentifier));
 
-	pxr::SdfLayerRefPtr UsdLayer = UE::FSdfLayer::FindOrOpen( *LayerIdentifier );
-	if ( UsdLayer )
+	pxr::SdfLayerRefPtr UsdLayer = UE::FSdfLayer::FindOrOpen(*LayerIdentifier);
+	if (UsdLayer)
 	{
 		LayerModel->bIsDirty = UsdLayer->IsDirty() && !UsdLayer->IsAnonymous();
 	}
@@ -165,42 +159,36 @@ void FUsdLayerViewModel::RefreshData()
 	// Note: Originally this was a proper check to see if UsdLayer was inside IsolatedStage's layer stack. However, this
 	// breaks down if the layer is also muted: When that happens the stages basically pretend the layer doesn't exist,
 	// and so any muted layer on the base stage would behave as if it was "non-isolated"
-	LayerModel->bIsInIsolatedStage = !IsolatedStageRef
-		|| bIsIsolatedRootLayer
-		|| ( ParentItem && ParentItem->IsInIsolatedStage() );
+	LayerModel->bIsInIsolatedStage = !IsolatedStageRef || bIsIsolatedRootLayer || (ParentItem && ParentItem->IsInIsolatedStage());
 
 	const pxr::SdfLayerHandle& EditTargetLayer = UsdStageRef->GetEditTarget().GetLayer();
-	const pxr::SdfLayerHandle& FocusedEditTarget = IsolatedStageRef
-		? IsolatedStageRef->GetEditTarget().GetLayer()
-		: pxr::SdfLayerHandle{};
+	const pxr::SdfLayerHandle& FocusedEditTarget = IsolatedStageRef ? IsolatedStageRef->GetEditTarget().GetLayer() : pxr::SdfLayerHandle{};
 
-	LayerModel->bIsEditTarget = ( IsolatedStageRef )
-		? LayerModel->bIsInIsolatedStage && FocusedEditTarget == UsdLayer
-		: EditTargetLayer == UsdLayer;
+	LayerModel->bIsEditTarget = (IsolatedStageRef) ? LayerModel->bIsInIsolatedStage && FocusedEditTarget == UsdLayer : EditTargetLayer == UsdLayer;
 
-	LayerModel->bIsMuted = ( IsolatedStageRef && LayerModel->bIsInIsolatedStage )
-		// If we isolating, we're only muted if we're muted on the isolated stage
-		? IsolatedStageRef->IsLayerMuted( UsdLayerIdentifier )
-		// Otherwise we're muted any time we're muted in the base stage
-		: UsdStageRef->IsLayerMuted( UsdLayerIdentifier );
+	LayerModel->bIsMuted = (IsolatedStageRef && LayerModel->bIsInIsolatedStage)
+							   // If we isolating, we're only muted if we're muted on the isolated stage
+							   ? IsolatedStageRef->IsLayerMuted(UsdLayerIdentifier)
+							   // Otherwise we're muted any time we're muted in the base stage
+							   : UsdStageRef->IsLayerMuted(UsdLayerIdentifier);
 
-	for ( const TSharedRef< FUsdLayerViewModel >& Child : Children )
+	for (const TSharedRef<FUsdLayerViewModel>& Child : Children)
 	{
 		Child->RefreshData();
 	}
-#endif // #if USE_USD_SDK
+#endif	  // #if USE_USD_SDK
 }
 
 UE::FSdfLayer FUsdLayerViewModel::GetLayer() const
 {
-	return UE::FSdfLayer::FindOrOpen( *LayerIdentifier );
+	return UE::FSdfLayer::FindOrOpen(*LayerIdentifier);
 }
 
 FText FUsdLayerViewModel::GetDisplayName() const
 {
 	FString DisplayName = LayerModel->DisplayName;
 
-	if ( IsLayerDirty() )
+	if (IsLayerDirty())
 	{
 		DisplayName += TEXT("*");
 	}
@@ -215,20 +203,19 @@ bool FUsdLayerViewModel::IsLayerMuted() const
 
 bool FUsdLayerViewModel::CanMuteLayer() const
 {
-	if ( !IsValid() )
+	if (!IsValid())
 	{
 		return false;
 	}
 
 	UE::FUsdStageWeak AffectedStage = IsolatedStage ? IsolatedStage : UsdStage;
 
-	return !AffectedStage.GetRootLayer().GetIdentifier().Equals( LayerIdentifier, ESearchCase::Type::IgnoreCase )
-		&& !LayerModel->bIsEditTarget;
+	return !AffectedStage.GetRootLayer().GetIdentifier().Equals(LayerIdentifier, ESearchCase::Type::IgnoreCase) && !LayerModel->bIsEditTarget;
 }
 
 void FUsdLayerViewModel::ToggleMuteLayer()
 {
-	if ( !IsValid() || !CanMuteLayer() )
+	if (!IsValid() || !CanMuteLayer())
 	{
 		return;
 	}
@@ -236,37 +223,37 @@ void FUsdLayerViewModel::ToggleMuteLayer()
 #if USE_USD_SDK
 	FScopedUsdAllocs UsdAllocs;
 
-	const TUsdStore< std::string > UsdLayerIdentifier = UnrealToUsd::ConvertString( *LayerIdentifier );
+	const TUsdStore<std::string> UsdLayerIdentifier = UnrealToUsd::ConvertString(*LayerIdentifier);
 
-	pxr::UsdStageRefPtr CurrentStage( IsolatedStage ? IsolatedStage : UsdStage );
+	pxr::UsdStageRefPtr CurrentStage(IsolatedStage ? IsolatedStage : UsdStage);
 
-	if ( CurrentStage->IsLayerMuted( UsdLayerIdentifier.Get() ) )
+	if (CurrentStage->IsLayerMuted(UsdLayerIdentifier.Get()))
 	{
-		CurrentStage->UnmuteLayer( UsdLayerIdentifier.Get() );
+		CurrentStage->UnmuteLayer(UsdLayerIdentifier.Get());
 	}
 	else
 	{
-		CurrentStage->MuteLayer( UsdLayerIdentifier.Get() );
+		CurrentStage->MuteLayer(UsdLayerIdentifier.Get());
 	}
 
 	// We have to refresh all of our view models here, because this layer can appear more than once on the layer
 	// stack and we want all of them to show the updated muted/unmuted icon
 	FUsdLayerViewModel* Ancestor = this;
-	while ( true )
+	while (true)
 	{
 		FUsdLayerViewModel* AncestorParent = Ancestor->ParentItem;
-		if ( !AncestorParent )
+		if (!AncestorParent)
 		{
 			break;
 		}
 
 		Ancestor = AncestorParent;
 	}
-	if ( Ancestor )
+	if (Ancestor)
 	{
 		Ancestor->RefreshData();
 	}
-#endif // #if USE_USD_SDK
+#endif	  // #if USE_USD_SDK
 }
 
 bool FUsdLayerViewModel::IsEditTarget() const
@@ -281,16 +268,16 @@ bool FUsdLayerViewModel::CanEditLayer() const
 
 bool FUsdLayerViewModel::EditLayer()
 {
-	UE::FSdfLayer Layer( GetLayer() );
+	UE::FSdfLayer Layer(GetLayer());
 
 	UE::FUsdStageWeak AffectedStage = IsolatedStage ? IsolatedStage : UsdStage;
 
-	if ( !AffectedStage || !Layer || !CanEditLayer() )
+	if (!AffectedStage || !Layer || !CanEditLayer())
 	{
 		return false;
 	}
 
-	AffectedStage.SetEditTarget( Layer );
+	AffectedStage.SetEditTarget(Layer);
 	RefreshData();
 
 	return true;
@@ -301,37 +288,37 @@ bool FUsdLayerViewModel::IsInIsolatedStage() const
 	return LayerModel->bIsInIsolatedStage;
 }
 
-void FUsdLayerViewModel::AddSubLayer( const TCHAR* SubLayerIdentifier )
+void FUsdLayerViewModel::AddSubLayer(const TCHAR* SubLayerIdentifier)
 {
 #if USE_USD_SDK
 	FScopedUsdAllocs UsdAllocs;
-	UsdUtils::InsertSubLayer( pxr::SdfLayerRefPtr( GetLayer() ), SubLayerIdentifier );
-#endif // #if USE_USD_SDK
+	UsdUtils::InsertSubLayer(pxr::SdfLayerRefPtr(GetLayer()), SubLayerIdentifier);
+#endif	  // #if USE_USD_SDK
 }
 
-void FUsdLayerViewModel::NewSubLayer( const TCHAR* SubLayerIdentifier )
+void FUsdLayerViewModel::NewSubLayer(const TCHAR* SubLayerIdentifier)
 {
 #if USE_USD_SDK
 	FScopedUsdAllocs UsdAllocs;
-	UsdUtils::CreateNewLayer( pxr::UsdStageRefPtr( UsdStage ), pxr::SdfLayerRefPtr( GetLayer() ), SubLayerIdentifier );
-#endif // #if USE_USD_SDK
+	UsdUtils::CreateNewLayer(pxr::UsdStageRefPtr(UsdStage), pxr::SdfLayerRefPtr(GetLayer()), SubLayerIdentifier);
+#endif	  // #if USE_USD_SDK
 }
 
-bool FUsdLayerViewModel::RemoveSubLayer( int32 SubLayerIndex )
+bool FUsdLayerViewModel::RemoveSubLayer(int32 SubLayerIndex)
 {
 	bool bLayerRemoved = false;
 
 #if USE_USD_SDK
 	FScopedUsdAllocs UsdAllocs;
 
-	pxr::SdfLayerRefPtr UsdLayer( GetLayer() );
+	pxr::SdfLayerRefPtr UsdLayer(GetLayer());
 
-	if ( UsdLayer )
+	if (UsdLayer)
 	{
-		UsdLayer->RemoveSubLayerPath( SubLayerIndex );
+		UsdLayer->RemoveSubLayerPath(SubLayerIndex);
 		bLayerRemoved = true;
 	}
-#endif // #if USE_USD_SDK
+#endif	  // #if USE_USD_SDK
 
 	return bLayerRemoved;
 }
diff --git a/Engine/Plugins/Importers/USDImporter/Source/USDStageEditorViewModels/Private/USDObjectFieldViewModel.cpp b/Engine/Plugins/Importers/USDImporter/Source/USDStageEditorViewModels/Private/USDObjectFieldViewModel.cpp
--- a/Engine/Plugins/Importers/USDImporter/Source/USDStageEditorViewModels/Private/USDObjectFieldViewModel.cpp
+++ b/Engine/Plugins/Importers/USDImporter/Source/USDStageEditorViewModels/Private/USDObjectFieldViewModel.cpp
@@ -5,7 +5,6 @@
 #include "UnrealUSDWrapper.h"
 #include "USDAttributeUtils.h"
 #include "USDErrorUtils.h"
-#include "USDLog.h"
 #include "USDMemory.h"
 #include "USDTypesConversion.h"
 #include "USDValueConversion.h"
@@ -21,87 +20,81 @@
 #include "Widgets/Notifications/SNotificationList.h"
 
 #if USE_USD_SDK
-
 #include "USDIncludesStart.h"
-	#include "pxr/base/tf/stringUtils.h"
-	#include "pxr/base/vt/value.h"
-	#include "pxr/usd/kind/registry.h"
-	#include "pxr/usd/sdf/types.h"
-	#include "pxr/usd/usd/attribute.h"
-	#include "pxr/usd/usd/common.h"
-	#include "pxr/usd/usd/editContext.h"
-	#include "pxr/usd/usd/prim.h"
-	#include "pxr/usd/usd/relationship.h"
-	#include "pxr/usd/usdGeom/tokens.h"
-	#include "pxr/usd/usdPhysics/tokens.h"
+#include "pxr/base/tf/stringUtils.h"
+#include "pxr/base/vt/value.h"
+#include "pxr/usd/kind/registry.h"
+#include "pxr/usd/sdf/types.h"
+#include "pxr/usd/usd/attribute.h"
+#include "pxr/usd/usd/common.h"
+#include "pxr/usd/usd/editContext.h"
+#include "pxr/usd/usd/prim.h"
+#include "pxr/usd/usd/relationship.h"
+#include "pxr/usd/usdGeom/tokens.h"
+#include "pxr/usd/usdPhysics/tokens.h"
 #include "USDIncludesEnd.h"
-
-#endif // #if USE_USD_SDK
+#endif	  // #if USE_USD_SDK
 
 #define LOCTEXT_NAMESPACE "USDObjectFieldViewModel"
 
-FUsdObjectFieldViewModel::FUsdObjectFieldViewModel( FUsdObjectFieldsViewModel* InOwner )
-	: Owner( InOwner )
+FUsdObjectFieldViewModel::FUsdObjectFieldViewModel(FUsdObjectFieldsViewModel* InOwner)
+	: Owner(InOwner)
 {
 }
 
-TArray< TSharedPtr< FString > > FUsdObjectFieldViewModel::GetDropdownOptions() const
+TArray<TSharedPtr<FString>> FUsdObjectFieldViewModel::GetDropdownOptions() const
 {
 #if USE_USD_SDK
-	if ( Label == TEXT("kind") )
+	if (Label == TEXT("kind"))
 	{
-		TArray< TSharedPtr< FString > > Options;
+		TArray<TSharedPtr<FString>> Options;
 		{
 			FScopedUsdAllocs Allocs;
 
-			std::vector< pxr::TfToken > Kinds = pxr::KindRegistry::GetAllKinds();
-			Options.Reserve( Kinds.size() );
+			std::vector<pxr::TfToken> Kinds = pxr::KindRegistry::GetAllKinds();
+			Options.Reserve(Kinds.size());
 
-			for ( const pxr::TfToken& Kind : Kinds )
+			for (const pxr::TfToken& Kind : Kinds)
 			{
-				Options.Add( MakeShared< FString >( UsdToUnreal::ConvertToken( Kind ) ) );
+				Options.Add(MakeShared<FString>(UsdToUnreal::ConvertToken(Kind)));
 			}
 
 			// They are supposed to be in an unspecified order, so let's make them consistent
-			Options.Sort( [](const TSharedPtr< FString >& A, const TSharedPtr< FString >& B )
+			Options.Sort(
+				[](const TSharedPtr<FString>& A, const TSharedPtr<FString>& B)
 				{
-					return A.IsValid() && B.IsValid() && ( *A < *B );
+					return A.IsValid() && B.IsValid() && (*A < *B);
 				}
 			);
-
 		}
 		return Options;
 	}
-	else if ( Label == UsdToUnreal::ConvertToken( pxr::UsdGeomTokens->purpose ) )
+	else if (Label == UsdToUnreal::ConvertToken(pxr::UsdGeomTokens->purpose))
 	{
-		return TArray< TSharedPtr< FString > >
-		{
-			MakeShared< FString >( UsdToUnreal::ConvertToken( pxr::UsdGeomTokens->default_ ) ),
-			MakeShared< FString >( UsdToUnreal::ConvertToken( pxr::UsdGeomTokens->proxy ) ),
-			MakeShared< FString >( UsdToUnreal::ConvertToken( pxr::UsdGeomTokens->render) ),
-			MakeShared< FString >( UsdToUnreal::ConvertToken( pxr::UsdGeomTokens->guide ) ),
+		return TArray<TSharedPtr<FString>>{
+			MakeShared<FString>(UsdToUnreal::ConvertToken(pxr::UsdGeomTokens->default_)),
+			MakeShared<FString>(UsdToUnreal::ConvertToken(pxr::UsdGeomTokens->proxy)),
+			MakeShared<FString>(UsdToUnreal::ConvertToken(pxr::UsdGeomTokens->render)),
+			MakeShared<FString>(UsdToUnreal::ConvertToken(pxr::UsdGeomTokens->guide)),
 		};
 	}
-	else if ( Label == UsdToUnreal::ConvertToken( pxr::UsdGeomTokens->upAxis ) )
+	else if (Label == UsdToUnreal::ConvertToken(pxr::UsdGeomTokens->upAxis))
 	{
-		return TArray< TSharedPtr< FString > >
-		{
-			MakeShared< FString >( UsdToUnreal::ConvertToken( pxr::UsdGeomTokens->y ) ),
-			MakeShared< FString >( UsdToUnreal::ConvertToken( pxr::UsdGeomTokens->z ) ),
+		return TArray<TSharedPtr<FString>>{
+			MakeShared<FString>(UsdToUnreal::ConvertToken(pxr::UsdGeomTokens->y)),
+			MakeShared<FString>(UsdToUnreal::ConvertToken(pxr::UsdGeomTokens->z)),
 		};
 	}
-	else if ( Label == UsdToUnreal::ConvertToken( pxr::UsdGeomTokens->visibility ) )
+	else if (Label == UsdToUnreal::ConvertToken(pxr::UsdGeomTokens->visibility))
 	{
-		return TArray< TSharedPtr< FString > >
-		{
-			MakeShared< FString >( UsdToUnreal::ConvertToken( pxr::UsdGeomTokens->inherited ) ),
-			MakeShared< FString >( UsdToUnreal::ConvertToken( pxr::UsdGeomTokens->invisible ) ),
+		return TArray<TSharedPtr<FString>>{
+			MakeShared<FString>(UsdToUnreal::ConvertToken(pxr::UsdGeomTokens->inherited)),
+			MakeShared<FString>(UsdToUnreal::ConvertToken(pxr::UsdGeomTokens->invisible)),
 		};
 	}
 	else if (Label == UsdToUnreal::ConvertToken(pxr::UsdPhysicsTokens->physicsApproximation))
 	{
-		return TArray<TSharedPtr<FString>>
-		{
+		return TArray<TSharedPtr<FString>>{
 			MakeShared<FString>(UsdToUnreal::ConvertToken(pxr::UsdPhysicsTokens->none)),
 			MakeShared<FString>(UsdToUnreal::ConvertToken(pxr::UsdPhysicsTokens->convexDecomposition)),
 			MakeShared<FString>(UsdToUnreal::ConvertToken(pxr::UsdPhysicsTokens->convexHull)),
@@ -111,14 +104,14 @@ TArray< TSharedPtr< FString > > FUsdObjectFieldViewModel::GetDropdownOptions() c
 		};
 	}
 
-#endif // #if USE_USD_SDK
+#endif	  // #if USE_USD_SDK
 
 	return {};
 }
 
-void FUsdObjectFieldViewModel::SetAttributeValue( const UsdUtils::FConvertedVtValue& InValue )
+void FUsdObjectFieldViewModel::SetAttributeValue(const UsdUtils::FConvertedVtValue& InValue)
 {
-	Owner->SetFieldValue( Label, InValue );
+	Owner->SetFieldValue(Label, InValue);
 }
 
 template<typename T>
@@ -132,41 +125,36 @@ void FUsdObjectFieldsViewModel::CreateField(
 )
 {
 	UsdUtils::FConvertedVtValue VtValue;
-	VtValue.Entries = { { UsdUtils::FConvertedVtValueComponent{ TInPlaceType<T>(), Value } } };
+	VtValue.Entries = {{UsdUtils::FConvertedVtValueComponent{TInPlaceType<T>(), Value}}};
 	VtValue.SourceType = SourceType;
 
-	FUsdObjectFieldViewModel Property( this );
+	FUsdObjectFieldViewModel Property(this);
 	Property.Type = Type;
 	Property.Label = FieldName;
 	Property.Value = VtValue;
 	Property.ValueRole = ValueRole;
 	Property.bReadOnly = bReadOnly;
 
-	Fields.Add( MakeSharedUnreal< FUsdObjectFieldViewModel >( MoveTemp( Property ) ) );
+	Fields.Add(MakeSharedUnreal<FUsdObjectFieldViewModel>(MoveTemp(Property)));
 }
 
-void FUsdObjectFieldsViewModel::CreateField(
-	EObjectFieldType Type,
-	const FString& FieldName,
-	const UsdUtils::FConvertedVtValue& Value,
-	bool bReadOnly
-)
+void FUsdObjectFieldsViewModel::CreateField(EObjectFieldType Type, const FString& FieldName, const UsdUtils::FConvertedVtValue& Value, bool bReadOnly)
 {
-	FUsdObjectFieldViewModel Property( this );
+	FUsdObjectFieldViewModel Property(this);
 	Property.Type = Type;
 	Property.Label = FieldName;
 	Property.Value = Value;
 	Property.bReadOnly = bReadOnly;
 
-	Fields.Add( MakeSharedUnreal< FUsdObjectFieldViewModel >( MoveTemp( Property ) ) );
+	Fields.Add(MakeSharedUnreal<FUsdObjectFieldViewModel>(MoveTemp(Property)));
 }
 
-void FUsdObjectFieldsViewModel::SetFieldValue( const FString& FieldName, const UsdUtils::FConvertedVtValue& InValue )
+void FUsdObjectFieldsViewModel::SetFieldValue(const FString& FieldName, const UsdUtils::FConvertedVtValue& InValue)
 {
 	bool bSuccess = false;
 
 #if USE_USD_SDK
-	if ( !UsdStage )
+	if (!UsdStage)
 	{
 		return;
 	}
@@ -182,16 +170,16 @@ void FUsdObjectFieldsViewModel::SetFieldValue( const FString& FieldName, const U
 	bool bIsPropertyPath = false;
 
 	UE::FVtValue VtValue;
-	if ( UnrealToUsd::ConvertValue( InValue, VtValue ) )
+	if (UnrealToUsd::ConvertValue(InValue, VtValue))
 	{
-		if ( bIsStageAttribute )
+		if (bIsStageAttribute)
 		{
 			FScopedUsdAllocs UsdAllocs;
 
 			// To set stage metadata the edit target must be the root or session layer
-			pxr::UsdStageRefPtr Stage{ UsdStage };
-			pxr::UsdEditContext( Stage, Stage->GetRootLayer() );
-			bSuccess = UsdStage.SetMetadata( *FieldName, VtValue );
+			pxr::UsdStageRefPtr Stage{UsdStage};
+			pxr::UsdEditContext(Stage, Stage->GetRootLayer());
+			bSuccess = UsdStage.SetMetadata(*FieldName, VtValue);
 		}
 		else if (UE::FUsdPrim UsdPrim = UsdStage.GetPrimAtPath(UE::FSdfPath(*ObjectPath)))
 		{
@@ -206,17 +194,14 @@ void FUsdObjectFieldsViewModel::SetFieldValue( const FString& FieldName, const U
 			pxr::TfToken FieldNameToken = UnrealToUsd::ConvertToken(*FieldName).Get();
 
 			// Single value, single component of FString
-			if ( FieldName == TEXT( "kind" ) && InValue.Entries.Num() == 1 && InValue.Entries[ 0 ].Num() == 1 && InValue.Entries[ 0 ][ 0 ].IsType<FString>() )
+			if (FieldName == TEXT("kind") && InValue.Entries.Num() == 1 && InValue.Entries[0].Num() == 1 && InValue.Entries[0][0].IsType<FString>())
 			{
-				bSuccess = IUsdPrim::SetKind(
-					UsdPrim,
-					UnrealToUsd::ConvertToken( *( InValue.Entries[ 0 ][ 0 ].Get<FString>() ) ).Get()
-				);
+				bSuccess = IUsdPrim::SetKind(UsdPrim, UnrealToUsd::ConvertToken(*(InValue.Entries[0][0].Get<FString>())).Get());
 			}
-			else if ( UE::FUsdAttribute Attribute = UsdPrim.GetAttribute( *FieldName ) )
+			else if (UE::FUsdAttribute Attribute = UsdPrim.GetAttribute(*FieldName))
 			{
-				bSuccess = Attribute.Set( VtValue );
-				UsdUtils::NotifyIfOverriddenOpinion( Attribute );
+				bSuccess = Attribute.Set(VtValue);
+				UsdUtils::NotifyIfOverriddenOpinion(Attribute);
 			}
 			else if (pxr::UsdRelationship Relationship = PxrUsdPrim.GetRelationship(FieldNameToken))
 			{
@@ -268,23 +253,28 @@ void FUsdObjectFieldsViewModel::SetFieldValue( const FString& FieldName, const U
 		}
 	}
 
-	if ( !bSuccess )
+	if (!bSuccess)
 	{
-		const FText ErrorMessage = FText::Format(LOCTEXT("FailToSetFieldMessage", "Failed to set '{0}' on {1} '{2}'"),
+		const FText ErrorMessage = FText::Format(
+			LOCTEXT("FailToSetFieldMessage", "Failed to set '{0}' on {1} '{2}'"),
 			FText::FromString(FieldName),
-			FText::FromString(bIsStageAttribute ? TEXT("stage") : bIsPropertyPath ? TEXT("property") : TEXT("prim")),
+			FText::FromString(
+				bIsStageAttribute ? TEXT("stage")
+				: bIsPropertyPath ? TEXT("property")
+								  : TEXT("prim")
+			),
 			FText::FromString(bIsStageAttribute ? UsdStage.GetRootLayer().GetRealPath() : ObjectPath)
 		);
 
-		FNotificationInfo ErrorToast( ErrorMessage );
+		FNotificationInfo ErrorToast(ErrorMessage);
 		ErrorToast.ExpireDuration = 5.0f;
 		ErrorToast.bFireAndForget = true;
-		ErrorToast.Image = FCoreStyle::Get().GetBrush( TEXT( "MessageLog.Warning" ) );
-		FSlateNotificationManager::Get().AddNotification( ErrorToast );
+		ErrorToast.Image = FCoreStyle::Get().GetBrush(TEXT("MessageLog.Warning"));
+		FSlateNotificationManager::Get().AddNotification(ErrorToast);
 
-		FUsdLogManager::LogMessage( EMessageSeverity::Warning, ErrorMessage );
+		FUsdLogManager::LogMessage(EMessageSeverity::Warning, ErrorMessage);
 	}
-#endif // #if USE_USD_SDK
+#endif	  // #if USE_USD_SDK
 }
 
 void FUsdObjectFieldsViewModel::Refresh(const UE::FUsdStageWeak& InUsdStage, const TCHAR* InObjectPath, float TimeCode)
@@ -297,52 +287,43 @@ void FUsdObjectFieldsViewModel::Refresh(const UE::FUsdStageWeak& InUsdStage, con
 	Fields.Reset();
 
 #if USE_USD_SDK
-	TFunction<void(const FString&, const pxr::VtValue&)> DisplayFieldAsString =
-		[this](const FString& FieldName, const pxr::VtValue& Value)
+	TFunction<void(const FString&, const pxr::VtValue&)> DisplayFieldAsString = [this](const FString& FieldName, const pxr::VtValue& Value)
+	{
+		FString Stringified;
+
+		if (Value.IsArrayValued())
 		{
-			FString Stringified;
+			Stringified = FString::Printf(TEXT("%d elements: "), Value.GetArraySize());
+		}
 
-			if (Value.IsArrayValued())
-			{
-				Stringified = FString::Printf(TEXT("%d elements: "), Value.GetArraySize());
-			}
+		// This array it's too large to even stringify fast enough, so for now just show the element count
+		if (Value.IsArrayValued() && Value.GetArraySize() > 5000)
+		{
+			Stringified += TEXT("[too many entries to expand]");
+		}
+		else
+		{
+			Stringified += UsdToUnreal::ConvertString(pxr::TfStringify(Value));
 
-			// This array it's too large to even stringify fast enough, so for now just show the element count
-			if (Value.IsArrayValued() && Value.GetArraySize() > 5000)
+			// STextBlock can get very slow calculating its desired size for very long string so chop it if needed
+			const int32 MaxValueLength = 300;
+			if (Stringified.Len() > MaxValueLength)
 			{
-				Stringified += TEXT("[too many entries to expand]");
-			}
-			else
-			{
-				Stringified += UsdToUnreal::ConvertString(pxr::TfStringify(Value));
-
-				// STextBlock can get very slow calculating its desired size for very long string so chop it if needed
-				const int32 MaxValueLength = 300;
-				if (Stringified.Len() > MaxValueLength)
-				{
-					Stringified.LeftInline(MaxValueLength);
-					Stringified.Append(TEXT("...]"));
-				}
+				Stringified.LeftInline(MaxValueLength);
+				Stringified.Append(TEXT("...]"));
 			}
+		}
 
-			const bool bAttrReadOnly = true;
-			CreateField(
-				EObjectFieldType::Metadata,
-				FieldName,
-				Stringified,
-				UsdUtils::EUsdBasicDataTypes::String,
-				TEXT(""),
-				bAttrReadOnly
-			);
-		};
+		const bool bAttrReadOnly = true;
+		CreateField(EObjectFieldType::Metadata, FieldName, Stringified, UsdUtils::EUsdBasicDataTypes::String, TEXT(""), bAttrReadOnly);
+	};
 
 	// Lambda to add rows to the model based on USD object metadata. Extracted here because we can reuse it for
 	// properties and prims
 	TFunction<void(pxr::UsdObject)> AddMetadataFields = [this, &DisplayFieldAsString](pxr::UsdObject Object)
 	{
 		std::map<class pxr::TfToken, pxr::VtValue, pxr::TfDictionaryLessThan> MetadataMap = Object.GetAllMetadata();
-		for (std::map<class pxr::TfToken, pxr::VtValue, pxr::TfDictionaryLessThan>::iterator Iter = MetadataMap.begin();
-			 Iter != MetadataMap.end();
+		for (std::map<class pxr::TfToken, pxr::VtValue, pxr::TfDictionaryLessThan>::iterator Iter = MetadataMap.begin(); Iter != MetadataMap.end();
 			 ++Iter)
 		{
 			const pxr::TfToken& Key = Iter->first;
@@ -360,12 +341,7 @@ void FUsdObjectFieldsViewModel::Refresh(const UE::FUsdStageWeak& InUsdStage, con
 				else
 				{
 					const bool bReadOnly = false;
-					CreateField(
-						EObjectFieldType::Metadata,
-						MetadataName,
-						ConvertedValue,
-						bReadOnly
-					);
+					CreateField(EObjectFieldType::Metadata, MetadataName, ConvertedValue, bReadOnly);
 				}
 			}
 			// If we don't know how to display this type, try stringifying it
@@ -376,13 +352,13 @@ void FUsdObjectFieldsViewModel::Refresh(const UE::FUsdStageWeak& InUsdStage, con
 		}
 	};
 
-	if ( UsdStage )
+	if (UsdStage)
 	{
 		// Show info about the stage
-		if ( ObjectPath.Equals( TEXT( "/" ) ) || ObjectPath.IsEmpty() )
+		if (ObjectPath.Equals(TEXT("/")) || ObjectPath.IsEmpty())
 		{
 			const bool bReadOnly = true;
-			const FString Role = TEXT( "" );
+			const FString Role = TEXT("");
 			CreateField(
 				EObjectFieldType::Metadata,
 				TEXT("path"),
@@ -394,30 +370,24 @@ void FUsdObjectFieldsViewModel::Refresh(const UE::FUsdStageWeak& InUsdStage, con
 
 			FScopedUsdAllocs UsdAllocs;
 
-			std::vector<pxr::TfToken> TokenVector = pxr::SdfSchema::GetInstance().GetMetadataFields( pxr::SdfSpecType::SdfSpecTypePseudoRoot );
-			for ( const pxr::TfToken& Token : TokenVector )
+			std::vector<pxr::TfToken> TokenVector = pxr::SdfSchema::GetInstance().GetMetadataFields(pxr::SdfSpecType::SdfSpecTypePseudoRoot);
+			for (const pxr::TfToken& Token : TokenVector)
 			{
 				pxr::VtValue VtValue;
-				if ( pxr::UsdStageRefPtr( UsdStage )->GetMetadata( Token, &VtValue ) )
+				if (pxr::UsdStageRefPtr(UsdStage)->GetMetadata(Token, &VtValue))
 				{
-					FString FieldName = UsdToUnreal::ConvertToken( Token );
+					FString FieldName = UsdToUnreal::ConvertToken(Token);
 
 					UsdUtils::FConvertedVtValue ConvertedValue;
-					if ( !UsdToUnreal::ConvertValue( UE::FVtValue{ VtValue }, ConvertedValue ) )
+					if (!UsdToUnreal::ConvertValue(UE::FVtValue{VtValue}, ConvertedValue))
 					{
 						continue;
 					}
 
 					const bool bAttrReadOnly = ConvertedValue.bIsArrayValued;
-					CreateField(
-						EObjectFieldType::Metadata,
-						FieldName,
-						ConvertedValue,
-						bAttrReadOnly
-					);
+					CreateField(EObjectFieldType::Metadata, FieldName, ConvertedValue, bAttrReadOnly);
 				}
 			}
-
 		}
 		// Show info about a prim
 		else if (UE::FUsdPrim UsdPrim = UsdStage.GetPrimAtPath(UE::FSdfPath{*ObjectPath}))
@@ -433,14 +403,7 @@ void FUsdObjectFieldsViewModel::Refresh(const UE::FUsdStageWeak& InUsdStage, con
 				Role,
 				bPrimReadOnly
 			);
-			CreateField(
-				EObjectFieldType::Metadata,
-				TEXT("path"),
-				ObjectPath,
-				UsdUtils::EUsdBasicDataTypes::String,
-				Role,
-				bPrimReadOnly
-			);
+			CreateField(EObjectFieldType::Metadata, TEXT("path"), ObjectPath, UsdUtils::EUsdBasicDataTypes::String, Role, bPrimReadOnly);
 			CreateField(
 				EObjectFieldType::Metadata,
 				TEXT("kind"),
@@ -452,15 +415,15 @@ void FUsdObjectFieldsViewModel::Refresh(const UE::FUsdStageWeak& InUsdStage, con
 
 			pxr::UsdPrim PxrUsdPrim{UsdPrim};
 
-			for ( const pxr::UsdAttribute& PrimAttribute : PxrUsdPrim.GetAttributes() )
+			for (const pxr::UsdAttribute& PrimAttribute : PxrUsdPrim.GetAttributes())
 			{
-				FString AttributeName = UsdToUnreal::ConvertString( PrimAttribute.GetName().GetString() );
+				FString AttributeName = UsdToUnreal::ConvertString(PrimAttribute.GetName().GetString());
 
 				pxr::VtValue VtValue;
-				PrimAttribute.Get( &VtValue, TimeCode );
+				PrimAttribute.Get(&VtValue, TimeCode);
 
 				// Just show arrays as readonly strings for now
-				if ( VtValue.IsArrayValued() )
+				if (VtValue.IsArrayValued())
 				{
 					FString Stringified = FString::Printf(TEXT("%d elements: "), VtValue.GetArraySize());
 
@@ -495,40 +458,30 @@ void FUsdObjectFieldsViewModel::Refresh(const UE::FUsdStageWeak& InUsdStage, con
 				else
 				{
 					UsdUtils::FConvertedVtValue ConvertedValue;
-					if ( UsdToUnreal::ConvertValue( UE::FVtValue{ VtValue }, ConvertedValue ) )
+					if (UsdToUnreal::ConvertValue(UE::FVtValue{VtValue}, ConvertedValue))
 					{
 						const bool bAttrReadOnly = false;
-						CreateField(
-							EObjectFieldType::Attribute,
-							AttributeName,
-							ConvertedValue,
-							bAttrReadOnly
-						);
+						CreateField(EObjectFieldType::Attribute, AttributeName, ConvertedValue, bAttrReadOnly);
 					}
 
-					if ( PrimAttribute.HasAuthoredConnections() )
+					if (PrimAttribute.HasAuthoredConnections())
 					{
 						const FString ConnectionAttributeName = AttributeName + TEXT(":connect");
 
 						pxr::SdfPathVector ConnectedSources;
-						PrimAttribute.GetConnections( &ConnectedSources );
+						PrimAttribute.GetConnections(&ConnectedSources);
 
-						for ( pxr::SdfPath& ConnectedSource : ConnectedSources )
+						for (pxr::SdfPath& ConnectedSource : ConnectedSources)
 						{
 							UsdUtils::FConvertedVtValueEntry Entry;
-							Entry.Emplace( TInPlaceType< FString >(), UsdToUnreal::ConvertPath( ConnectedSource ) );
+							Entry.Emplace(TInPlaceType<FString>(), UsdToUnreal::ConvertPath(ConnectedSource));
 
 							UsdUtils::FConvertedVtValue ConnectionPropertyValue;
 							ConnectionPropertyValue.SourceType = UsdUtils::EUsdBasicDataTypes::String;
-							ConnectionPropertyValue.Entries = { Entry };
+							ConnectionPropertyValue.Entries = {Entry};
 
 							const bool bConnectionValueReadOnly = true;
-							CreateField(
-								EObjectFieldType::Attribute,
-								ConnectionAttributeName,
-								ConnectionPropertyValue,
-								bConnectionValueReadOnly
-							);
+							CreateField(EObjectFieldType::Attribute, ConnectionAttributeName, ConnectionPropertyValue, bConnectionValueReadOnly);
 						}
 					}
 				}
@@ -612,7 +565,7 @@ void FUsdObjectFieldsViewModel::Refresh(const UE::FUsdStageWeak& InUsdStage, con
 	}
 
 	Sort();
-#endif // #if USE_USD_SDK
+#endif	  // #if USE_USD_SDK
 }
 
 void FUsdObjectFieldsViewModel::Sort()
@@ -622,31 +575,35 @@ void FUsdObjectFieldsViewModel::Sort()
 	{
 		if (CurrentSortMode == EColumnSortMode::Ascending)
 		{
-			Fields.Sort([](const TSharedPtr<FUsdObjectFieldViewModel>& A, const TSharedPtr<FUsdObjectFieldViewModel>& B)
-			{
-				if(A->Type == B->Type)
+			Fields.Sort(
+				[](const TSharedPtr<FUsdObjectFieldViewModel>& A, const TSharedPtr<FUsdObjectFieldViewModel>& B)
 				{
-					return A->Label < B->Label;
-				}
-				else
-				{
-					return A->Type < B->Type;
+					if (A->Type == B->Type)
+					{
+						return A->Label < B->Label;
+					}
+					else
+					{
+						return A->Type < B->Type;
+					}
 				}
-			});
+			);
 		}
 		else
 		{
-			Fields.Sort([](const TSharedPtr<FUsdObjectFieldViewModel>& A, const TSharedPtr<FUsdObjectFieldViewModel>& B)
-			{
-				if(A->Type == B->Type)
+			Fields.Sort(
+				[](const TSharedPtr<FUsdObjectFieldViewModel>& A, const TSharedPtr<FUsdObjectFieldViewModel>& B)
 				{
-					return A->Label < B->Label;
-				}
-				else
-				{
-					return A->Type > B->Type;
+					if (A->Type == B->Type)
+					{
+						return A->Label < B->Label;
+					}
+					else
+					{
+						return A->Type > B->Type;
+					}
 				}
-			});
+			);
 		}
 	}
 	// When sorting by name we'll sort alphabetically regardless of type
@@ -654,17 +611,21 @@ void FUsdObjectFieldsViewModel::Sort()
 	{
 		if (CurrentSortMode == EColumnSortMode::Ascending)
 		{
-			Fields.Sort([](const TSharedPtr<FUsdObjectFieldViewModel>& A, const TSharedPtr<FUsdObjectFieldViewModel>& B)
-			{
-				return A->Label < B->Label;
-			});
+			Fields.Sort(
+				[](const TSharedPtr<FUsdObjectFieldViewModel>& A, const TSharedPtr<FUsdObjectFieldViewModel>& B)
+				{
+					return A->Label < B->Label;
+				}
+			);
 		}
 		else
 		{
-			Fields.Sort([](const TSharedPtr<FUsdObjectFieldViewModel>& A, const TSharedPtr<FUsdObjectFieldViewModel>& B)
-			{
-				return A->Label > B->Label;
-			});
+			Fields.Sort(
+				[](const TSharedPtr<FUsdObjectFieldViewModel>& A, const TSharedPtr<FUsdObjectFieldViewModel>& B)
+				{
+					return A->Label > B->Label;
+				}
+			);
 		}
 	}
 }
@@ -680,5 +641,3 @@ FString FUsdObjectFieldsViewModel::GetObjectPath() const
 }
 
 #undef LOCTEXT_NAMESPACE
-
-
diff --git a/Engine/Plugins/Importers/USDImporter/Source/USDStageEditorViewModels/Private/USDStageEditorViewModelsModule.cpp b/Engine/Plugins/Importers/USDImporter/Source/USDStageEditorViewModels/Private/USDStageEditorViewModelsModule.cpp
--- a/Engine/Plugins/Importers/USDImporter/Source/USDStageEditorViewModels/Private/USDStageEditorViewModelsModule.cpp
+++ b/Engine/Plugins/Importers/USDImporter/Source/USDStageEditorViewModels/Private/USDStageEditorViewModelsModule.cpp
@@ -4,12 +4,9 @@
 
 #include "USDMemory.h"
 
-#include "Modules/ModuleManager.h"
-
 class FUsdStageEditorViewModelsModule : public IUsdStageEditorViewModelsModule
 {
 public:
-
 };
 
-IMPLEMENT_MODULE_USD( FUsdStageEditorViewModelsModule, USDStageEditorViewModels );
+IMPLEMENT_MODULE_USD(FUsdStageEditorViewModelsModule, USDStageEditorViewModels);
diff --git a/Engine/Plugins/Importers/USDImporter/Source/USDStageEditorViewModels/Private/USDStageViewModel.cpp b/Engine/Plugins/Importers/USDImporter/Source/USDStageEditorViewModels/Private/USDStageViewModel.cpp
--- a/Engine/Plugins/Importers/USDImporter/Source/USDStageEditorViewModels/Private/USDStageViewModel.cpp
+++ b/Engine/Plugins/Importers/USDImporter/Source/USDStageEditorViewModels/Private/USDStageViewModel.cpp
@@ -22,19 +22,13 @@
 #include "Engine/World.h"
 #include "Misc/Paths.h"
 #include "ScopedTransaction.h"
-#include "UObject/GCObjectScopeGuard.h"
 
 #if USE_USD_SDK
-
 #include "USDIncludesStart.h"
-	#include "pxr/usd/kind/registry.h"
-	#include "pxr/usd/usd/common.h"
-	#include "pxr/usd/usd/modelAPI.h"
-	#include "pxr/usd/usd/stage.h"
-	#include "pxr/usd/usdGeom/xform.h"
+#include "pxr/usd/usd/common.h"
+#include "pxr/usd/usd/stage.h"
 #include "USDIncludesEnd.h"
-
-#endif // #if USE_USD_SDK
+#endif	  // #if USE_USD_SDK
 
 #define LOCTEXT_NAMESPACE "UsdStageViewModel"
 
@@ -46,30 +40,30 @@ namespace UsdViewModelImpl
 	 * We use this instead of pxr::UsdStage::SaveSessionLayers because that function
 	 * will emit a warning about the main session layer not being saved every time it is used
 	 */
-	void SaveUEStateLayer( const UE::FUsdStage& UsdStage )
+	void SaveUEStateLayer(const UE::FUsdStage& UsdStage)
 	{
 		const bool bCreateIfNeeded = false;
-		if ( UE::FSdfLayer UEStateLayer = UsdUtils::GetUEPersistentStateSublayer( UsdStage, bCreateIfNeeded ) )
+		if (UE::FSdfLayer UEStateLayer = UsdUtils::GetUEPersistentStateSublayer(UsdStage, bCreateIfNeeded))
 		{
-			UEStateLayer.Export( *UEStateLayer.GetRealPath() );
+			UEStateLayer.Export(*UEStateLayer.GetRealPath());
 		}
 	}
 }
-#endif // #if USE_USD_SDK
+#endif	  // #if USE_USD_SDK
 
 void FUsdStageViewModel::NewStage()
 {
-	FScopedTransaction Transaction( LOCTEXT( "NewStageTransaction", "Created new USD stage" ) );
+	FScopedTransaction Transaction(LOCTEXT("NewStageTransaction", "Created new USD stage"));
 
 	UsdUtils::StartMonitoringErrors();
 
-	if ( !UsdStageActor.IsValid() )
+	if (!UsdStageActor.IsValid())
 	{
-		IUsdStageModule& UsdStageModule = FModuleManager::GetModuleChecked< IUsdStageModule >( TEXT( "USDStage" ) );
-		UsdStageActor = &UsdStageModule.GetUsdStageActor( IUsdClassesModule::GetCurrentWorld() );
+		IUsdStageModule& UsdStageModule = FModuleManager::GetModuleChecked<IUsdStageModule>(TEXT("USDStage"));
+		UsdStageActor = &UsdStageModule.GetUsdStageActor(IUsdClassesModule::GetCurrentWorld());
 	}
 
-	if ( AUsdStageActor* StageActor = UsdStageActor.Get() )
+	if (AUsdStageActor* StageActor = UsdStageActor.Get())
 	{
 		StageActor->Modify();
 		StageActor->NewStage();
@@ -78,20 +72,20 @@ void FUsdStageViewModel::NewStage()
 	UsdUtils::ShowErrorsAndStopMonitoring();
 }
 
-void FUsdStageViewModel::OpenStage( const TCHAR* FilePath )
+void FUsdStageViewModel::OpenStage(const TCHAR* FilePath)
 {
 	UsdUtils::StartMonitoringErrors();
 
-	if ( !UsdStageActor.IsValid() )
+	if (!UsdStageActor.IsValid())
 	{
-		IUsdStageModule& UsdStageModule = FModuleManager::GetModuleChecked< IUsdStageModule >( TEXT("USDStage") );
-		UsdStageActor = &UsdStageModule.GetUsdStageActor( IUsdClassesModule::GetCurrentWorld() );
+		IUsdStageModule& UsdStageModule = FModuleManager::GetModuleChecked<IUsdStageModule>(TEXT("USDStage"));
+		UsdStageActor = &UsdStageModule.GetUsdStageActor(IUsdClassesModule::GetCurrentWorld());
 	}
 
-	if ( AUsdStageActor* StageActor = UsdStageActor.Get() )
+	if (AUsdStageActor* StageActor = UsdStageActor.Get())
 	{
 		StageActor->Modify();
-		StageActor->SetRootLayer( FilePath );
+		StageActor->SetRootLayer(FilePath);
 	}
 	else
 	{
@@ -103,23 +97,23 @@ void FUsdStageViewModel::OpenStage( const TCHAR* FilePath )
 
 void FUsdStageViewModel::ReloadStage()
 {
-	if ( !UsdStageActor.IsValid() )
+	if (!UsdStageActor.IsValid())
 	{
 		return;
 	}
 
 #if USE_USD_SDK
 	UE::FUsdStage Stage = UsdStageActor->GetOrOpenUsdStage();
-	pxr::UsdStageRefPtr UsdStage = pxr::UsdStageRefPtr( Stage );
+	pxr::UsdStageRefPtr UsdStage = pxr::UsdStageRefPtr(Stage);
 
 	// Can't reload from disk something that doesn't exist on disk yet
 	// (actually USD will let us do this but it seems to just clear the anonymous layers instead)
-	if ( Stage.GetRootLayer().IsAnonymous() )
+	if (Stage.GetRootLayer().IsAnonymous())
 	{
 		return;
 	}
 
-	if ( UsdStage )
+	if (UsdStage)
 	{
 		UsdUtils::StartMonitoringErrors();
 		{
@@ -156,10 +150,10 @@ void FUsdStageViewModel::ReloadStage()
 			// that would automatically pull the UEState session layer and cause it to be reloaded, so we need to try
 			// to load it back again
 			const bool bCreateIfNeeded = false;
-			UsdUtils::GetUEPersistentStateSublayer( Stage, bCreateIfNeeded );
+			UsdUtils::GetUEPersistentStateSublayer(Stage, bCreateIfNeeded);
 		}
 
-		if ( UsdUtils::ShowErrorsAndStopMonitoring() )
+		if (UsdUtils::ShowErrorsAndStopMonitoring())
 		{
 			return;
 		}
@@ -167,41 +161,41 @@ void FUsdStageViewModel::ReloadStage()
 		// If we were editing an unsaved layer, when we reload the edit target will be cleared.
 		// We need to make sure we're always editing something or else UsdEditContext might trigger some errors
 		const pxr::UsdEditTarget& EditTarget = UsdStage->GetEditTarget();
-		if ( !EditTarget.IsValid() || EditTarget.IsNull() )
+		if (!EditTarget.IsValid() || EditTarget.IsNull())
 		{
-			UsdStage->SetEditTarget( UsdStage->GetEditTargetForLocalLayer( UsdStage->GetRootLayer() ) );
+			UsdStage->SetEditTarget(UsdStage->GetEditTargetForLocalLayer(UsdStage->GetRootLayer()));
 		}
 	}
-#endif // #if USE_USD_SDK
+#endif	  // #if USE_USD_SDK
 }
 
 void FUsdStageViewModel::ResetStage()
 {
 #if USE_USD_SDK
-	if ( !UsdStageActor.IsValid() )
+	if (!UsdStageActor.IsValid())
 	{
 		return;
 	}
 
 	UE::FUsdStage Stage = UsdStageActor->GetOrOpenUsdStage();
-	pxr::UsdStageRefPtr UsdStage = pxr::UsdStageRefPtr( Stage );
+	pxr::UsdStageRefPtr UsdStage = pxr::UsdStageRefPtr(Stage);
 
-	if ( UsdStage )
+	if (UsdStage)
 	{
 		FScopedUsdAllocs Allocs;
 
 		UsdStage->GetSessionLayer()->Clear();
 
-		UsdStage->SetEditTarget( UsdStage->GetEditTargetForLocalLayer( UsdStage->GetRootLayer() ) );
+		UsdStage->SetEditTarget(UsdStage->GetEditTargetForLocalLayer(UsdStage->GetRootLayer()));
 
-		UsdStage->MuteAndUnmuteLayers( {}, UsdStage->GetMutedLayers() );
+		UsdStage->MuteAndUnmuteLayers({}, UsdStage->GetMutedLayers());
 	}
-#endif // #if USE_USD_SDK
+#endif	  // #if USE_USD_SDK
 }
 
 void FUsdStageViewModel::CloseStage()
 {
-	if ( AUsdStageActor* StageActor = UsdStageActor.Get() )
+	if (AUsdStageActor* StageActor = UsdStageActor.Get())
 	{
 		StageActor->Reset();
 	}
@@ -210,9 +204,9 @@ void FUsdStageViewModel::CloseStage()
 void FUsdStageViewModel::SaveStage()
 {
 #if USE_USD_SDK
-	if ( UsdStageActor.IsValid() )
+	if (UsdStageActor.IsValid())
 	{
-		if ( UE::FUsdStage UsdStage = UsdStageActor->GetOrOpenUsdStage() )
+		if (UE::FUsdStage UsdStage = UsdStageActor->GetOrOpenUsdStage())
 		{
 			FScopedUsdAllocs UsdAllocs;
 
@@ -221,40 +215,37 @@ void FUsdStageViewModel::SaveStage()
 			// Save layers manually instead of calling UsdStage::Save(). This is roughly the same implementation anyway, except
 			// that UsdStage::Save() will ignore a layer if it also happens to be added as a sublayer to a session layer, and
 			// we want to ensure we always save dirty layers when we hit SaveStage
-			for ( const pxr::SdfLayerHandle& Handle : pxr::UsdStageRefPtr{ UsdStage }->GetUsedLayers() )
+			for (const pxr::SdfLayerHandle& Handle : pxr::UsdStageRefPtr { UsdStage } -> GetUsedLayers())
 			{
-				if ( !Handle->IsAnonymous() && Handle->IsDirty() )
+				if (!Handle->IsAnonymous() && Handle->IsDirty())
 				{
 					Handle->Save();
 				}
 			}
 
-			UsdViewModelImpl::SaveUEStateLayer( UsdStage );
+			UsdViewModelImpl::SaveUEStateLayer(UsdStage);
 
 			UsdUtils::ShowErrorsAndStopMonitoring(LOCTEXT("USDSaveError", "Failed to save current USD Stage!\nCheck the Output Log for details."));
 		}
 	}
-#endif // #if USE_USD_SDK
+#endif	  // #if USE_USD_SDK
 }
 
-void FUsdStageViewModel::SaveStageAs( const TCHAR* FilePath )
+void FUsdStageViewModel::SaveStageAs(const TCHAR* FilePath)
 {
 #if USE_USD_SDK
-	FScopedTransaction Transaction( FText::Format(
-		LOCTEXT( "SaveAsTransaction", "Saved USD stage as '{0}'" ),
-		FText::FromString( FilePath )
-	) );
+	FScopedTransaction Transaction(FText::Format(LOCTEXT("SaveAsTransaction", "Saved USD stage as '{0}'"), FText::FromString(FilePath)));
 
 	AUsdStageActor* StageActorPtr = UsdStageActor.Get();
-	if ( StageActorPtr )
+	if (StageActorPtr)
 	{
-		if ( UE::FUsdStage UsdStage = UsdStageActor->GetOrOpenUsdStage() )
+		if (UE::FUsdStage UsdStage = UsdStageActor->GetOrOpenUsdStage())
 		{
 			UsdUtils::StartMonitoringErrors();
 
-			if ( UE::FSdfLayer RootLayer = UsdStage.GetRootLayer() )
+			if (UE::FSdfLayer RootLayer = UsdStage.GetRootLayer())
 			{
-				if ( pxr::SdfLayerRefPtr( RootLayer )->Export( TCHAR_TO_ANSI( FilePath ) ) )
+				if (pxr::SdfLayerRefPtr(RootLayer)->Export(TCHAR_TO_ANSI(FilePath)))
 				{
 					FScopedUnrealAllocs UEAllocs;
 
@@ -270,35 +261,36 @@ void FUsdStageViewModel::SaveStageAs( const TCHAR* FilePath )
 					// listening to the previous stage after we open the next one. If we just recorded writing the
 					// comment, undoing/redoing through this operation would have left the temp stage with the comment
 					// in it permanently.
-					FScopedBlockNoticeListening BlockListening( StageActorPtr );
+					FScopedBlockNoticeListening BlockListening(StageActorPtr);
 					FString OldComment = RootLayer.GetComment();
-					RootLayer.SetComment( UnrealIdentifiers::LayerSavedComment );
+					RootLayer.SetComment(UnrealIdentifiers::LayerSavedComment);
 
-					OpenStage( FilePath );
+					OpenStage(FilePath);
 
-					RootLayer.SetComment( *OldComment );
+					RootLayer.SetComment(*OldComment);
 
-					UsdViewModelImpl::SaveUEStateLayer( UsdStage );
+					UsdViewModelImpl::SaveUEStateLayer(UsdStage);
 				}
 			}
 
-			UsdUtils::ShowErrorsAndStopMonitoring( LOCTEXT( "USDSaveAsError", "Failed to SaveAs current USD Stage!\nCheck the Output Log for details." ) );
+			UsdUtils::ShowErrorsAndStopMonitoring(LOCTEXT("USDSaveAsError", "Failed to SaveAs current USD Stage!\nCheck the Output Log for details.")
+			);
 		}
 	}
-#endif // #if USE_USD_SDK
+#endif	  // #if USE_USD_SDK
 }
 
-void FUsdStageViewModel::ImportStage( const TCHAR* TargetContentFolder, UUsdStageImportOptions* Options )
+void FUsdStageViewModel::ImportStage(const TCHAR* TargetContentFolder, UUsdStageImportOptions* Options)
 {
 #if USE_USD_SDK
 	AUsdStageActor* StageActor = UsdStageActor.Get();
-	if ( !StageActor )
+	if (!StageActor)
 	{
 		return;
 	}
 
 	const UE::FUsdStage UsdStage = StageActor->GetOrOpenUsdStage();
-	if ( !UsdStage )
+	if (!UsdStage)
 	{
 		return;
 	}
@@ -311,8 +303,8 @@ void FUsdStageViewModel::ImportStage( const TCHAR* TargetContentFolder, UUsdStag
 		ImportContext.ImportOptions->PurposesToImport = StageActor->PurposesToLoad;
 		ImportContext.ImportOptions->RenderContextToImport = StageActor->RenderContext;
 		ImportContext.ImportOptions->MaterialPurpose = StageActor->MaterialPurpose;
-		ImportContext.ImportOptions->StageOptions.MetersPerUnit = UsdUtils::GetUsdStageMetersPerUnit( UsdStage );
-		ImportContext.ImportOptions->StageOptions.UpAxis = UsdUtils::GetUsdStageUpAxisAsEnum( UsdStage );
+		ImportContext.ImportOptions->StageOptions.MetersPerUnit = UsdUtils::GetUsdStageMetersPerUnit(UsdStage);
+		ImportContext.ImportOptions->StageOptions.UpAxis = UsdUtils::GetUsdStageUpAxisAsEnum(UsdStage);
 		ImportContext.ImportOptions->ImportTimeCode = StageActor->GetTime();
 		ImportContext.ImportOptions->NaniteTriangleThreshold = StageActor->NaniteTriangleThreshold;
 		ImportContext.ImportOptions->RootMotionHandling = StageActor->RootMotionHandling;
@@ -322,13 +314,13 @@ void FUsdStageViewModel::ImportStage( const TCHAR* TargetContentFolder, UUsdStag
 		ImportContext.ImportOptions->bMergeIdenticalMaterialSlots = StageActor->bMergeIdenticalMaterialSlots;
 		ImportContext.ImportOptions->bReuseIdenticalAssets = StageActor->bReuseIdenticalAssets;
 
-		ImportContext.bReadFromStageCache = true; // So that we import whatever the user has open right now, even if the file has changes
+		ImportContext.bReadFromStageCache = true;	 // So that we import whatever the user has open right now, even if the file has changes
 
 		const FString RootPath = UsdStage.GetRootLayer().GetRealPath();
-		FString StageName = FPaths::GetBaseFilename( RootPath );
+		FString StageName = FPaths::GetBaseFilename(RootPath);
 
 		// Provide a StageName when importing transient stages as this is used for the content folder name and actor label
-		if ( UsdStage.GetRootLayer().IsAnonymous() && RootPath.IsEmpty() )
+		if (UsdStage.GetRootLayer().IsAnonymous() && RootPath.IsEmpty())
 		{
 			StageName = TEXT("TransientStage");
 		}
@@ -340,12 +332,12 @@ void FUsdStageViewModel::ImportStage( const TCHAR* TargetContentFolder, UUsdStag
 
 		if (ImportContext.Init(StageName, RootPath, TEXT("/Game/"), RF_Public | RF_Transactional | RF_Standalone, bIsAutomated))
 		{
-			FScopedTransaction Transaction( FText::Format(LOCTEXT("ImportTransaction", "Import USD stage '{0}'"), FText::FromString(StageName)));
+			FScopedTransaction Transaction(FText::Format(LOCTEXT("ImportTransaction", "Import USD stage '{0}'"), FText::FromString(StageName)));
 
-			if ( bIsAutomated )
+			if (bIsAutomated)
 			{
 				// Apply same conversion that FUsdStageImportContext::Init does on our received path
-				ImportContext.PackagePath = FString::Printf( TEXT( "%s/%s/" ), TargetContentFolder, *StageName );
+				ImportContext.PackagePath = FString::Printf(TEXT("%s/%s/"), TargetContentFolder, *StageName);
 				ImportContext.ImportOptions = Options;
 			}
 
@@ -374,7 +366,7 @@ void FUsdStageViewModel::ImportStage( const TCHAR* TargetContentFolder, UUsdStag
 		}
 	}
 
-#endif // #if USE_USD_SDK
+#endif	  // #if USE_USD_SDK
 }
 
 #undef LOCTEXT_NAMESPACE
diff --git a/Engine/Plugins/Importers/USDImporter/Source/USDStageEditorViewModels/Private/USDVariantSetsViewModel.cpp b/Engine/Plugins/Importers/USDImporter/Source/USDStageEditorViewModels/Private/USDVariantSetsViewModel.cpp
--- a/Engine/Plugins/Importers/USDImporter/Source/USDStageEditorViewModels/Private/USDVariantSetsViewModel.cpp
+++ b/Engine/Plugins/Importers/USDImporter/Source/USDStageEditorViewModels/Private/USDVariantSetsViewModel.cpp
@@ -6,96 +6,93 @@
 #include "USDTypesConversion.h"
 
 #include "UsdWrappers/SdfPath.h"
+#include "UsdWrappers/UsdPrim.h"
 
 #include "ScopedTransaction.h"
 
 #if USE_USD_SDK
-
 #include "USDIncludesStart.h"
-	#include "pxr/usd/usd/prim.h"
-	#include "pxr/usd/usd/variantSets.h"
+#include "pxr/usd/usd/prim.h"
+#include "pxr/usd/usd/variantSets.h"
 #include "USDIncludesEnd.h"
+#endif	  // #if USE_USD_SDK
 
-#endif // #if USE_USD_SDK
-
-FUsdVariantSetViewModel::FUsdVariantSetViewModel( FUsdVariantSetsViewModel* InOwner )
-	: Owner( InOwner )
+FUsdVariantSetViewModel::FUsdVariantSetViewModel(FUsdVariantSetsViewModel* InOwner)
+	: Owner(InOwner)
 {
 }
 
-void FUsdVariantSetViewModel::SetVariantSelection( const TSharedPtr< FString >& InVariantSelection )
+void FUsdVariantSetViewModel::SetVariantSelection(const TSharedPtr<FString>& InVariantSelection)
 {
 #if USE_USD_SDK
-	if ( !Owner || !Owner->UsdStage )
+	if (!Owner || !Owner->UsdStage)
 	{
 		return;
 	}
 
-	FScopedTransaction Transaction( FText::Format(
-		NSLOCTEXT( "USDVariantSetsList", "SwitchVariantSetTransaction", "Switch USD Variant Set '{0}' to option '{1}'" ),
-		FText::FromString( SetName ),
-		FText::FromString( *VariantSelection )
-	) );
+	FScopedTransaction Transaction(FText::Format(
+		NSLOCTEXT("USDVariantSetsList", "SwitchVariantSetTransaction", "Switch USD Variant Set '{0}' to option '{1}'"),
+		FText::FromString(SetName),
+		FText::FromString(*VariantSelection)
+	));
 
 	VariantSelection = InVariantSelection;
 
 	FScopedUsdAllocs UsdAllocs;
 
-	std::string UsdVariantSetName = UnrealToUsd::ConvertString( *SetName ).Get();
-	std::string UsdVariantSelection = VariantSelection
-		? UnrealToUsd::ConvertString( *( *VariantSelection ) ).Get()
-		: std::string{};
+	std::string UsdVariantSetName = UnrealToUsd::ConvertString(*SetName).Get();
+	std::string UsdVariantSelection = VariantSelection ? UnrealToUsd::ConvertString(*(*VariantSelection)).Get() : std::string{};
 
-	if ( const UE::FUsdPrim& Prim = Owner->UsdStage.GetPrimAtPath( UE::FSdfPath{ *Owner->PrimPath } ) )
+	if (const UE::FUsdPrim& Prim = Owner->UsdStage.GetPrimAtPath(UE::FSdfPath{*Owner->PrimPath}))
 	{
-		pxr::UsdVariantSets UsdVariantSets = pxr::UsdPrim{ Prim }.GetVariantSets();
-		UsdVariantSets.SetSelection( UsdVariantSetName, UsdVariantSelection );
+		pxr::UsdVariantSets UsdVariantSets = pxr::UsdPrim{Prim}.GetVariantSets();
+		UsdVariantSets.SetSelection(UsdVariantSetName, UsdVariantSelection);
 	}
 
-#endif // #if USE_USD_SDK
+#endif	  // #if USE_USD_SDK
 }
 
-void FUsdVariantSetsViewModel::UpdateVariantSets( const UE::FUsdStageWeak& InUsdStage, const TCHAR* InPrimPath )
+void FUsdVariantSetsViewModel::UpdateVariantSets(const UE::FUsdStageWeak& InUsdStage, const TCHAR* InPrimPath)
 {
 #if USE_USD_SDK
 	VariantSets.Reset();
 	UsdStage = InUsdStage;
 	PrimPath = InPrimPath;
 
-	if ( !UsdStage || PrimPath.IsEmpty() || UE::FSdfPath{ InPrimPath }.IsAbsoluteRootPath() )
+	if (!UsdStage || PrimPath.IsEmpty() || UE::FSdfPath{InPrimPath}.IsAbsoluteRootPath())
 	{
 		return;
 	}
 
-	UE::FUsdPrim UsdPrim = UsdStage.GetPrimAtPath( UE::FSdfPath( *PrimPath ) );
-	if ( !UsdPrim )
+	UE::FUsdPrim UsdPrim = UsdStage.GetPrimAtPath(UE::FSdfPath(*PrimPath));
+	if (!UsdPrim)
 	{
 		return;
 	}
 
 	FScopedUsdAllocs UsdAllocs;
 
-	pxr::UsdVariantSets UsdVariantSets = pxr::UsdPrim( UsdPrim ).GetVariantSets();
+	pxr::UsdVariantSets UsdVariantSets = pxr::UsdPrim(UsdPrim).GetVariantSets();
 
-	std::vector< std::string > UsdVariantSetsNames;
-	UsdVariantSets.GetNames( &UsdVariantSetsNames );
+	std::vector<std::string> UsdVariantSetsNames;
+	UsdVariantSets.GetNames(&UsdVariantSetsNames);
 
-	for ( const std::string& UsdVariantSetName : UsdVariantSetsNames )
+	for (const std::string& UsdVariantSetName : UsdVariantSetsNames)
 	{
-		FUsdVariantSetViewModel VariantSet( this );
-		VariantSet.SetName = UsdToUnreal::ConvertString( UsdVariantSetName.c_str() );
+		FUsdVariantSetViewModel VariantSet(this);
+		VariantSet.SetName = UsdToUnreal::ConvertString(UsdVariantSetName.c_str());
 
-		pxr::UsdVariantSet UsdVariantSet = pxr::UsdPrim( UsdPrim ).GetVariantSet( UsdVariantSetName.c_str() );
-		VariantSet.VariantSelection = MakeSharedUnreal< FString >( UsdToUnreal::ConvertString( UsdVariantSet.GetVariantSelection().c_str() ) );
+		pxr::UsdVariantSet UsdVariantSet = pxr::UsdPrim(UsdPrim).GetVariantSet(UsdVariantSetName.c_str());
+		VariantSet.VariantSelection = MakeSharedUnreal<FString>(UsdToUnreal::ConvertString(UsdVariantSet.GetVariantSelection().c_str()));
 
-		std::vector< std::string > VariantNames = UsdVariantSet.GetVariantNames();
+		std::vector<std::string> VariantNames = UsdVariantSet.GetVariantNames();
 
-		for ( const std::string& VariantName : VariantNames )
+		for (const std::string& VariantName : VariantNames)
 		{
-			VariantSet.Variants.Add( MakeSharedUnreal< FString >( UsdToUnreal::ConvertString( VariantName ) ) );
+			VariantSet.Variants.Add(MakeSharedUnreal<FString>(UsdToUnreal::ConvertString(VariantName)));
 		}
 
-		VariantSets.Add( MakeSharedUnreal< FUsdVariantSetViewModel >( MoveTemp( VariantSet ) ) );
+		VariantSets.Add(MakeSharedUnreal<FUsdVariantSetViewModel>(MoveTemp(VariantSet)));
 	}
-#endif // #if USE_USD_SDK
+#endif	  // #if USE_USD_SDK
 }
diff --git a/Engine/Plugins/Importers/USDImporter/Source/USDStageEditorViewModels/Public/USDPrimViewModel.h b/Engine/Plugins/Importers/USDImporter/Source/USDStageEditorViewModels/Public/USDPrimViewModel.h
--- a/Engine/Plugins/Importers/USDImporter/Source/USDStageEditorViewModels/Public/USDPrimViewModel.h
+++ b/Engine/Plugins/Importers/USDImporter/Source/USDStageEditorViewModels/Public/USDPrimViewModel.h
@@ -2,29 +2,45 @@
 
 #pragma once
 
-#include "UsdWrappers/UsdStage.h"
-#include "UsdWrappers/UsdPrim.h"
 #include "Widgets/IUSDTreeViewItem.h"
 
 #include "UsdWrappers/ForwardDeclarations.h"
-#include "UsdWrappers/UsdStage.h"
 #include "UsdWrappers/UsdPrim.h"
+#include "UsdWrappers/UsdStage.h"
 
-#include "Templates/SharedPointer.h"
 #include "Internationalization/Text.h"
+#include "Templates/SharedPointer.h"
 
-using FUsdPrimViewModelRef = TSharedRef< class FUsdPrimViewModel >;
-using FUsdPrimViewModelPtr = TSharedPtr< class FUsdPrimViewModel >;
+using FUsdPrimViewModelRef = TSharedRef<class FUsdPrimViewModel>;
+using FUsdPrimViewModelPtr = TSharedPtr<class FUsdPrimViewModel>;
 
-class USDSTAGEEDITORVIEWMODELS_API FUsdPrimModel : public TSharedFromThis< FUsdPrimModel >
+class USDSTAGEEDITORVIEWMODELS_API FUsdPrimModel : public TSharedFromThis<FUsdPrimModel>
 {
 public:
-	FText GetName() const { return Name; }
-	FText GetType() const { return Type; }
-	bool HasPayload() const { return bHasPayload; }
-	bool IsLoaded() const { return bIsLoaded; }
-	bool HasCompositionArcs() const { return bHasCompositionArcs; }
-	bool IsVisible() const { return bIsVisible; }
+	FText GetName() const
+	{
+		return Name;
+	}
+	FText GetType() const
+	{
+		return Type;
+	}
+	bool HasPayload() const
+	{
+		return bHasPayload;
+	}
+	bool IsLoaded() const
+	{
+		return bIsLoaded;
+	}
+	bool HasCompositionArcs() const
+	{
+		return bHasCompositionArcs;
+	}
+	bool IsVisible() const
+	{
+		return bIsVisible;
+	}
 
 	FText Name;
 	FText Type;
@@ -37,30 +53,30 @@ public:
 class USDSTAGEEDITORVIEWMODELS_API FUsdPrimViewModel : public IUsdTreeViewItem
 {
 public:
-	FUsdPrimViewModel( FUsdPrimViewModel* InParentItem, const UE::FUsdStageWeak& InUsdStage, const UE::FUsdPrim& InPrim = {} );
+	FUsdPrimViewModel(FUsdPrimViewModel* InParentItem, const UE::FUsdStageWeak& InUsdStage, const UE::FUsdPrim& InPrim = {});
 
-	TArray< FUsdPrimViewModelRef >& UpdateChildren();
+	TArray<FUsdPrimViewModelRef>& UpdateChildren();
 
 	void FillChildren();
 
-	void RefreshData( bool bRefreshChildren );
+	void RefreshData(bool bRefreshChildren);
 
 	bool HasVisibilityAttribute() const;
 	void ToggleVisibility();
 	void TogglePayload();
 
-	void ApplySchema( FName SchemaName );
-	bool CanApplySchema( FName SchemaName ) const;
-	void RemoveSchema( FName SchemaName );
-	bool CanRemoveSchema( FName SchemaName ) const;
+	void ApplySchema(FName SchemaName);
+	bool CanApplySchema(FName SchemaName) const;
+	void RemoveSchema(FName SchemaName);
+	bool CanRemoveSchema(FName SchemaName) const;
 
 	// Returns true if this prim has at least one spec on the stage's local layer stack
 	bool HasSpecsOnLocalLayer() const;
 
-	void DefinePrim( const TCHAR* PrimName );
+	void DefinePrim(const TCHAR* PrimName);
 
 	/** Delegate for hooking up an inline editable text block to be notified that a rename is requested. */
-	DECLARE_DELEGATE( FOnRenameRequest );
+	DECLARE_DELEGATE(FOnRenameRequest);
 
 	/** Broadcasts whenever a rename is requested */
 	FOnRenameRequest RenameRequestEvent;
@@ -74,9 +90,9 @@ public:
 	UE::FUsdPrim UsdPrim;
 
 	FUsdPrimViewModel* ParentItem;
-	TArray< FUsdPrimViewModelRef > Children;
+	TArray<FUsdPrimViewModelRef> Children;
 
-	TSharedRef< FUsdPrimModel > RowData; // Data model
+	TSharedRef<FUsdPrimModel> RowData;	  // Data model
 
 	bool bIsRenamingExistingPrim = false;
 };
diff --git a/Engine/Plugins/Importers/USDImporter/Source/USDStageEditorViewModels/Public/USDStageViewModel.h b/Engine/Plugins/Importers/USDImporter/Source/USDStageEditorViewModels/Public/USDStageViewModel.h
--- a/Engine/Plugins/Importers/USDImporter/Source/USDStageEditorViewModels/Public/USDStageViewModel.h
+++ b/Engine/Plugins/Importers/USDImporter/Source/USDStageEditorViewModels/Public/USDStageViewModel.h
@@ -4,7 +4,6 @@
 
 #include "CoreMinimal.h"
 #include "UObject/WeakObjectPtr.h"
-#include "UObject/WeakObjectPtrTemplates.h"
 
 class AUsdStageActor;
 class UPackage;
@@ -14,15 +13,15 @@ class USDSTAGEEDITORVIEWMODELS_API FUsdStageViewModel
 {
 public:
 	void NewStage();
-	void OpenStage( const TCHAR* FilePath );
+	void OpenStage(const TCHAR* FilePath);
 	void ReloadStage();
 	void ResetStage();
 	void CloseStage();
 	void SaveStage();
 	/** Temporary until SaveAs feature is properly implemented, may be removed in a future release */
-	void SaveStageAs( const TCHAR* FilePath );
-	void ImportStage( const TCHAR* TargetContentFolder = nullptr, UUsdStageImportOptions* Options = nullptr );
+	void SaveStageAs(const TCHAR* FilePath);
+	void ImportStage(const TCHAR* TargetContentFolder = nullptr, UUsdStageImportOptions* Options = nullptr);
 
 public:
-	TWeakObjectPtr< AUsdStageActor > UsdStageActor;
+	TWeakObjectPtr<AUsdStageActor> UsdStageActor;
 };
diff --git a/Engine/Plugins/Importers/USDImporter/Source/USDStageEditorViewModels/Public/USDVariantSetsViewModel.h b/Engine/Plugins/Importers/USDImporter/Source/USDStageEditorViewModels/Public/USDVariantSetsViewModel.h
--- a/Engine/Plugins/Importers/USDImporter/Source/USDStageEditorViewModels/Public/USDVariantSetsViewModel.h
+++ b/Engine/Plugins/Importers/USDImporter/Source/USDStageEditorViewModels/Public/USDVariantSetsViewModel.h
@@ -4,22 +4,21 @@
 
 #include "Templates/SharedPointer.h"
 
-#include "UsdWrappers/UsdPrim.h"
 #include "UsdWrappers/UsdStage.h"
 
 class FUsdVariantSetsViewModel;
 
-class USDSTAGEEDITORVIEWMODELS_API FUsdVariantSetViewModel : public TSharedFromThis< FUsdVariantSetViewModel >
+class USDSTAGEEDITORVIEWMODELS_API FUsdVariantSetViewModel : public TSharedFromThis<FUsdVariantSetViewModel>
 {
 public:
-	explicit FUsdVariantSetViewModel( FUsdVariantSetsViewModel* InOwner );
+	explicit FUsdVariantSetViewModel(FUsdVariantSetsViewModel* InOwner);
 
-	void SetVariantSelection( const TSharedPtr< FString >& InVariantSelection );
+	void SetVariantSelection(const TSharedPtr<FString>& InVariantSelection);
 
 public:
 	FString SetName;
-	TSharedPtr< FString > VariantSelection;
-	TArray< TSharedPtr< FString > > Variants;
+	TSharedPtr<FString> VariantSelection;
+	TArray<TSharedPtr<FString>> Variants;
 
 private:
 	FUsdVariantSetsViewModel* Owner;
@@ -28,11 +27,11 @@ private:
 class USDSTAGEEDITORVIEWMODELS_API FUsdVariantSetsViewModel
 {
 public:
-	void UpdateVariantSets( const UE::FUsdStageWeak& InUsdStage, const TCHAR* PrimPath );
+	void UpdateVariantSets(const UE::FUsdStageWeak& InUsdStage, const TCHAR* PrimPath);
 
 public:
 	UE::FUsdStageWeak UsdStage;
 	FString PrimPath;
 
-	TArray< TSharedPtr< FUsdVariantSetViewModel > > VariantSets;
+	TArray<TSharedPtr<FUsdVariantSetViewModel>> VariantSets;
 };
