USD: Reformat USDTests module.

#rnx

diff --git a/Engine/Plugins/Importers/USDImporter/Source/USDTests/Private/USDTestsBlueprintLibrary.cpp b/Engine/Plugins/Importers/USDImporter/Source/USDTests/Private/USDTestsBlueprintLibrary.cpp
--- a/Engine/Plugins/Importers/USDImporter/Source/USDTests/Private/USDTestsBlueprintLibrary.cpp
+++ b/Engine/Plugins/Importers/USDImporter/Source/USDTests/Private/USDTestsBlueprintLibrary.cpp
@@ -11,99 +11,98 @@
 #include "Editor/Transactor.h"
 #include "Kismet2/BlueprintEditorUtils.h"
 #include "Kismet2/KismetEditorUtilities.h"
-#include "UObject/ObjectMacros.h"
 
-bool USDTestsBlueprintLibrary::RecompileBlueprintStageActor( AUsdStageActor* BlueprintDerivedStageActor )
+bool USDTestsBlueprintLibrary::RecompileBlueprintStageActor(AUsdStageActor* BlueprintDerivedStageActor)
 {
 #if WITH_EDITOR
-	if ( !BlueprintDerivedStageActor )
+	if (!BlueprintDerivedStageActor)
 	{
 		return false;
 	}
 
-	UBlueprintGeneratedClass* BPClass = Cast<UBlueprintGeneratedClass>( BlueprintDerivedStageActor->GetClass() );
-	if ( !BPClass )
+	UBlueprintGeneratedClass* BPClass = Cast<UBlueprintGeneratedClass>(BlueprintDerivedStageActor->GetClass());
+	if (!BPClass)
 	{
 		return false;
 	}
 
-	UBlueprint* BP = Cast<UBlueprint>( BPClass->ClassGeneratedBy );
-	if ( !BP )
+	UBlueprint* BP = Cast<UBlueprint>(BPClass->ClassGeneratedBy);
+	if (!BP)
 	{
 		return false;
 	}
 
 	// Compile blueprint. Copied from BlueprintEditorTests.cpp
-	FBlueprintEditorUtils::RefreshAllNodes( BP );
-	FKismetEditorUtilities::CompileBlueprint( BP, EBlueprintCompileOptions::SkipGarbageCollection );
-	if ( BP->Status == EBlueprintStatus::BS_UpToDate )
+	FBlueprintEditorUtils::RefreshAllNodes(BP);
+	FKismetEditorUtilities::CompileBlueprint(BP, EBlueprintCompileOptions::SkipGarbageCollection);
+	if (BP->Status == EBlueprintStatus::BS_UpToDate)
 	{
-		UE_LOG( LogUsd, Log, TEXT( "Blueprint compiled successfully (%s)" ), *BP->GetName() );
+		UE_LOG(LogUsd, Log, TEXT("Blueprint compiled successfully (%s)"), *BP->GetName());
 		return true;
 	}
-	else if ( BP->Status == EBlueprintStatus::BS_UpToDateWithWarnings )
+	else if (BP->Status == EBlueprintStatus::BS_UpToDateWithWarnings)
 	{
-		UE_LOG( LogUsd, Warning, TEXT( "Blueprint compiled successfully with warnings(%s)" ), *BP->GetName() );
+		UE_LOG(LogUsd, Warning, TEXT("Blueprint compiled successfully with warnings(%s)"), *BP->GetName());
 		return true;
 	}
-	else if ( BP->Status == EBlueprintStatus::BS_Error )
+	else if (BP->Status == EBlueprintStatus::BS_Error)
 	{
-		UE_LOG( LogUsd, Error, TEXT( "Blueprint failed to compile (%s)" ), *BP->GetName() );
+		UE_LOG(LogUsd, Error, TEXT("Blueprint failed to compile (%s)"), *BP->GetName());
 		return false;
 	}
 	else
 	{
-		UE_LOG( LogUsd, Error, TEXT( "Blueprint is in an unexpected state after compiling (%s)" ), *BP->GetName() );
+		UE_LOG(LogUsd, Error, TEXT("Blueprint is in an unexpected state after compiling (%s)"), *BP->GetName());
 		return false;
 	}
-#endif // WITH_EDITOR
+#endif	  // WITH_EDITOR
 
 	return false;
 }
 
-void USDTestsBlueprintLibrary::DirtyStageActorBlueprint( AUsdStageActor* BlueprintDerivedStageActor )
+void USDTestsBlueprintLibrary::DirtyStageActorBlueprint(AUsdStageActor* BlueprintDerivedStageActor)
 {
 #if WITH_EDITOR
-	if ( !BlueprintDerivedStageActor )
+	if (!BlueprintDerivedStageActor)
 	{
 		return;
 	}
 
-	UBlueprintGeneratedClass* BPClass = Cast<UBlueprintGeneratedClass>( BlueprintDerivedStageActor->GetClass() );
-	if ( !BPClass )
+	UBlueprintGeneratedClass* BPClass = Cast<UBlueprintGeneratedClass>(BlueprintDerivedStageActor->GetClass());
+	if (!BPClass)
 	{
 		return;
 	}
 
-	UBlueprint* BP = Cast<UBlueprint>( BPClass->ClassGeneratedBy );
-	if ( !BP )
+	UBlueprint* BP = Cast<UBlueprint>(BPClass->ClassGeneratedBy);
+	if (!BP)
 	{
 		return;
 	}
 
 	// We need to add something that potentially modifies the blueprint code. Just dirtying the blueprint or changing it status doesn't trigger
 	// a recompile when going into PIE
-	FName VarName = FBlueprintEditorUtils::FindUniqueKismetName( BP, TEXT( "NewVar" ) );
-	FEdGraphPinType StringPinType( UEdGraphSchema_K2::PC_String, NAME_None, nullptr, EPinContainerType::None, false, FEdGraphTerminalType() );
-	bool bSuccess = FBlueprintEditorUtils::AddMemberVariable( BP, VarName, StringPinType );
-	if ( !bSuccess )
+	FName VarName = FBlueprintEditorUtils::FindUniqueKismetName(BP, TEXT("NewVar"));
+	FEdGraphPinType StringPinType(UEdGraphSchema_K2::PC_String, NAME_None, nullptr, EPinContainerType::None, false, FEdGraphTerminalType());
+	bool bSuccess = FBlueprintEditorUtils::AddMemberVariable(BP, VarName, StringPinType);
+	if (!bSuccess)
 	{
-		UE_LOG( LogUsd, Error, TEXT( "Failed to add new variable to blueprint (%s)" ), *BP->GetName() );
+		UE_LOG(LogUsd, Error, TEXT("Failed to add new variable to blueprint (%s)"), *BP->GetName());
 	}
-#endif // WITH_EDITOR
+#endif	  // WITH_EDITOR
 }
 
-int64 USDTestsBlueprintLibrary::GetSubtreeVertexCount( AUsdStageActor* StageActor, const FString& PrimPath )
+int64 USDTestsBlueprintLibrary::GetSubtreeVertexCount(AUsdStageActor* StageActor, const FString& PrimPath)
 {
-	if ( StageActor )
+	if (StageActor)
 	{
-		if ( TSharedPtr<FUsdInfoCache> Cache = StageActor->GetInfoCache() )
+		if (TSharedPtr<FUsdInfoCache> Cache = StageActor->GetInfoCache())
 		{
-			TOptional<uint64> Result = Cache->GetSubtreeVertexCount( UE::FSdfPath{ *PrimPath } );
-			if ( Result.IsSet() )
+			TOptional<uint64> Result = Cache->GetSubtreeVertexCount(UE::FSdfPath{*PrimPath});
+			if (Result.IsSet())
 			{
 				// Narrowing conversion here, but we're only using this for our test scenes, which have at most a few thousand vertices
-				return static_cast< int64 >( Result.GetValue() );
+				return static_cast<int64>(Result.GetValue());
 			}
 			else
 			{
@@ -115,16 +114,16 @@ int64 USDTestsBlueprintLibrary::GetSubtreeVertexCount( AUsdStageActor* StageActo
 	return -1;
 }
 
-int64 USDTestsBlueprintLibrary::GetSubtreeMaterialSlotCount( AUsdStageActor* StageActor, const FString& PrimPath )
+int64 USDTestsBlueprintLibrary::GetSubtreeMaterialSlotCount(AUsdStageActor* StageActor, const FString& PrimPath)
 {
-	if ( StageActor )
+	if (StageActor)
 	{
-		if ( TSharedPtr<FUsdInfoCache> Cache = StageActor->GetInfoCache() )
+		if (TSharedPtr<FUsdInfoCache> Cache = StageActor->GetInfoCache())
 		{
-			TOptional<uint64> Result = Cache->GetSubtreeMaterialSlotCount( UE::FSdfPath{ *PrimPath } );
-			if ( Result.IsSet() )
+			TOptional<uint64> Result = Cache->GetSubtreeMaterialSlotCount(UE::FSdfPath{*PrimPath});
+			if (Result.IsSet())
 			{
-				return static_cast< int64 >( Result.GetValue() );
+				return static_cast<int64>(Result.GetValue());
 			}
 			else
 			{
diff --git a/Engine/Plugins/Importers/USDImporter/Source/USDTests/Private/USDTestsModule.cpp b/Engine/Plugins/Importers/USDImporter/Source/USDTests/Private/USDTestsModule.cpp
--- a/Engine/Plugins/Importers/USDImporter/Source/USDTests/Private/USDTestsModule.cpp
+++ b/Engine/Plugins/Importers/USDImporter/Source/USDTests/Private/USDTestsModule.cpp
@@ -4,8 +4,6 @@
 
 #include "USDMemory.h"
 
-#include "Modules/ModuleManager.h"
-
 class FUsdTestsModule : public IUsdTestsModule
 {
 public:
@@ -19,4 +17,4 @@ public:
 	}
 };
 
-IMPLEMENT_MODULE_USD( FUsdTestsModule, USDTests );
+IMPLEMENT_MODULE_USD(FUsdTestsModule, USDTests);
diff --git a/Engine/Plugins/Importers/USDImporter/Source/USDTests/Public/USDTestsBlueprintLibrary.h b/Engine/Plugins/Importers/USDImporter/Source/USDTests/Public/USDTestsBlueprintLibrary.h
--- a/Engine/Plugins/Importers/USDImporter/Source/USDTests/Public/USDTestsBlueprintLibrary.h
+++ b/Engine/Plugins/Importers/USDImporter/Source/USDTests/Public/USDTestsBlueprintLibrary.h
@@ -3,7 +3,6 @@
 #pragma once
 
 #include "CoreMinimal.h"
-#include "UObject/ObjectMacros.h"
 #include "Kismet/BlueprintFunctionLibrary.h"
 
 #include "USDTestsBlueprintLibrary.generated.h"
@@ -11,7 +10,7 @@
 /**
  * Library of functions that can be used via Python scripting to help testing the other USD functionality
  */
-UCLASS(meta=(ScriptName="USDTestingLibrary"))
+UCLASS(meta = (ScriptName = "USDTestingLibrary"))
 class USDTESTS_API USDTestsBlueprintLibrary : public UBlueprintFunctionLibrary
 {
 	GENERATED_BODY()
@@ -22,30 +21,30 @@ public:
 	 * Helps test the effects of blueprint recompilation on the spawned actors and assets when a stage is opened.
 	 * Returns whether it compiled successfully or not.
 	 */
-	UFUNCTION( BlueprintCallable, Category = "Blueprint" )
-	static bool RecompileBlueprintStageActor( AUsdStageActor* BlueprintDerivedStageActor );
+	UFUNCTION(BlueprintCallable, Category = "Blueprint")
+	static bool RecompileBlueprintStageActor(AUsdStageActor* BlueprintDerivedStageActor);
 
 	/**
 	 * Intentionally dirties the UBlueprint for the given stage actor's generated class.
 	 * This is useful for testing how the stage actor behaves when going into PIE with a dirty blueprint, as that usually triggers
 	 * a recompile at the very sensitive PIE transition
 	 */
-	UFUNCTION( BlueprintCallable, Category = "Blueprint" )
-	static void DirtyStageActorBlueprint( AUsdStageActor* BlueprintDerivedStageActor );
+	UFUNCTION(BlueprintCallable, Category = "Blueprint")
+	static void DirtyStageActorBlueprint(AUsdStageActor* BlueprintDerivedStageActor);
 
 	/**
 	 * Queries a subtree vertex count using the stage actor's info cache, which is not yet exposed to blueprint.
 	 * May return -1 in case of an error.
 	 */
-	UFUNCTION( BlueprintCallable, Category = "Counts" )
-	static int64 GetSubtreeVertexCount( AUsdStageActor* StageActor, const FString& PrimPath );
+	UFUNCTION(BlueprintCallable, Category = "Counts")
+	static int64 GetSubtreeVertexCount(AUsdStageActor* StageActor, const FString& PrimPath);
 
 	/**
 	 * Queries a subtree material slot count using the stage actor's info cache, which is not yet exposed to blueprint.
 	 * May return -1 in case of an error.
 	 */
-	UFUNCTION( BlueprintCallable, Category = "Counts" )
-	static int64 GetSubtreeMaterialSlotCount( AUsdStageActor* StageActor, const FString& PrimPath );
+	UFUNCTION(BlueprintCallable, Category = "Counts")
+	static int64 GetSubtreeMaterialSlotCount(AUsdStageActor* StageActor, const FString& PrimPath);
 
 	/**
 	 * Internally opens the provided stage and sets it on the stage actor via C++
