USD: Reformat USDStageImporter module.

#rnx

diff --git a/Engine/Plugins/Experimental/ChaosClothAssetEditor/Source/ChaosClothAssetDataflowNodes/Private/ChaosClothAsset/USDImportNode.cpp b/Engine/Plugins/Experimental/ChaosClothAssetEditor/Source/ChaosClothAssetDataflowNodes/Private/ChaosClothAsset/USDImportNode.cpp
--- a/Engine/Plugins/Experimental/ChaosClothAssetEditor/Source/ChaosClothAssetDataflowNodes/Private/ChaosClothAsset/USDImportNode.cpp
+++ b/Engine/Plugins/Experimental/ChaosClothAssetEditor/Source/ChaosClothAssetDataflowNodes/Private/ChaosClothAsset/USDImportNode.cpp
@@ -18,12 +18,15 @@
 #include "Rendering/SkeletalMeshLODModel.h"
 #include "StaticMeshAttributes.h"
 #include "UObject/Package.h"
-#include "UsdWrappers/UsdPrim.h"
-#include "USDStageImporter.h"
+#include "UnrealUSDWrapper.h"
+#include "USDConversionUtils.h"
 #include "USDStageImportContext.h"
+#include "USDStageImporter.h"
 #include "USDStageImportOptions.h"
 #include "USDTypesConversion.h"
 #include "USDValueConversion.h"
+#include "UsdWrappers/UsdPrim.h"
+#include "UsdWrappers/VtValue.h"
 
 #include UE_INLINE_GENERATED_CPP_BY_NAME(USDImportNode)
 
diff --git a/Engine/Plugins/Importers/USDImporter/Source/USDStage/USDStage.Build.cs b/Engine/Plugins/Importers/USDImporter/Source/USDStage/USDStage.Build.cs
--- a/Engine/Plugins/Importers/USDImporter/Source/USDStage/USDStage.Build.cs
+++ b/Engine/Plugins/Importers/USDImporter/Source/USDStage/USDStage.Build.cs
@@ -29,7 +29,6 @@ namespace UnrealBuildTool.Rules
 					"GeometryCacheTracks",
 					"GeometryCacheUSD",
 					"HairStrandsCore",
-					"InterchangePipelines",
 					"LevelSequence",
 					"LiveLinkComponents", // For tracking edits to LiveLinkComponentController properties and writing to USD
 					"LiveLinkInterface",
diff --git a/Engine/Plugins/Importers/USDImporter/Source/USDStageEditorViewModels/Private/USDStageViewModel.cpp b/Engine/Plugins/Importers/USDImporter/Source/USDStageEditorViewModels/Private/USDStageViewModel.cpp
--- a/Engine/Plugins/Importers/USDImporter/Source/USDStageEditorViewModels/Private/USDStageViewModel.cpp
+++ b/Engine/Plugins/Importers/USDImporter/Source/USDStageEditorViewModels/Private/USDStageViewModel.cpp
@@ -11,6 +11,7 @@
 #include "USDLog.h"
 #include "USDStageActor.h"
 #include "USDStageImportContext.h"
+#include "USDStageImporter.h"
 #include "USDStageImporterModule.h"
 #include "USDStageImportOptions.h"
 #include "USDStageModule.h"
diff --git a/Engine/Plugins/Importers/USDImporter/Source/USDStageImporter/Private/USDOptionsWindow.cpp b/Engine/Plugins/Importers/USDImporter/Source/USDStageImporter/Private/USDOptionsWindow.cpp
--- a/Engine/Plugins/Importers/USDImporter/Source/USDStageImporter/Private/USDOptionsWindow.cpp
+++ b/Engine/Plugins/Importers/USDImporter/Source/USDStageImporter/Private/USDOptionsWindow.cpp
@@ -3,16 +3,11 @@
 #include "USDOptionsWindow.h"
 
 #include "USDStageImportOptions.h"
-#include "UsdWrappers/SdfLayer.h"
-#include "UsdWrappers/SdfPath.h"
-#include "UsdWrappers/UsdPrim.h"
 #include "UsdWrappers/UsdStage.h"
 #include "Widgets/SUSDStagePreviewTree.h"
 
 #include "CoreGlobals.h"
-#include "CoreMinimal.h"
 #include "Framework/Application/SlateApplication.h"
-#include "HAL/PlatformApplicationMisc.h"
 #include "Interfaces/IMainFrameModule.h"
 #include "Misc/ConfigCacheIni.h"
 #include "Misc/ScopedSlowTask.h"
@@ -20,7 +15,6 @@
 #include "PropertyEditorModule.h"
 #include "Widgets/Input/SButton.h"
 #include "Widgets/Input/SSearchBox.h"
-#include "Widgets/Layout/SBox.h"
 #include "Widgets/Layout/SSplitter.h"
 #include "Widgets/Layout/SUniformGridPanel.h"
 #include "Widgets/SBoxPanel.h"
@@ -28,104 +22,99 @@
 
 #define LOCTEXT_NAMESPACE "USDOptionsWindow"
 
-bool SUsdOptionsWindow::ShowOptions(
-	UObject& OptionsObject,
-	const FText& WindowTitle,
-	const FText& AcceptText,
-	const UE::FUsdStage* Stage
-)
+bool SUsdOptionsWindow::ShowOptions(UObject& OptionsObject, const FText& WindowTitle, const FText& AcceptText, const UE::FUsdStage* Stage)
 {
 	TSharedPtr<SWindow> ParentWindow;
 
-	if ( FModuleManager::Get().IsModuleLoaded( "MainFrame" ) )
+	if (FModuleManager::Get().IsModuleLoaded("MainFrame"))
 	{
-		IMainFrameModule& MainFrame = FModuleManager::LoadModuleChecked<IMainFrameModule>( "MainFrame" );
+		IMainFrameModule& MainFrame = FModuleManager::LoadModuleChecked<IMainFrameModule>("MainFrame");
 		ParentWindow = MainFrame.GetParentWindow();
 	}
 
 	FDisplayMetrics DisplayMetrics;
-	FSlateApplication::Get().GetDisplayMetrics( DisplayMetrics );
-	const float DPIScaleFactor = FPlatformApplicationMisc::GetDPIScaleFactorAtPoint(
-		DisplayMetrics.PrimaryDisplayWorkAreaRect.Left,
-		DisplayMetrics.PrimaryDisplayWorkAreaRect.Top
-	);
+	FSlateApplication::Get().GetDisplayMetrics(DisplayMetrics);
 
-	const static FString ImporterName = TEXT( "USDImporter" );
+	const static FString ImporterName = TEXT("USDImporter");
 	const FString ClassName = OptionsObject.GetClass()->GetName();
 
 	const FString WidthConfigName = ClassName + TEXT("_DialogWidth");
 	int32 Width = 100;
-	if ( GConfig->GetInt( *ImporterName, *WidthConfigName, Width, GEditorPerProjectIni ) )
+	if (GConfig->GetInt(*ImporterName, *WidthConfigName, Width, GEditorPerProjectIni))
 	{
 		// Just to make sure we don't load something unusable
-		Width = FMath::Min( FMath::Max( Width, 100 ), DisplayMetrics.PrimaryDisplayWidth );
+		Width = FMath::Min(FMath::Max(Width, 100), DisplayMetrics.PrimaryDisplayWidth);
 	}
 	else
 	{
 		// Make it wider if we're going to show two panes
-		Width = 500 * ( Stage == nullptr ? 1 : 2 );
+		Width = 500 * (Stage == nullptr ? 1 : 2);
 	}
 
-	const FString HeightConfigName = ClassName + TEXT( "_DialogHeight" );
+	const FString HeightConfigName = ClassName + TEXT("_DialogHeight");
 	int32 Height = 100;
-	if ( GConfig->GetInt( *ImporterName, *HeightConfigName, Height, GEditorPerProjectIni ) )
+	if (GConfig->GetInt(*ImporterName, *HeightConfigName, Height, GEditorPerProjectIni))
 	{
-		Height = FMath::Min( FMath::Max( Height, 100 ), DisplayMetrics.PrimaryDisplayHeight );
+		Height = FMath::Min(FMath::Max(Height, 100), DisplayMetrics.PrimaryDisplayHeight);
 	}
 	else
 	{
 		Height = 600;
 	}
 
-	const FVector2D ClientSize{ static_cast< float >( Width ), static_cast< float >( Height ) };
-	TSharedRef<SWindow> Window = SNew( SWindow )
-		.Title( WindowTitle )
-		.SizingRule( ESizingRule::UserSized )
-		.AdjustInitialSizeAndPositionForDPIScale( false )
-		.ClientSize( ClientSize );
+	// clang-format off
+	const FVector2D ClientSize{static_cast<float>(Width), static_cast<float>(Height)};
+	TSharedRef<SWindow> Window = SNew(SWindow)
+		.Title(WindowTitle)
+		.SizingRule(ESizingRule::UserSized)
+		.AdjustInitialSizeAndPositionForDPIScale(false)
+		.ClientSize(ClientSize);
 
 	TSharedPtr<SUsdOptionsWindow> OptionsWindow;
 	Window->SetContent
 	(
-		SAssignNew( OptionsWindow, SUsdOptionsWindow )
-		.OptionsObject( &OptionsObject )
-		.AcceptText( AcceptText )
-		.WidgetWindow( Window )
-		.Stage( Stage )
+		SAssignNew(OptionsWindow, SUsdOptionsWindow)
+		.OptionsObject(&OptionsObject)
+		.AcceptText(AcceptText)
+		.WidgetWindow(Window)
+		.Stage(Stage)
 	);
+	// clang-format on
 
-	Window->GetOnWindowClosedEvent().AddLambda( [WidthConfigName, HeightConfigName, OptionsWindow]( const TSharedRef<SWindow>& Window )
-	{
-		// We use the geometry here because any sort of size we can extract from the window includes some added margins
-		// that are not easy to account for, even with GetWindowSizeFromClientSize. The original provided ClientSize
-		// seems to always match the content geometry though, so we're saving and loading the same thing
-		FGeometry Geometry = Window->GetContent()->GetCachedGeometry();
-		FVector2D ContentSize = Geometry.Size * Geometry.Scale;
-		const int32 Width = static_cast< int32 >( ContentSize.X );
-		const int32 Height = static_cast< int32 >( ContentSize.Y );
-
-		GConfig->SetInt( *ImporterName, *WidthConfigName, Width, GEditorPerProjectIni );
-		GConfig->SetInt( *ImporterName, *HeightConfigName, Height, GEditorPerProjectIni );
-
-		const bool bRemoveFromCache = false;
-		GConfig->Flush( bRemoveFromCache, GEditorPerProjectIni );
-	});
+	Window->GetOnWindowClosedEvent().AddLambda(
+		[WidthConfigName, HeightConfigName, OptionsWindow](const TSharedRef<SWindow>& Window)
+		{
+			// We use the geometry here because any sort of size we can extract from the window includes some added margins
+			// that are not easy to account for, even with GetWindowSizeFromClientSize. The original provided ClientSize
+			// seems to always match the content geometry though, so we're saving and loading the same thing
+			FGeometry Geometry = Window->GetContent()->GetCachedGeometry();
+			FVector2D ContentSize = Geometry.Size * Geometry.Scale;
+			const int32 Width = static_cast<int32>(ContentSize.X);
+			const int32 Height = static_cast<int32>(ContentSize.Y);
+
+			GConfig->SetInt(*ImporterName, *WidthConfigName, Width, GEditorPerProjectIni);
+			GConfig->SetInt(*ImporterName, *HeightConfigName, Height, GEditorPerProjectIni);
+
+			const bool bRemoveFromCache = false;
+			GConfig->Flush(bRemoveFromCache, GEditorPerProjectIni);
+		}
+	);
 
 	// Preemptively make sure we have a progress dialog created before showing our modal. This because the progress
 	// dialog itself is also modal. If it doesn't exist yet, and our options dialog causes a progress dialog
 	// to be spawned (e.g. when switching the Level to export via the LevelSequenceUSDExporter), the progress dialog
 	// will be pushed to the end of FSlateApplication::ActiveModalWindows (SlateApplication.cpp) and cause our options
 	// dialog to pop out of its modal loop (FSlateApplication::AddModalWindow), instantly returning false to our caller
-	FScopedSlowTask Progress( 1, LOCTEXT( "ShowingDialog", "Picking options..." ) );
+	FScopedSlowTask Progress(1, LOCTEXT("ShowingDialog", "Picking options..."));
 	Progress.MakeDialogDelayed(0.25f);
 
 	const bool bSlowTaskWindow = false;
-	FSlateApplication::Get().AddModalWindow( Window, ParentWindow, bSlowTaskWindow );
+	FSlateApplication::Get().AddModalWindow(Window, ParentWindow, bSlowTaskWindow);
 
 	const bool bAccepted = OptionsWindow->UserAccepted();
-	if ( bAccepted )
+	if (bAccepted)
 	{
-		if ( UUsdStageImportOptions* StageImportOptions = Cast< UUsdStageImportOptions>( &OptionsObject ) )
+		if (UUsdStageImportOptions* StageImportOptions = Cast<UUsdStageImportOptions>(&OptionsObject))
 		{
 			StageImportOptions->PrimsToImport = OptionsWindow->GetSelectedFullPrimPaths();
 		}
@@ -134,26 +123,26 @@ bool SUsdOptionsWindow::ShowOptions(
 	return bAccepted;
 }
 
-bool SUsdOptionsWindow::ShowImportOptions( UObject& OptionsObject, const UE::FUsdStage* StageToImport )
+bool SUsdOptionsWindow::ShowImportOptions(UObject& OptionsObject, const UE::FUsdStage* StageToImport)
 {
 	return SUsdOptionsWindow::ShowOptions(
 		OptionsObject,
-		LOCTEXT( "USDImportOptionsTitle", "USD Import Options" ),
-		LOCTEXT( "USDOptionWindow_Import", "Import" ),
+		LOCTEXT("USDImportOptionsTitle", "USD Import Options"),
+		LOCTEXT("USDOptionWindow_Import", "Import"),
 		StageToImport
 	);
 }
 
-bool SUsdOptionsWindow::ShowExportOptions( UObject& OptionsObject )
+bool SUsdOptionsWindow::ShowExportOptions(UObject& OptionsObject)
 {
 	return SUsdOptionsWindow::ShowOptions(
 		OptionsObject,
-		LOCTEXT( "USDExportOptionsTitle", "USD Export Options" ),
-		LOCTEXT( "USDOptionWindow_Export", "Export" )
+		LOCTEXT("USDExportOptionsTitle", "USD Export Options"),
+		LOCTEXT("USDOptionWindow_Export", "Export")
 	);
 }
 
-void SUsdOptionsWindow::Construct( const FArguments& InArgs )
+void SUsdOptionsWindow::Construct(const FArguments& InArgs)
 {
 	OptionsObject = InArgs._OptionsObject;
 	Window = InArgs._WidgetWindow;
@@ -164,65 +153,66 @@ void SUsdOptionsWindow::Construct( const FArguments& InArgs )
 	// Prepare box that shows the OptionsObject's properties
 	TSharedPtr<SWidget> DetailsViewBox;
 	{
-		FPropertyEditorModule& PropertyEditorModule = FModuleManager::GetModuleChecked<FPropertyEditorModule>( "PropertyEditor" );
+		FPropertyEditorModule& PropertyEditorModule = FModuleManager::GetModuleChecked<FPropertyEditorModule>("PropertyEditor");
 
 		FDetailsViewArgs DetailsViewArgs;
 		DetailsViewArgs.bAllowSearch = false;
 		DetailsViewArgs.NameAreaSettings = FDetailsViewArgs::HideNameArea;
-		TSharedPtr<IDetailsView> DetailsView = PropertyEditorModule.CreateDetailView( DetailsViewArgs );
-		DetailsView->SetObject( OptionsObject );
+		TSharedPtr<IDetailsView> DetailsView = PropertyEditorModule.CreateDetailView(DetailsViewArgs);
+		DetailsView->SetObject(OptionsObject);
 
 		DetailsViewBox = DetailsView.ToSharedRef();
 	}
 
-	TSharedRef<SSplitter> Splitter = SNew( SSplitter ).Orientation( Orient_Horizontal );
+	// clang-format off
+	TSharedRef<SSplitter> Splitter = SNew(SSplitter).Orientation(Orient_Horizontal);
 	{
 		// If we have a stage, show the preview tree on the left
-		if ( Stage )
+		if (Stage)
 		{
-			SAssignNew( StagePreviewTree, SUsdStagePreviewTree, *Stage );
+			SAssignNew(StagePreviewTree, SUsdStagePreviewTree, *Stage);
 
 			Splitter->AddSlot()
 			.Value(0.5f)
 			[
-				SNew( SVerticalBox )
+				SNew(SVerticalBox)
 
 				+SVerticalBox::Slot()
 				.AutoHeight()
-				.Padding( 0.0f )
+				.Padding(0.0f)
 				[
-					SNew( SHorizontalBox )
+					SNew(SHorizontalBox)
 					+SHorizontalBox::Slot()
-					.FillWidth( 1.0f )
-					.Padding( 0.0f )
-					.VAlign( VAlign_Center )
+					.FillWidth(1.0f)
+					.Padding(0.0f)
+					.VAlign(VAlign_Center)
 					[
-						SNew( STextBlock )
-						.Font( FAppStyle::GetFontStyle( "CurveEd.LabelFont" ) )
-						.Text( LOCTEXT( "TreeDescriptionText", "Select prims to import:" ) )
+						SNew(STextBlock)
+						.Font(FAppStyle::GetFontStyle("CurveEd.LabelFont"))
+						.Text(LOCTEXT("TreeDescriptionText", "Select prims to import:"))
 					]
 
 					+SHorizontalBox::Slot()
 					.AutoWidth()
-					.MaxWidth( 300.0f )
-					.Padding( FMargin( 5.0f, 2.0f ) )
-					.VAlign( VAlign_Center )
+					.MaxWidth(300.0f)
+					.Padding(FMargin(5.0f, 2.0f))
+					.VAlign(VAlign_Center)
 					[
-						SNew( SSearchBox )
-						.HintText( LOCTEXT( "SearchInitialText", "Filter prims..." ) )
-						.OnTextChanged_Lambda( [this]( const FText& SearchText )
+						SNew(SSearchBox)
+						.HintText(LOCTEXT("SearchInitialText", "Filter prims..."))
+						.OnTextChanged_Lambda([this](const FText& SearchText)
 						{
-							if ( StagePreviewTree )
+							if (StagePreviewTree)
 							{
-								StagePreviewTree->SetFilterText( SearchText );
+								StagePreviewTree->SetFilterText(SearchText);
 							}
 						})
 					]
 				]
 
 				+SVerticalBox::Slot()
-				.FillHeight( 1.0f )
-				.Padding( 0.0f )
+				.FillHeight(1.0f)
+				.Padding(0.0f)
 				[
 					StagePreviewTree.ToSharedRef()
 				]
@@ -238,44 +228,45 @@ void SUsdOptionsWindow::Construct( const FArguments& InArgs )
 
 	ChildSlot
 	[
-		SNew( SBorder )
-		.Padding( FMargin( 10.0f, 0.0f, 10.0f, 0.0f ) )
-		.BorderImage( FAppStyle::GetBrush( "ToolPanel.DarkGroupBorder" ) )
+		SNew(SBorder)
+		.Padding(FMargin(10.0f, 0.0f, 10.0f, 0.0f))
+		.BorderImage(FAppStyle::GetBrush("ToolPanel.DarkGroupBorder"))
 		[
-			SNew( SVerticalBox )
+			SNew(SVerticalBox)
 			+SVerticalBox::Slot()
-			.FillHeight( 1.0f )
-			.Padding( 0.0f )
+			.FillHeight(1.0f)
+			.Padding(0.0f)
 			[
 				Splitter
 			]
 
 			+SVerticalBox::Slot()
 			.AutoHeight()
-			.HAlign( HAlign_Right )
-			.Padding( 8 )
+			.HAlign(HAlign_Right)
+			.Padding(8)
 			[
-				SNew( SUniformGridPanel )
-				.SlotPadding( 2 )
+				SNew(SUniformGridPanel)
+				.SlotPadding(2)
 
-				+ SUniformGridPanel::Slot( 0, 0 )
+				+ SUniformGridPanel::Slot(0, 0)
 				[
-					SNew( SButton )
-					.HAlign( HAlign_Center )
-					.Text( AcceptText )
-					.OnClicked( this, &SUsdOptionsWindow::OnAccept )
+					SNew(SButton)
+					.HAlign(HAlign_Center)
+					.Text(AcceptText)
+					.OnClicked(this, &SUsdOptionsWindow::OnAccept)
 				]
 
-				+ SUniformGridPanel::Slot( 1, 0 )
+				+ SUniformGridPanel::Slot(1, 0)
 				[
-					SNew( SButton )
-					.HAlign( HAlign_Center )
-					.Text( LOCTEXT( "USDOptionWindow_Cancel", "Cancel" ) )
-					.OnClicked( this, &SUsdOptionsWindow::OnCancel )
+					SNew(SButton)
+					.HAlign(HAlign_Center)
+					.Text(LOCTEXT("USDOptionWindow_Cancel", "Cancel"))
+					.OnClicked(this, &SUsdOptionsWindow::OnCancel)
 				]
 			]
 		]
 	];
+	// clang-format on
 }
 
 bool SUsdOptionsWindow::SupportsKeyboardFocus() const
@@ -285,13 +276,13 @@ bool SUsdOptionsWindow::SupportsKeyboardFocus() const
 
 FReply SUsdOptionsWindow::OnAccept()
 {
-	if ( OptionsObject )
+	if (OptionsObject)
 	{
 		OptionsObject->SaveConfig();
 	}
 
 	bAccepted = true;
-	if ( Window.IsValid() )
+	if (Window.IsValid())
 	{
 		Window.Pin()->RequestDestroyWindow();
 	}
@@ -301,16 +292,16 @@ FReply SUsdOptionsWindow::OnAccept()
 FReply SUsdOptionsWindow::OnCancel()
 {
 	bAccepted = false;
-	if ( Window.IsValid() )
+	if (Window.IsValid())
 	{
 		Window.Pin()->RequestDestroyWindow();
 	}
 	return FReply::Handled();
 }
 
-FReply SUsdOptionsWindow::OnKeyDown( const FGeometry& MyGeometry, const FKeyEvent& InKeyEvent )
+FReply SUsdOptionsWindow::OnKeyDown(const FGeometry& MyGeometry, const FKeyEvent& InKeyEvent)
 {
-	if ( InKeyEvent.GetKey() == EKeys::Escape )
+	if (InKeyEvent.GetKey() == EKeys::Escape)
 	{
 		return OnCancel();
 	}
@@ -324,7 +315,7 @@ bool SUsdOptionsWindow::UserAccepted() const
 
 TArray<FString> SUsdOptionsWindow::GetSelectedFullPrimPaths() const
 {
-	if ( StagePreviewTree )
+	if (StagePreviewTree)
 	{
 		return StagePreviewTree->GetSelectedFullPrimPaths();
 	}
@@ -332,4 +323,4 @@ TArray<FString> SUsdOptionsWindow::GetSelectedFullPrimPaths() const
 	return {};
 }
 
-#undef LOCTEXT_NAMESPACE
\ No newline at end of file
+#undef LOCTEXT_NAMESPACE
diff --git a/Engine/Plugins/Importers/USDImporter/Source/USDStageImporter/Private/USDStageAssetImportFactory.cpp b/Engine/Plugins/Importers/USDImporter/Source/USDStageImporter/Private/USDStageAssetImportFactory.cpp
--- a/Engine/Plugins/Importers/USDImporter/Source/USDStageImporter/Private/USDStageAssetImportFactory.cpp
+++ b/Engine/Plugins/Importers/USDImporter/Source/USDStageImporter/Private/USDStageAssetImportFactory.cpp
@@ -4,20 +4,17 @@
 
 #include "USDAssetImportData.h"
 #include "USDAssetUserData.h"
+#include "USDConversionUtils.h"
 #include "USDErrorUtils.h"
-#include "USDLog.h"
 #include "USDStageImporter.h"
 #include "USDStageImporterModule.h"
 #include "USDStageImportOptions.h"
 
-#include "ActorFactories/ActorFactoryStaticMesh.h"
 #include "AssetImportTask.h"
 #include "Editor.h"
 #include "Engine/SkeletalMesh.h"
 #include "Engine/StaticMesh.h"
-#include "JsonObjectConverter.h"
 #include "Misc/Paths.h"
-#include "ProfilingDebugging/ScopedTimers.h"
 
 #define LOCTEXT_NAMESPACE "USDStageAssetImportFactory"
 
@@ -36,7 +33,7 @@ UUsdStageAssetImportFactory::UUsdStageAssetImportFactory(const FObjectInitialize
 	bEditorImport = true;
 	bText = false;
 
-	for ( const FString& Extension : UnrealUSDWrapper::GetNativeFileFormats() )
+	for (const FString& Extension : UnrealUSDWrapper::GetNativeFileFormats())
 	{
 		Formats.Add(FString::Printf(TEXT("%s; Universal Scene Description files"), *Extension));
 	}
@@ -52,24 +49,33 @@ UClass* UUsdStageAssetImportFactory::ResolveSupportedClass()
 	return UStaticMesh::StaticClass();
 }
 
-UObject* UUsdStageAssetImportFactory::FactoryCreateFile(UClass* InClass, UObject* InParent, FName InName, EObjectFlags Flags, const FString& Filename, const TCHAR* Parms, FFeedbackContext* Warn, bool& bOutOperationCanceled)
+UObject* UUsdStageAssetImportFactory::FactoryCreateFile(
+	UClass* InClass,
+	UObject* InParent,
+	FName InName,
+	EObjectFlags Flags,
+	const FString& Filename,
+	const TCHAR* Parms,
+	FFeedbackContext* Warn,
+	bool& bOutOperationCanceled
+)
 {
 	UObject* ImportedObject = nullptr;
 
-	if ( AssetImportTask && IsAutomatedImport() )
+	if (AssetImportTask && IsAutomatedImport())
 	{
-		ImportContext.ImportOptions = Cast<UUsdStageImportOptions>( AssetImportTask->Options );
+		ImportContext.ImportOptions = Cast<UUsdStageImportOptions>(AssetImportTask->Options);
 	}
 
 	// When importing from file we don't want to use any opened stage
 	ImportContext.bReadFromStageCache = false;
 
-	const FString InitialPackagePath = InParent ? InParent->GetName() : TEXT( "/Game/" );
+	const FString InitialPackagePath = InParent ? InParent->GetName() : TEXT("/Game/");
 	const bool bIsReimport = false;
 	const bool bAllowActorImport = false;
 	if (ImportContext.Init(InName.ToString(), Filename, InitialPackagePath, Flags, IsAutomatedImport(), bIsReimport, bAllowActorImport))
 	{
-		GEditor->GetEditorSubsystem<UImportSubsystem>()->BroadcastAssetPreImport( this, InClass, InParent, InName, Parms );
+		GEditor->GetEditorSubsystem<UImportSubsystem>()->BroadcastAssetPreImport(this, InClass, InParent, InName, Parms);
 
 		FScopedUsdMessageLog ScopedMessageLog;
 
@@ -91,11 +97,11 @@ UObject* UUsdStageAssetImportFactory::FactoryCreateFile(UClass* InClass, UObject
 
 bool UUsdStageAssetImportFactory::FactoryCanImport(const FString& Filename)
 {
-	const FString Extension = FPaths::GetExtension( Filename );
+	const FString Extension = FPaths::GetExtension(Filename);
 
-	for ( const FString& SupportedExtension : UnrealUSDWrapper::GetAllSupportedFileFormats() )
+	for (const FString& SupportedExtension : UnrealUSDWrapper::GetAllSupportedFileFormats())
 	{
-		if ( SupportedExtension.Equals( Extension, ESearchCase::IgnoreCase ) )
+		if (SupportedExtension.Equals(Extension, ESearchCase::IgnoreCase))
 		{
 			return true;
 		}
@@ -138,10 +144,7 @@ EReimportResult::Type UUsdStageAssetImportFactory::Reimport(UObject* Obj)
 {
 	if (!Obj)
 	{
-		FUsdLogManager::LogMessage(
-			EMessageSeverity::Error,
-			LOCTEXT("ReimportErrorInvalidAsset", "Failed to reimport asset as it is invalid!")
-		);
+		FUsdLogManager::LogMessage(EMessageSeverity::Error, LOCTEXT("ReimportErrorInvalidAsset", "Failed to reimport asset as it is invalid!"));
 		return EReimportResult::Failed;
 	}
 
@@ -171,10 +174,7 @@ EReimportResult::Type UUsdStageAssetImportFactory::Reimport(UObject* Obj)
 		FUsdLogManager::LogMessage(
 			EMessageSeverity::Error,
 			FText::Format(
-				LOCTEXT(
-					"ReimportErrorNoImportData",
-					"Failed to reimport asset '{0}' as it doesn't seem to have valid USD import data or user data!"
-				),
+				LOCTEXT("ReimportErrorNoImportData", "Failed to reimport asset '{0}' as it doesn't seem to have valid USD import data or user data!"),
 				FText::FromName(Obj->GetFName())
 			)
 		);
@@ -193,15 +193,7 @@ EReimportResult::Type UUsdStageAssetImportFactory::Reimport(UObject* Obj)
 
 	const bool bIsReimport = true;
 	const bool bAllowActorImport = false;
-	if (!ImportContext.Init(
-			Obj->GetName(),
-			ReimportFilePath,
-			Obj->GetName(),
-			Obj->GetFlags(),
-			IsAutomatedImport(),
-			bIsReimport,
-			bAllowActorImport
-		))
+	if (!ImportContext.Init(Obj->GetName(), ReimportFilePath, Obj->GetName(), Obj->GetFlags(), IsAutomatedImport(), bIsReimport, bAllowActorImport))
 	{
 		FUsdLogManager::LogMessage(
 			EMessageSeverity::Error,
@@ -231,4 +223,4 @@ int32 UUsdStageAssetImportFactory::GetPriority() const
 	return ImportPriority;
 }
 
-#undef LOCTEXT_NAMESPACE
\ No newline at end of file
+#undef LOCTEXT_NAMESPACE
diff --git a/Engine/Plugins/Importers/USDImporter/Source/USDStageImporter/Private/USDStageAssetImportFactory.h b/Engine/Plugins/Importers/USDImporter/Source/USDStageImporter/Private/USDStageAssetImportFactory.h
--- a/Engine/Plugins/Importers/USDImporter/Source/USDStageImporter/Private/USDStageAssetImportFactory.h
+++ b/Engine/Plugins/Importers/USDImporter/Source/USDStageImporter/Private/USDStageAssetImportFactory.h
@@ -3,25 +3,34 @@
 #pragma once
 
 #include "USDStageImportContext.h"
-#include "USDStageImportOptions.h"
 
 #include "EditorReimportHandler.h"
 #include "Factories/Factory.h"
-#include "Factories/ImportSettings.h"
 
 #include "USDStageAssetImportFactory.generated.h"
 
 /** Factory to import USD files that gets called when we hit Import in the Content Browser, as well as during reimport */
-UCLASS(hidecategories=Object)
-class UUsdStageAssetImportFactory : public UFactory, public FReimportHandler
+UCLASS(hidecategories = Object)
+class UUsdStageAssetImportFactory
+	: public UFactory
+	, public FReimportHandler
 {
 	GENERATED_UCLASS_BODY()
 
 public:
 	/** UFactory interface */
-	virtual bool DoesSupportClass(UClass * Class) override;
+	virtual bool DoesSupportClass(UClass* Class) override;
 	virtual UClass* ResolveSupportedClass() override;
-	virtual UObject* FactoryCreateFile(UClass* InClass, UObject* InParent, FName InName, EObjectFlags Flags, const FString& Filename, const TCHAR* Parms, FFeedbackContext* Warn, bool& bOutOperationCanceled) override;
+	virtual UObject* FactoryCreateFile(
+		UClass* InClass,
+		UObject* InParent,
+		FName InName,
+		EObjectFlags Flags,
+		const FString& Filename,
+		const TCHAR* Parms,
+		FFeedbackContext* Warn,
+		bool& bOutOperationCanceled
+	) override;
 	virtual bool FactoryCanImport(const FString& Filename) override;
 	virtual void CleanUp() override;
 
diff --git a/Engine/Plugins/Importers/USDImporter/Source/USDStageImporter/Private/USDStageImportContext.cpp b/Engine/Plugins/Importers/USDImporter/Source/USDStageImporter/Private/USDStageImportContext.cpp
--- a/Engine/Plugins/Importers/USDImporter/Source/USDStageImporter/Private/USDStageImportContext.cpp
+++ b/Engine/Plugins/Importers/USDImporter/Source/USDStageImporter/Private/USDStageImportContext.cpp
@@ -2,25 +2,27 @@
 
 #include "USDStageImportContext.h"
 
-#include "USDAssetCache.h"
-#include "USDLog.h"
-#include "USDMemory.h"
 #include "USDOptionsWindow.h"
 #include "USDStageImportOptions.h"
 
 #include "Dialogs/DlgPickPath.h"
 #include "Editor.h"
-#include "IMessageLogListing.h"
-#include "MessageLogModule.h"
 #include "Misc/Paths.h"
-#include "Modules/ModuleManager.h"
 
 FUsdStageImportContext::FUsdStageImportContext()
 {
 	Reset();
 }
 
-bool FUsdStageImportContext::Init(const FString& InName, const FString& InFilePath, const FString& InInitialPackagePath, EObjectFlags InFlags, bool bInIsAutomated, bool bIsReimport, bool bAllowActorImport)
+bool FUsdStageImportContext::Init(
+	const FString& InName,
+	const FString& InFilePath,
+	const FString& InInitialPackagePath,
+	EObjectFlags InFlags,
+	bool bInIsAutomated,
+	bool bIsReimport,
+	bool bAllowActorImport
+)
 {
 	ObjectName = InName;
 	FilePath = InFilePath;
@@ -29,29 +31,31 @@ bool FUsdStageImportContext::Init(const FString& InName, const FString& InFilePa
 	World = GEditor->GetEditorWorldContext().World();
 	PackagePath = InInitialPackagePath;
 
-	if ( !PackagePath.EndsWith( TEXT("/") ) )
+	if (!PackagePath.EndsWith(TEXT("/")))
 	{
-		PackagePath.Append( TEXT("/") );
+		PackagePath.Append(TEXT("/"));
 	}
 
 	FPaths::NormalizeFilename(FilePath);
 
 	// Open the stage if we haven't yet, as we'll need it open to show the preview tree
-	if ( !Stage )
+	if (!Stage)
 	{
 		const EUsdInitialLoadSet InitialLoadSet = EUsdInitialLoadSet::LoadAll;
-		Stage = UnrealUSDWrapper::OpenStage( *FilePath, InitialLoadSet, bReadFromStageCache );
+		Stage = UnrealUSDWrapper::OpenStage(*FilePath, InitialLoadSet, bReadFromStageCache);
 	}
 
-	if(!bIsAutomated)
+	if (!bIsAutomated)
 	{
 		// Show dialog for content folder
 		if (!bIsReimport)
 		{
+			// clang-format off
 			TSharedRef<SDlgPickPath> PickContentPathDlg =
 				SNew(SDlgPickPath)
 				.Title(NSLOCTEXT("USDStageImportContext", "ChooseImportRootContentPath", "Choose where to place the imported USD assets"))
-				.DefaultPath(FText::FromString( InInitialPackagePath ));
+				.DefaultPath(FText::FromString(InInitialPackagePath));
+			// clang-format on
 
 			if (PickContentPathDlg->ShowModal() == EAppReturnType::Cancel)
 			{
@@ -60,13 +64,13 @@ bool FUsdStageImportContext::Init(const FString& InName, const FString& InFilePa
 			// e.g. "/Game/MyFolder/layername/"
 			// We inject the package path here because this is what the automated import task upstream code will do. This way the importer
 			// can always expect to receive /ContentPath/layername/
-			PackagePath = FString::Printf( TEXT( "%s/%s/" ), *PickContentPathDlg->GetPath().ToString(), *InName );
+			PackagePath = FString::Printf(TEXT("%s/%s/"), *PickContentPathDlg->GetPath().ToString(), *InName);
 		}
 
-		ImportOptions->EnableActorImport( bAllowActorImport );
+		ImportOptions->EnableActorImport(bAllowActorImport);
 
 		// Show dialog for import options
-		bool bProceedWithImport = SUsdOptionsWindow::ShowImportOptions( *ImportOptions, &Stage );
+		bool bProceedWithImport = SUsdOptionsWindow::ShowImportOptions(*ImportOptions, &Stage);
 		if (!bProceedWithImport)
 		{
 			return false;
@@ -85,7 +89,7 @@ void FUsdStageImportContext::Reset()
 	ObjectName = FString{};
 	PackagePath = FString{};
 	FilePath = FString{};
-	ImportOptions = NewObject< UUsdStageImportOptions >();
+	ImportOptions = NewObject<UUsdStageImportOptions>();
 	ImportedAsset = nullptr;
 	LevelSequenceHelper.Clear();
 	AssetCache = nullptr;
@@ -108,4 +112,4 @@ void FUsdStageImportContext::Reset()
 	OriginalUpAxis = EUsdUpAxis::ZAxis;
 
 	bNeedsGarbageCollection = false;
-}
\ No newline at end of file
+}
diff --git a/Engine/Plugins/Importers/USDImporter/Source/USDStageImporter/Private/USDStageImportFactory.cpp b/Engine/Plugins/Importers/USDImporter/Source/USDStageImporter/Private/USDStageImportFactory.cpp
--- a/Engine/Plugins/Importers/USDImporter/Source/USDStageImporter/Private/USDStageImportFactory.cpp
+++ b/Engine/Plugins/Importers/USDImporter/Source/USDStageImporter/Private/USDStageImportFactory.cpp
@@ -2,21 +2,13 @@
 
 #include "USDStageImportFactory.h"
 
-#include "USDConversionUtils.h"
 #include "USDErrorUtils.h"
+#include "USDStageImporter.h"
 #include "USDStageImporterModule.h"
 #include "USDStageImportOptions.h"
 
 #include "AssetImportTask.h"
-#include "AssetRegistry/AssetRegistryModule.h"
-#include "AssetSelection.h"
 #include "Editor.h"
-#include "AssetRegistry/IAssetRegistry.h"
-#include "JsonObjectConverter.h"
-#include "Misc/ScopedSlowTask.h"
-#include "ObjectTools.h"
-#include "PackageTools.h"
-#include "ScopedTransaction.h"
 
 #define LOCTEXT_NAMESPACE "USDImportFactory"
 
@@ -32,31 +24,40 @@ UUsdStageImportFactory::UUsdStageImportFactory(const FObjectInitializer& ObjectI
 	bEditorImport = true;
 	bText = false;
 
-	for ( const FString& Extension : UnrealUSDWrapper::GetNativeFileFormats() )
+	for (const FString& Extension : UnrealUSDWrapper::GetNativeFileFormats())
 	{
-		Formats.Add( FString::Printf( TEXT( "%s; Universal Scene Description files" ), *Extension ) );
+		Formats.Add(FString::Printf(TEXT("%s; Universal Scene Description files"), *Extension));
 	}
 }
 
-UObject* UUsdStageImportFactory::FactoryCreateFile(UClass* InClass, UObject* InParent, FName InName, EObjectFlags Flags, const FString& Filename, const TCHAR* Parms, FFeedbackContext* Warn, bool& bOutOperationCanceled)
+UObject* UUsdStageImportFactory::FactoryCreateFile(
+	UClass* InClass,
+	UObject* InParent,
+	FName InName,
+	EObjectFlags Flags,
+	const FString& Filename,
+	const TCHAR* Parms,
+	FFeedbackContext* Warn,
+	bool& bOutOperationCanceled
+)
 {
 	UObject* ImportedObject = nullptr;
 
-	if ( AssetImportTask && IsAutomatedImport() )
+	if (AssetImportTask && IsAutomatedImport())
 	{
-		ImportContext.ImportOptions = Cast<UUsdStageImportOptions>( AssetImportTask->Options );
+		ImportContext.ImportOptions = Cast<UUsdStageImportOptions>(AssetImportTask->Options);
 	}
 
 	// When importing from file we don't want to use any opened stage
 	ImportContext.bReadFromStageCache = false;
 
 #if USE_USD_SDK
-	const FString InitialPackagePath =InParent ? InParent->GetName() : TEXT( "/Game/" );
+	const FString InitialPackagePath = InParent ? InParent->GetName() : TEXT("/Game/");
 	const bool bIsReimport = false;
 	const bool bAllowActorImport = true;
 	if (ImportContext.Init(InName.ToString(), Filename, InitialPackagePath, Flags, IsAutomatedImport(), bIsReimport, bAllowActorImport))
 	{
-		GEditor->GetEditorSubsystem<UImportSubsystem>()->BroadcastAssetPreImport( this, InClass, InParent, InName, Parms );
+		GEditor->GetEditorSubsystem<UImportSubsystem>()->BroadcastAssetPreImport(this, InClass, InParent, InName, Parms);
 
 		FScopedUsdMessageLog ScopedMessageLog;
 
@@ -67,14 +68,14 @@ UObject* UUsdStageImportFactory::FactoryCreateFile(UClass* InClass, UObject* InP
 		GEditor->BroadcastLevelActorListChanged();
 		GEditor->RedrawLevelEditingViewports();
 
-		ImportedObject = ImportContext.ImportedAsset ? ToRawPtr(ImportContext.ImportedAsset) : Cast<UObject>( ImportContext.SceneActor );
+		ImportedObject = ImportContext.ImportedAsset ? ToRawPtr(ImportContext.ImportedAsset) : Cast<UObject>(ImportContext.SceneActor);
 	}
 	else
 	{
 		bOutOperationCanceled = true;
 	}
 
-#endif // #if USE_USD_SDK
+#endif	  // #if USE_USD_SDK
 
 	return ImportedObject;
 }
@@ -83,9 +84,9 @@ bool UUsdStageImportFactory::FactoryCanImport(const FString& Filename)
 {
 	const FString Extension = FPaths::GetExtension(Filename);
 
-	for ( const FString& SupportedExtension : UnrealUSDWrapper::GetAllSupportedFileFormats() )
+	for (const FString& SupportedExtension : UnrealUSDWrapper::GetAllSupportedFileFormats())
 	{
-		if ( SupportedExtension.Equals( Extension, ESearchCase::IgnoreCase ) )
+		if (SupportedExtension.Equals(Extension, ESearchCase::IgnoreCase))
 		{
 			return true;
 		}
diff --git a/Engine/Plugins/Importers/USDImporter/Source/USDStageImporter/Private/USDStageImportFactory.h b/Engine/Plugins/Importers/USDImporter/Source/USDStageImporter/Private/USDStageImportFactory.h
--- a/Engine/Plugins/Importers/USDImporter/Source/USDStageImporter/Private/USDStageImportFactory.h
+++ b/Engine/Plugins/Importers/USDImporter/Source/USDStageImporter/Private/USDStageImportFactory.h
@@ -3,17 +3,11 @@
 #pragma once
 
 #include "USDStageImportContext.h"
-#include "USDStageImporter.h"
 
-#include "Editor/EditorEngine.h"
-#include "Factories/ImportSettings.h"
-#include "Factories/Factory.h"
 #include "Factories/SceneImportFactory.h"
 
 #include "USDStageImportFactory.generated.h"
 
-class UUsdStageImportOptions;
-
 /** Factory to import USD files that gets called when we hit File -> Import into level... */
 UCLASS(Transient)
 class UUsdStageImportFactory : public USceneImportFactory
@@ -22,7 +16,16 @@ class UUsdStageImportFactory : public USceneImportFactory
 
 public:
 	// UFactory Interface
-	virtual UObject* FactoryCreateFile(UClass* InClass, UObject* InParent, FName InName, EObjectFlags Flags, const FString& Filename, const TCHAR* Parms, FFeedbackContext* Warn, bool& bOutOperationCanceled) override;
+	virtual UObject* FactoryCreateFile(
+		UClass* InClass,
+		UObject* InParent,
+		FName InName,
+		EObjectFlags Flags,
+		const FString& Filename,
+		const TCHAR* Parms,
+		FFeedbackContext* Warn,
+		bool& bOutOperationCanceled
+	) override;
 	virtual bool FactoryCanImport(const FString& Filename) override;
 	virtual void CleanUp() override;
 
@@ -30,4 +33,3 @@ private:
 	UPROPERTY()
 	FUsdStageImportContext ImportContext;
 };
-
diff --git a/Engine/Plugins/Importers/USDImporter/Source/USDStageImporter/Private/USDStageImportOptions.cpp b/Engine/Plugins/Importers/USDImporter/Source/USDStageImporter/Private/USDStageImportOptions.cpp
--- a/Engine/Plugins/Importers/USDImporter/Source/USDStageImporter/Private/USDStageImportOptions.cpp
+++ b/Engine/Plugins/Importers/USDImporter/Source/USDStageImporter/Private/USDStageImportOptions.cpp
@@ -6,7 +6,6 @@
 #include "USDSchemaTranslator.h"
 
 #include "AnalyticsEventAttribute.h"
-#include "InterchangeGenericMaterialPipeline.h"
 #include "Modules/ModuleManager.h"
 #include "UObject/UnrealType.h"
 
@@ -21,18 +20,18 @@ UUsdStageImportOptions::UUsdStageImportOptions(const FObjectInitializer& ObjectI
 	bImportGroomAssets = true;
 	bImportOnlyUsedMaterials = false;
 
-	PurposesToImport = (int32) (EUsdPurpose::Default | EUsdPurpose::Proxy | EUsdPurpose::Render | EUsdPurpose::Guide);
+	PurposesToImport = (int32)(EUsdPurpose::Default | EUsdPurpose::Proxy | EUsdPurpose::Render | EUsdPurpose::Guide);
 	NaniteTriangleThreshold = INT32_MAX;
-	IUsdSchemasModule& UsdSchemasModule = FModuleManager::Get().LoadModuleChecked< IUsdSchemasModule >( TEXT("USDSchemas") );
+	IUsdSchemasModule& UsdSchemasModule = FModuleManager::Get().LoadModuleChecked<IUsdSchemasModule>(TEXT("USDSchemas"));
 	RenderContextToImport = UsdSchemasModule.GetRenderContextRegistry().GetUnrealRenderContext();
 	MaterialPurpose = *UnrealIdentifiers::MaterialPreviewPurpose;
 	SubdivisionLevel = 0;
 	MetadataOptions = FUsdMetadataImportOptions{
-		true, 	/* bCollectMetadata */
-		true, 	/* bCollectFromEntireSubtrees */
-		false, 	/* bCollectOnComponents */
-		{},		/* BlockedPrefixFilters */
-		false	/* bInvertFilters */
+		true,  /* bCollectMetadata */
+		true,  /* bCollectFromEntireSubtrees */
+		false, /* bCollectOnComponents */
+		{},	   /* BlockedPrefixFilters */
+		false  /* bInvertFilters */
 	};
 	bOverrideStageOptions = false;
 	StageOptions.MetersPerUnit = 0.01f;
@@ -45,19 +44,19 @@ UUsdStageImportOptions::UUsdStageImportOptions(const FObjectInitializer& ObjectI
 	bReuseIdenticalAssets = true;
 
 	bPrimPathFolderStructure = false;
-	KindsToCollapse = ( int32 ) ( EUsdDefaultKind::Component | EUsdDefaultKind::Subcomponent );
+	KindsToCollapse = (int32)(EUsdDefaultKind::Component | EUsdDefaultKind::Subcomponent);
 	bMergeIdenticalMaterialSlots = true;
 	bInterpretLODs = true;
 }
 
-void UUsdStageImportOptions::EnableActorImport( bool bEnable )
+void UUsdStageImportOptions::EnableActorImport(bool bEnable)
 {
-	for ( TFieldIterator<FProperty> PropertyIterator( UUsdStageImportOptions::StaticClass() ); PropertyIterator; ++PropertyIterator )
+	for (TFieldIterator<FProperty> PropertyIterator(UUsdStageImportOptions::StaticClass()); PropertyIterator; ++PropertyIterator)
 	{
 		FProperty* Property = *PropertyIterator;
-		if ( Property && Property->GetFName() == GET_MEMBER_NAME_CHECKED( UUsdStageImportOptions, bImportActors ) )
+		if (Property && Property->GetFName() == GET_MEMBER_NAME_CHECKED(UUsdStageImportOptions, bImportActors))
 		{
-			if ( bEnable )
+			if (bEnable)
 			{
 				Property->SetMetaData(TEXT("ToolTip"), TEXT("Whether to spawn imported actors into the current level"));
 				Property->ClearPropertyFlags(CPF_EditConst);
@@ -65,7 +64,10 @@ void UUsdStageImportOptions::EnableActorImport( bool bEnable )
 			else
 			{
 				bImportActors = false;
-				Property->SetMetaData(TEXT("ToolTip"), TEXT("Actor import is disabled when importing via the Content Browser. Use File->\"Import into Level...\" to also import actors."));
+				Property->SetMetaData(
+					TEXT("ToolTip"),
+					TEXT("Actor import is disabled when importing via the Content Browser. Use File->\"Import into Level...\" to also import actors.")
+				);
 				Property->SetPropertyFlags(CPF_EditConst);
 			}
 			break;
@@ -83,47 +85,44 @@ void UUsdStageImportOptions::PostEditChangeProperty(struct FPropertyChangedEvent
 	}
 }
 
-void UsdUtils::AddAnalyticsAttributes(
-	const UUsdStageImportOptions& Options,
-	TArray< FAnalyticsEventAttribute >& InOutAttributes
-)
+void UsdUtils::AddAnalyticsAttributes(const UUsdStageImportOptions& Options, TArray<FAnalyticsEventAttribute>& InOutAttributes)
 {
-	InOutAttributes.Emplace( TEXT( "ImportActors" ), LexToString( Options.bImportActors ) );
-	InOutAttributes.Emplace( TEXT( "ImportGeometry" ), LexToString( Options.bImportGeometry ) );
-	InOutAttributes.Emplace( TEXT( "ImportSkeletalAnimations" ), LexToString( Options.bImportSkeletalAnimations ) );
-	InOutAttributes.Emplace( TEXT( "ImportLevelSequences" ), LexToString( Options.bImportLevelSequences ) );
-	InOutAttributes.Emplace( TEXT( "ImportMaterials" ), LexToString( Options.bImportMaterials ) );
-	InOutAttributes.Emplace( TEXT( "ImportGroomAssets" ), LexToString( Options.bImportGroomAssets ) );
-	InOutAttributes.Emplace( TEXT( "ImportOnlyUsedMaterials" ), LexToString( Options.bImportOnlyUsedMaterials) );
-	if ( Options.PrimsToImport != TArray<FString>{ TEXT( "/" ) } )
+	InOutAttributes.Emplace(TEXT("ImportActors"), LexToString(Options.bImportActors));
+	InOutAttributes.Emplace(TEXT("ImportGeometry"), LexToString(Options.bImportGeometry));
+	InOutAttributes.Emplace(TEXT("ImportSkeletalAnimations"), LexToString(Options.bImportSkeletalAnimations));
+	InOutAttributes.Emplace(TEXT("ImportLevelSequences"), LexToString(Options.bImportLevelSequences));
+	InOutAttributes.Emplace(TEXT("ImportMaterials"), LexToString(Options.bImportMaterials));
+	InOutAttributes.Emplace(TEXT("ImportGroomAssets"), LexToString(Options.bImportGroomAssets));
+	InOutAttributes.Emplace(TEXT("ImportOnlyUsedMaterials"), LexToString(Options.bImportOnlyUsedMaterials));
+	if (Options.PrimsToImport != TArray<FString>{TEXT("/")})
 	{
-		InOutAttributes.Emplace( TEXT( "NumPrimsToImport" ), LexToString( Options.PrimsToImport.Num() ) );
+		InOutAttributes.Emplace(TEXT("NumPrimsToImport"), LexToString(Options.PrimsToImport.Num()));
 	}
-	InOutAttributes.Emplace( TEXT( "PurposesToImport" ), LexToString( Options.PurposesToImport ) );
-	InOutAttributes.Emplace( TEXT( "NaniteTriangleThreshold" ), LexToString( Options.NaniteTriangleThreshold ) );
-	InOutAttributes.Emplace( TEXT( "RenderContextToImport" ), Options.RenderContextToImport.ToString() );
-	InOutAttributes.Emplace( TEXT( "MaterialPurpose" ), Options.MaterialPurpose.ToString() );
-	InOutAttributes.Emplace( TEXT( "RootMotionHandling" ), LexToString( ( uint8 ) Options.RootMotionHandling ) );
+	InOutAttributes.Emplace(TEXT("PurposesToImport"), LexToString(Options.PurposesToImport));
+	InOutAttributes.Emplace(TEXT("NaniteTriangleThreshold"), LexToString(Options.NaniteTriangleThreshold));
+	InOutAttributes.Emplace(TEXT("RenderContextToImport"), Options.RenderContextToImport.ToString());
+	InOutAttributes.Emplace(TEXT("MaterialPurpose"), Options.MaterialPurpose.ToString());
+	InOutAttributes.Emplace(TEXT("RootMotionHandling"), LexToString((uint8)Options.RootMotionHandling));
 	InOutAttributes.Emplace(TEXT("SubdivisionLevel"), LexToString(Options.SubdivisionLevel));
 
 	UsdUtils::AddAnalyticsAttributes(Options.MetadataOptions, InOutAttributes);
 
-	InOutAttributes.Emplace( TEXT( "OverrideStageOptions" ), Options.bOverrideStageOptions );
-	if ( Options.bOverrideStageOptions )
+	InOutAttributes.Emplace(TEXT("OverrideStageOptions"), Options.bOverrideStageOptions);
+	if (Options.bOverrideStageOptions)
 	{
-		UsdUtils::AddAnalyticsAttributes( Options.StageOptions, InOutAttributes );
+		UsdUtils::AddAnalyticsAttributes(Options.StageOptions, InOutAttributes);
 	}
-	InOutAttributes.Emplace( TEXT( "ImportAtSpecificTimeCode" ), Options.bImportAtSpecificTimeCode );
-	if ( Options.bImportAtSpecificTimeCode )
+	InOutAttributes.Emplace(TEXT("ImportAtSpecificTimeCode"), Options.bImportAtSpecificTimeCode);
+	if (Options.bImportAtSpecificTimeCode)
 	{
-		InOutAttributes.Emplace( TEXT( "ImportTimeCode" ), LexToString( Options.ImportTimeCode ) );
+		InOutAttributes.Emplace(TEXT("ImportTimeCode"), LexToString(Options.ImportTimeCode));
 	}
-	InOutAttributes.Emplace( TEXT( "NumGroomInterpolationSettings" ), LexToString( Options.GroomInterpolationSettings.Num() ) );
-	InOutAttributes.Emplace( TEXT( "ReplaceActorPolicy" ), LexToString( ( uint8 ) Options.ExistingActorPolicy ) );
-	InOutAttributes.Emplace( TEXT( "ReplaceAssetPolicy" ), LexToString( ( uint8 ) Options.ExistingAssetPolicy ) );
+	InOutAttributes.Emplace(TEXT("NumGroomInterpolationSettings"), LexToString(Options.GroomInterpolationSettings.Num()));
+	InOutAttributes.Emplace(TEXT("ReplaceActorPolicy"), LexToString((uint8)Options.ExistingActorPolicy));
+	InOutAttributes.Emplace(TEXT("ReplaceAssetPolicy"), LexToString((uint8)Options.ExistingAssetPolicy));
 	InOutAttributes.Emplace(TEXT("ReuseIdenticalAssets"), LexToString(Options.bReuseIdenticalAssets));
-	InOutAttributes.Emplace( TEXT( "PrimPathFolderStructure" ), LexToString( Options.bPrimPathFolderStructure ) );
-	InOutAttributes.Emplace( TEXT( "KindsToCollapse" ), LexToString( Options.KindsToCollapse ) );
-	InOutAttributes.Emplace( TEXT( "MergeIdenticalMaterialSlots" ), LexToString( Options.bMergeIdenticalMaterialSlots ) );
-	InOutAttributes.Emplace( TEXT( "InterpretLODs" ), LexToString( Options.bInterpretLODs ) );
-}
\ No newline at end of file
+	InOutAttributes.Emplace(TEXT("PrimPathFolderStructure"), LexToString(Options.bPrimPathFolderStructure));
+	InOutAttributes.Emplace(TEXT("KindsToCollapse"), LexToString(Options.KindsToCollapse));
+	InOutAttributes.Emplace(TEXT("MergeIdenticalMaterialSlots"), LexToString(Options.bMergeIdenticalMaterialSlots));
+	InOutAttributes.Emplace(TEXT("InterpretLODs"), LexToString(Options.bInterpretLODs));
+}
diff --git a/Engine/Plugins/Importers/USDImporter/Source/USDStageImporter/Private/USDStageImportOptionsCustomization.cpp b/Engine/Plugins/Importers/USDImporter/Source/USDStageImporter/Private/USDStageImportOptionsCustomization.cpp
--- a/Engine/Plugins/Importers/USDImporter/Source/USDStageImporter/Private/USDStageImportOptionsCustomization.cpp
+++ b/Engine/Plugins/Importers/USDImporter/Source/USDStageImporter/Private/USDStageImportOptionsCustomization.cpp
@@ -30,121 +30,129 @@ TSharedRef<IDetailCustomization> FUsdStageImportOptionsCustomization::MakeInstan
 void FUsdStageImportOptionsCustomization::CustomizeDetails(IDetailLayoutBuilder& DetailLayoutBuilder)
 {
 	TArray<TWeakObjectPtr<UObject>> SelectedObjects = DetailLayoutBuilder.GetSelectedObjects();
-	if ( SelectedObjects.Num() != 1 )
+	if (SelectedObjects.Num() != 1)
 	{
 		return;
 	}
 
-	TWeakObjectPtr<UObject> SelectedObject = SelectedObjects[ 0 ];
-	if ( !SelectedObject.IsValid() )
+	TWeakObjectPtr<UObject> SelectedObject = SelectedObjects[0];
+	if (!SelectedObject.IsValid())
 	{
 		return;
 	}
 
-	CurrentOptions = Cast<UUsdStageImportOptions>( SelectedObject.Get() );
-	if ( !CurrentOptions )
+	CurrentOptions = Cast<UUsdStageImportOptions>(SelectedObject.Get());
+	if (!CurrentOptions)
 	{
 		return;
 	}
 
 	// Hide this property since we'll show the preview tree for it
-	IDetailCategoryBuilder& PrimsToImportCatBuilder = DetailLayoutBuilder.EditCategory( TEXT( "Prims to Import" ) );
-	if ( TSharedPtr<IPropertyHandle> PrimsToImportProperty = DetailLayoutBuilder.GetProperty( GET_MEMBER_NAME_CHECKED( UUsdStageImportOptions, PrimsToImport ) ) )
+	DetailLayoutBuilder.EditCategory(TEXT("Prims to Import"));
+	if (TSharedPtr<IPropertyHandle> PrimsToImportProperty = DetailLayoutBuilder.GetProperty(
+			GET_MEMBER_NAME_CHECKED(UUsdStageImportOptions, PrimsToImport)
+		))
 	{
-		DetailLayoutBuilder.HideProperty( PrimsToImportProperty );
+		DetailLayoutBuilder.HideProperty(PrimsToImportProperty);
 	}
 
-	IUsdSchemasModule& UsdSchemasModule = FModuleManager::Get().LoadModuleChecked< IUsdSchemasModule >( TEXT( "USDSchemas" ) );
+	IUsdSchemasModule& UsdSchemasModule = FModuleManager::Get().LoadModuleChecked<IUsdSchemasModule>(TEXT("USDSchemas"));
 
 	RenderContextComboBoxItems.Reset();
 	TSharedPtr<FString> InitiallySelectedContext;
-	for ( const FName& Context : UsdSchemasModule.GetRenderContextRegistry().GetRenderContexts() )
+	for (const FName& Context : UsdSchemasModule.GetRenderContextRegistry().GetRenderContexts())
 	{
 		TSharedPtr<FString> ContextStr;
-		if ( Context == NAME_None )
+		if (Context == NAME_None)
 		{
-			ContextStr = MakeShared<FString>( TEXT( "universal" ) );
+			ContextStr = MakeShared<FString>(TEXT("universal"));
 		}
 		else
 		{
-			ContextStr = MakeShared<FString>( Context.ToString() );
+			ContextStr = MakeShared<FString>(Context.ToString());
 		}
 
-		if ( Context == CurrentOptions->RenderContextToImport )
+		if (Context == CurrentOptions->RenderContextToImport)
 		{
 			InitiallySelectedContext = ContextStr;
 		}
 
-		RenderContextComboBoxItems.Add( ContextStr );
+		RenderContextComboBoxItems.Add(ContextStr);
 	}
 
-	IDetailCategoryBuilder& CatBuilder = DetailLayoutBuilder.EditCategory( TEXT( "USD options" ) );
+	IDetailCategoryBuilder& CatBuilder = DetailLayoutBuilder.EditCategory(TEXT("USD options"));
 
-	if ( TSharedPtr<IPropertyHandle> RenderContextProperty = DetailLayoutBuilder.GetProperty( GET_MEMBER_NAME_CHECKED( UUsdStageImportOptions, RenderContextToImport ) ) )
+	if (TSharedPtr<IPropertyHandle> RenderContextProperty = DetailLayoutBuilder.GetProperty(
+			GET_MEMBER_NAME_CHECKED(UUsdStageImportOptions, RenderContextToImport)
+		))
 	{
-		DetailLayoutBuilder.HideProperty( RenderContextProperty );
+		DetailLayoutBuilder.HideProperty(RenderContextProperty);
 
-		CatBuilder.AddCustomRow( FText::FromString( TEXT( "RenderContextCustomization" ) ) )
+		// clang-format off
+		CatBuilder.AddCustomRow(FText::FromString(TEXT("RenderContextCustomization")))
 		.NameContent()
 		[
-			SNew( STextBlock )
-			.Text( FText::FromString( TEXT( "Render Context to Import" ) ) )
-			.Font( FAppStyle::GetFontStyle( TEXT( "PropertyWindow.NormalFont" ) ) )
-			.ToolTipText( RenderContextProperty->GetToolTipText() )
+			SNew(STextBlock)
+			.Text(FText::FromString(TEXT("Render Context to Import")))
+			.Font(FAppStyle::GetFontStyle(TEXT("PropertyWindow.NormalFont")))
+			.ToolTipText(RenderContextProperty->GetToolTipText())
 		]
 		.ValueContent()
 		[
-			SAssignNew( RenderContextComboBox, SComboBox<TSharedPtr<FString>> )
-			.OptionsSource( &RenderContextComboBoxItems )
-			.InitiallySelectedItem( InitiallySelectedContext )
-			.OnSelectionChanged( this, &FUsdStageImportOptionsCustomization::OnComboBoxSelectionChanged )
-			.OnGenerateWidget_Lambda( []( TSharedPtr<FString> Item )
+			SAssignNew(RenderContextComboBox, SComboBox<TSharedPtr<FString>>)
+			.OptionsSource(&RenderContextComboBoxItems)
+			.InitiallySelectedItem(InitiallySelectedContext)
+			.OnSelectionChanged(this, &FUsdStageImportOptionsCustomization::OnComboBoxSelectionChanged)
+			.OnGenerateWidget_Lambda([](TSharedPtr<FString> Item)
 			{
-				return SNew( STextBlock )
-					.Text( Item.IsValid() ? FText::FromString( *Item ) : FText::GetEmpty() )
-					.Font(FAppStyle::GetFontStyle( TEXT( "PropertyWindow.NormalFont" ) ) );
-			} )
+				return SNew(STextBlock)
+					.Text(Item.IsValid() ? FText::FromString(*Item) : FText::GetEmpty())
+					.Font(FAppStyle::GetFontStyle(TEXT("PropertyWindow.NormalFont")));
+			})
 			.Content()
 			[
-				SNew( STextBlock )
-				.Text( this, &FUsdStageImportOptionsCustomization::GetComboBoxSelectedOptionText )
-				.Font(FAppStyle::GetFontStyle( TEXT( "PropertyWindow.NormalFont" ) ) )
+				SNew(STextBlock)
+				.Text(this, &FUsdStageImportOptionsCustomization::GetComboBoxSelectedOptionText)
+				.Font(FAppStyle::GetFontStyle(TEXT("PropertyWindow.NormalFont")))
 			]
 		];
+		// clang-format on
 	}
 
-
-	if ( TSharedPtr<IPropertyHandle> MaterialPurposeProperty = DetailLayoutBuilder.GetProperty( GET_MEMBER_NAME_CHECKED( UUsdStageImportOptions, MaterialPurpose ) ) )
+	if (TSharedPtr<IPropertyHandle> MaterialPurposeProperty = DetailLayoutBuilder.GetProperty(
+			GET_MEMBER_NAME_CHECKED(UUsdStageImportOptions, MaterialPurpose)
+		))
 	{
-		DetailLayoutBuilder.HideProperty( MaterialPurposeProperty );
+		DetailLayoutBuilder.HideProperty(MaterialPurposeProperty);
 
-		CatBuilder.AddCustomRow( FText::FromString( TEXT( "MaterialPurposeCustomization" ) ) )
+		// clang-format off
+		CatBuilder.AddCustomRow(FText::FromString(TEXT("MaterialPurposeCustomization")))
 		.NameContent()
 		[
-			SNew( STextBlock )
-			.Text( FText::FromString( TEXT( "Material purpose" ) ) )
-			.Font( FAppStyle::GetFontStyle( TEXT( "PropertyWindow.NormalFont" ) ) )
-			.ToolTipText( MaterialPurposeProperty->GetToolTipText() )
+			SNew(STextBlock)
+			.Text(FText::FromString(TEXT("Material purpose")))
+			.Font(FAppStyle::GetFontStyle(TEXT("PropertyWindow.NormalFont")))
+			.ToolTipText(MaterialPurposeProperty->GetToolTipText())
 		]
 		.ValueContent()
 		[
-			SNew( SBox )
-			.VAlign( VAlign_Center )
+			SNew(SBox)
+			.VAlign(VAlign_Center)
 			[
-				SNew( SComboBox< TSharedPtr<FString> > )
-				.OptionsSource( &MaterialPurposeComboBoxItems )
+				SNew(SComboBox<TSharedPtr<FString>>)
+				.OptionsSource(&MaterialPurposeComboBoxItems)
 				.OnComboBoxOpening_Lambda([this]()
 				{
 					MaterialPurposeComboBoxItems = {
-						MakeShared<FString>( UnrealIdentifiers::MaterialAllPurpose ),
-						MakeShared<FString>( UnrealIdentifiers::MaterialPreviewPurpose ),
-						MakeShared<FString>( UnrealIdentifiers::MaterialFullPurpose )
+						MakeShared<FString>(UnrealIdentifiers::MaterialAllPurpose),
+						MakeShared<FString>(UnrealIdentifiers::MaterialPreviewPurpose),
+						MakeShared<FString>(UnrealIdentifiers::MaterialFullPurpose)
 					};
 
 					// Add additional purposes from project settings
-					if ( const UUsdProjectSettings* ProjectSettings = GetDefault<UUsdProjectSettings>() )
+					if (const UUsdProjectSettings* ProjectSettings = GetDefault<UUsdProjectSettings>())
 					{
-						MaterialPurposeComboBoxItems.Reserve( MaterialPurposeComboBoxItems.Num() + ProjectSettings->AdditionalMaterialPurposes.Num() );
+						MaterialPurposeComboBoxItems.Reserve(MaterialPurposeComboBoxItems.Num() + ProjectSettings->AdditionalMaterialPurposes.Num());
 
 						TSet<FString> ExistingEntries = {
 							UnrealIdentifiers::MaterialAllPurpose,
@@ -152,47 +160,47 @@ void FUsdStageImportOptionsCustomization::CustomizeDetails(IDetailLayoutBuilder&
 							UnrealIdentifiers::MaterialFullPurpose
 						};
 
-						for ( const FName& AdditionalPurpose : ProjectSettings->AdditionalMaterialPurposes )
+						for (const FName& AdditionalPurpose : ProjectSettings->AdditionalMaterialPurposes)
 						{
 							FString AdditionalPurposeStr = AdditionalPurpose.ToString();
 
-							if ( !ExistingEntries.Contains( AdditionalPurposeStr ) )
+							if (!ExistingEntries.Contains(AdditionalPurposeStr))
 							{
-								ExistingEntries.Add( AdditionalPurposeStr );
-								MaterialPurposeComboBoxItems.AddUnique( MakeShared<FString>( AdditionalPurposeStr ) );
+								ExistingEntries.Add(AdditionalPurposeStr);
+								MaterialPurposeComboBoxItems.AddUnique(MakeShared<FString>(AdditionalPurposeStr));
 							}
 						}
 					}
 				})
-				.OnGenerateWidget_Lambda( [ & ]( TSharedPtr<FString> Option )
+				.OnGenerateWidget_Lambda([&](TSharedPtr<FString> Option)
 				{
 					TSharedPtr<SWidget> Widget = SNullWidget::NullWidget;
-					if ( Option )
+					if (Option)
 					{
-						Widget = SNew( STextBlock )
-							.Text( FText::FromString( ( *Option ) == UnrealIdentifiers::MaterialAllPurpose
+						Widget = SNew(STextBlock)
+							.Text(FText::FromString((*Option) == UnrealIdentifiers::MaterialAllPurpose
 								? UnrealIdentifiers::MaterialAllPurposeText
 								: *Option
-							) )
-							.Font( FAppStyle::GetFontStyle( "PropertyWindow.NormalFont" ) );
+							))
+							.Font(FAppStyle::GetFontStyle("PropertyWindow.NormalFont"));
 					}
 
 					return Widget.ToSharedRef();
 				})
-				.OnSelectionChanged_Lambda([this]( TSharedPtr<FString> ChosenOption, ESelectInfo::Type SelectInfo )
+				.OnSelectionChanged_Lambda([this](TSharedPtr<FString> ChosenOption, ESelectInfo::Type SelectInfo)
 				{
-					if ( CurrentOptions && ChosenOption )
+					if (CurrentOptions && ChosenOption)
 					{
 						CurrentOptions->MaterialPurpose = **ChosenOption;
 					}
 				})
 				[
-					SNew( SEditableTextBox )
+					SNew(SEditableTextBox)
 					.Text_Lambda([this]() -> FText
 					{
-						if ( CurrentOptions )
+						if (CurrentOptions)
 						{
-							return FText::FromString( CurrentOptions->MaterialPurpose == *UnrealIdentifiers::MaterialAllPurpose
+							return FText::FromString(CurrentOptions->MaterialPurpose == *UnrealIdentifiers::MaterialAllPurpose
 								? UnrealIdentifiers::MaterialAllPurposeText
 								: CurrentOptions->MaterialPurpose.ToString()
 							);
@@ -200,10 +208,10 @@ void FUsdStageImportOptionsCustomization::CustomizeDetails(IDetailLayoutBuilder&
 
 						return FText::GetEmpty();
 					})
-					.Font( FAppStyle::GetFontStyle( "PropertyWindow.NormalFont" ) )
-					.OnTextCommitted_Lambda( [this]( const FText& NewText, ETextCommit::Type CommitType )
+					.Font(FAppStyle::GetFontStyle("PropertyWindow.NormalFont"))
+					.OnTextCommitted_Lambda([this](const FText& NewText, ETextCommit::Type CommitType)
 					{
-						if ( CommitType != ETextCommit::OnEnter )
+						if (CommitType != ETextCommit::OnEnter)
 						{
 							return;
 						}
@@ -212,25 +220,25 @@ void FUsdStageImportOptionsCustomization::CustomizeDetails(IDetailLayoutBuilder&
 						FName NewPurpose = *NewPurposeString;
 
 						bool bIsNew = true;
-						for ( const TSharedPtr<FString>& Purpose : MaterialPurposeComboBoxItems )
+						for (const TSharedPtr<FString>& Purpose : MaterialPurposeComboBoxItems)
 						{
-							if ( Purpose && *Purpose == NewPurposeString )
+							if (Purpose && *Purpose == NewPurposeString)
 							{
 								bIsNew = false;
 								break;
 							}
 						}
 
-						if ( bIsNew )
+						if (bIsNew)
 						{
-							if ( UUsdProjectSettings* ProjectSettings = GetMutableDefault<UUsdProjectSettings>() )
+							if (UUsdProjectSettings* ProjectSettings = GetMutableDefault<UUsdProjectSettings>())
 							{
-								ProjectSettings->AdditionalMaterialPurposes.AddUnique( NewPurpose );
+								ProjectSettings->AdditionalMaterialPurposes.AddUnique(NewPurpose);
 								ProjectSettings->SaveConfig();
 							}
 						}
 
-						if ( CurrentOptions )
+						if (CurrentOptions)
 						{
 							CurrentOptions->MaterialPurpose = NewPurpose;
 						}
@@ -238,32 +246,37 @@ void FUsdStageImportOptionsCustomization::CustomizeDetails(IDetailLayoutBuilder&
 				]
 			]
 		];
+		// clang-format on
 	}
 
 	// Add/remove properties so that they retain their usual order
-	if ( TSharedPtr<IPropertyHandle> OverrideStageOptionsProperty = DetailLayoutBuilder.GetProperty( GET_MEMBER_NAME_CHECKED( UUsdStageImportOptions, bOverrideStageOptions ) ) )
+	if (TSharedPtr<IPropertyHandle> OverrideStageOptionsProperty = DetailLayoutBuilder.GetProperty(
+			GET_MEMBER_NAME_CHECKED(UUsdStageImportOptions, bOverrideStageOptions)
+		))
 	{
-		CatBuilder.AddProperty( OverrideStageOptionsProperty );
+		CatBuilder.AddProperty(OverrideStageOptionsProperty);
 	}
-	if ( TSharedPtr<IPropertyHandle> StageOptionsProperty = DetailLayoutBuilder.GetProperty( GET_MEMBER_NAME_CHECKED( UUsdStageImportOptions, StageOptions ) ) )
+	if (TSharedPtr<IPropertyHandle> StageOptionsProperty = DetailLayoutBuilder.GetProperty(
+			GET_MEMBER_NAME_CHECKED(UUsdStageImportOptions, StageOptions)
+		))
 	{
-		CatBuilder.AddProperty( StageOptionsProperty );
+		CatBuilder.AddProperty(StageOptionsProperty);
 	}
 }
 
-void FUsdStageImportOptionsCustomization::CustomizeDetails( const TSharedPtr<IDetailLayoutBuilder>& DetailBuilder )
+void FUsdStageImportOptionsCustomization::CustomizeDetails(const TSharedPtr<IDetailLayoutBuilder>& DetailBuilder)
 {
-	CustomizeDetails( *DetailBuilder );
+	CustomizeDetails(*DetailBuilder);
 }
 
-void FUsdStageImportOptionsCustomization::OnComboBoxSelectionChanged( TSharedPtr<FString> NewContext, ESelectInfo::Type SelectType )
+void FUsdStageImportOptionsCustomization::OnComboBoxSelectionChanged(TSharedPtr<FString> NewContext, ESelectInfo::Type SelectType)
 {
-	if ( CurrentOptions == nullptr || !NewContext.IsValid() )
+	if (CurrentOptions == nullptr || !NewContext.IsValid())
 	{
 		return;
 	}
 
-	FName NewContextName = ( *NewContext ) == TEXT( "universal" ) ? NAME_None : FName( **NewContext );
+	FName NewContextName = (*NewContext) == TEXT("universal") ? NAME_None : FName(**NewContext);
 
 	CurrentOptions->RenderContextToImport = NewContextName;
 }
@@ -271,9 +284,9 @@ void FUsdStageImportOptionsCustomization::OnComboBoxSelectionChanged( TSharedPtr
 FText FUsdStageImportOptionsCustomization::GetComboBoxSelectedOptionText() const
 {
 	TSharedPtr<FString> SelectedItem = RenderContextComboBox->GetSelectedItem();
-	if ( SelectedItem.IsValid() )
+	if (SelectedItem.IsValid())
 	{
-		return FText::FromString( *SelectedItem );
+		return FText::FromString(*SelectedItem);
 	}
 
 	return FText::GetEmpty();
diff --git a/Engine/Plugins/Importers/USDImporter/Source/USDStageImporter/Private/USDStageImporter.cpp b/Engine/Plugins/Importers/USDImporter/Source/USDStageImporter/Private/USDStageImporter.cpp
--- a/Engine/Plugins/Importers/USDImporter/Source/USDStageImporter/Private/USDStageImporter.cpp
+++ b/Engine/Plugins/Importers/USDImporter/Source/USDStageImporter/Private/USDStageImporter.cpp
@@ -2,13 +2,12 @@
 
 #include "USDStageImporter.h"
 
-#include "USDAssetCache.h"
 #include "USDAssetCache2.h"
 #include "USDAssetImportData.h"
 #include "USDAssetUserData.h"
-#include "USDDrawModeComponent.h"
 #include "USDClassesModule.h"
 #include "USDConversionUtils.h"
+#include "USDDrawModeComponent.h"
 #include "USDErrorUtils.h"
 #include "USDGeomMeshConversion.h"
 #include "USDLog.h"
@@ -32,7 +31,6 @@
 #include "ComponentRecreateRenderStateContext.h"
 #include "Components/SkeletalMeshComponent.h"
 #include "Components/StaticMeshComponent.h"
-#include "Dialogs/DlgPickPath.h"
 #include "Editor.h"
 #include "Engine/SkeletalMesh.h"
 #include "Engine/SkinnedAssetCommon.h"
@@ -48,7 +46,6 @@
 #include "HAL/FileManager.h"
 #include "IAssetTools.h"
 #include "LevelSequence.h"
-#include "Materials/Material.h"
 #include "Materials/MaterialInstance.h"
 #include "Materials/MaterialInterface.h"
 #include "Misc/Paths.h"
@@ -58,20 +55,19 @@
 #include "PhysicsEngine/PhysicsAsset.h"
 #include "Serialization/ArchiveReplaceObjectRef.h"
 #include "Subsystems/AssetEditorSubsystem.h"
-#include "UObject/StrongObjectPtr.h"
 #include "UObject/UObjectIterator.h"
 
 #define LOCTEXT_NAMESPACE "USDStageImporter"
 
-namespace UsdStageImporterImpl
+namespace UE::USDStageImporter::Private
 {
-	void OpenStage( FUsdStageImportContext& ImportContext, bool bNeedsMasking )
+	void OpenStage(FUsdStageImportContext& ImportContext, bool bNeedsMasking)
 	{
 		const FString FilePath = !ImportContext.FilePath.IsEmpty()
-			? IFileManager::Get().ConvertToAbsolutePathForExternalAppForRead( *ImportContext.FilePath )
-			: FString{};
+									 ? IFileManager::Get().ConvertToAbsolutePathForExternalAppForRead(*ImportContext.FilePath)
+									 : FString{};
 
-		if ( FilePath.IsEmpty() && !ImportContext.Stage )
+		if (FilePath.IsEmpty() && !ImportContext.Stage)
 		{
 			// There's nothing we can do here to create a stage
 			return;
@@ -80,7 +76,7 @@ namespace UsdStageImporterImpl
 		UsdUtils::StartMonitoringErrors();
 
 		UE::FUsdStage Stage;
-		if ( bNeedsMasking )
+		if (bNeedsMasking)
 		{
 			// If we're masking we'll make a unique stage for us.
 			// Also, within USD the OpenMasked methods do not consult the stage cache anyway
@@ -93,7 +89,7 @@ namespace UsdStageImporterImpl
 			const bool bForceReloadLayersFromDisk = false;
 
 			// We don't have a file path to reopen, so just reopen the existing layers
-			if ( FilePath.IsEmpty() )
+			if (FilePath.IsEmpty())
 			{
 				Stage = UnrealUSDWrapper::OpenMaskedStage(
 					ImportContext.Stage.GetRootLayer(),
@@ -115,11 +111,11 @@ namespace UsdStageImporterImpl
 		}
 		else
 		{
-			for ( const UE::FUsdStage& OpenedStage : UnrealUSDWrapper::GetAllStagesFromCache() )
+			for (const UE::FUsdStage& OpenedStage : UnrealUSDWrapper::GetAllStagesFromCache())
 			{
 				FString RootPath = OpenedStage.GetRootLayer().GetRealPath();
-				FPaths::NormalizeFilename( RootPath );
-				if ( ImportContext.FilePath == RootPath )
+				FPaths::NormalizeFilename(RootPath);
+				if (ImportContext.FilePath == RootPath)
 				{
 					ImportContext.bStageWasOriginallyOpenInCache = true;
 					break;
@@ -145,7 +141,14 @@ namespace UsdStageImporterImpl
 		else
 		{
 			ImportContext.Stage = UE::FUsdStage();
-			FUsdLogManager::LogMessage( EMessageSeverity::Error, FText::Format( LOCTEXT( "CouldNotImportUSDFile", "Could not import USD file {0}\n {1}" ), FText::FromString( FilePath ), FText::FromString( Error ) ) );
+			FUsdLogManager::LogMessage(
+				EMessageSeverity::Error,
+				FText::Format(
+					LOCTEXT("CouldNotImportUSDFile", "Could not import USD file {0}\n {1}"),
+					FText::FromString(FilePath),
+					FText::FromString(Error)
+				)
+			);
 		}
 	}
 
@@ -167,8 +170,7 @@ namespace UsdStageImporterImpl
 			}
 
 			SearchPackagePath = InPackagePath + TEXT("_") + LexToString(Suffix++);
-		}
-		while(ExistingPackage != nullptr);
+		} while (ExistingPackage != nullptr);
 
 		// Undo the last SearchPackagePath update, returning the path that worked (vacant Package path)
 		return Suffix == 1 ? InPackagePath : InPackagePath + TEXT("_") + LexToString(Suffix - 1);
@@ -176,13 +178,13 @@ namespace UsdStageImporterImpl
 
 	void SetupSceneActor(FUsdStageImportContext& ImportContext)
 	{
-		if ( !ImportContext.ImportOptions->bImportActors )
+		if (!ImportContext.ImportOptions->bImportActors)
 		{
 			return;
 		}
 
 		ULevel* Level = ImportContext.World->GetCurrentLevel();
-		if(!Level)
+		if (!Level)
 		{
 			return;
 		}
@@ -212,12 +214,12 @@ namespace UsdStageImporterImpl
 			RootComponent->RegisterComponent();
 		}
 
-		if ( ImportContext.TargetSceneActorAttachParent )
+		if (ImportContext.TargetSceneActorAttachParent)
 		{
-			RootComponent->AttachToComponent( ImportContext.TargetSceneActorAttachParent, FAttachmentTransformRules::KeepRelativeTransform );
+			RootComponent->AttachToComponent(ImportContext.TargetSceneActorAttachParent, FAttachmentTransformRules::KeepRelativeTransform);
 		}
 
-		Actor->SetActorTransform( ImportContext.TargetSceneActorTargetTransform );
+		Actor->SetActorTransform(ImportContext.TargetSceneActorTargetTransform);
 
 		ImportContext.SceneActor = Actor;
 	}
@@ -231,10 +233,8 @@ namespace UsdStageImporterImpl
 			AActor* ThisActor = *ActorItr;
 
 			// Found a top level actor with the same label
-			if ( !ThisActor->HasAnyFlags(RF_Transient) &&
-				 ThisActor->GetAttachParentActor() == nullptr &&
-				 ThisActor->GetActorLabel() == TargetActorLabel &&
-				 ThisActor != ImportContext.SceneActor)
+			if (!ThisActor->HasAnyFlags(RF_Transient) && ThisActor->GetAttachParentActor() == nullptr
+				&& ThisActor->GetActorLabel() == TargetActorLabel && ThisActor != ImportContext.SceneActor)
 			{
 				return ThisActor;
 			}
@@ -243,31 +243,34 @@ namespace UsdStageImporterImpl
 		return nullptr;
 	}
 
-	void SetupStageForImport( FUsdStageImportContext& ImportContext )
+	void SetupStageForImport(FUsdStageImportContext& ImportContext)
 	{
 #if USE_USD_SDK
-		if ( ImportContext.ImportOptions->bOverrideStageOptions )
+		if (ImportContext.ImportOptions->bOverrideStageOptions)
 		{
-			ImportContext.OriginalMetersPerUnit = UsdUtils::GetUsdStageMetersPerUnit( ImportContext.Stage );
-			ImportContext.OriginalUpAxis = UsdUtils::GetUsdStageUpAxisAsEnum( ImportContext.Stage );
+			ImportContext.OriginalMetersPerUnit = UsdUtils::GetUsdStageMetersPerUnit(ImportContext.Stage);
+			ImportContext.OriginalUpAxis = UsdUtils::GetUsdStageUpAxisAsEnum(ImportContext.Stage);
 
-			UsdUtils::SetUsdStageMetersPerUnit( ImportContext.Stage, ImportContext.ImportOptions->StageOptions.MetersPerUnit );
-			UsdUtils::SetUsdStageUpAxis( ImportContext.Stage, ImportContext.ImportOptions->StageOptions.UpAxis );
+			UsdUtils::SetUsdStageMetersPerUnit(ImportContext.Stage, ImportContext.ImportOptions->StageOptions.MetersPerUnit);
+			UsdUtils::SetUsdStageUpAxis(ImportContext.Stage, ImportContext.ImportOptions->StageOptions.UpAxis);
 		}
-#endif // #if USE_USD_SDK
+#endif	  // #if USE_USD_SDK
 	}
 
-	void CreateAssetsForPrims(const TArray<UE::FUsdPrim>& Prims, FUsdSchemaTranslationContext& TranslationContext, const FText& ProgressMessage )
+	void CreateAssetsForPrims(const TArray<UE::FUsdPrim>& Prims, FUsdSchemaTranslationContext& TranslationContext, const FText& ProgressMessage)
 	{
-		TRACE_CPUPROFILER_EVENT_SCOPE( CreateAssetsForPrims );
+		TRACE_CPUPROFILER_EVENT_SCOPE(CreateAssetsForPrims);
 
 		IUsdSchemasModule& UsdSchemasModule = FModuleManager::Get().LoadModuleChecked<IUsdSchemasModule>(TEXT("USDSchemas"));
 
-		FScopedSlowTask SlowTask( Prims.Num(), ProgressMessage );
+		FScopedSlowTask SlowTask(Prims.Num(), ProgressMessage);
 
 		for (const UE::FUsdPrim& Prim : Prims)
 		{
-			if (TSharedPtr<FUsdSchemaTranslator> SchemaTranslator = UsdSchemasModule.GetTranslatorRegistry().CreateTranslatorForSchema(TranslationContext.AsShared(), UE::FUsdTyped(Prim)))
+			if (TSharedPtr<FUsdSchemaTranslator> SchemaTranslator = UsdSchemasModule.GetTranslatorRegistry().CreateTranslatorForSchema(
+					TranslationContext.AsShared(),
+					UE::FUsdTyped(Prim)
+				))
 			{
 				SchemaTranslator->CreateAssets();
 			}
@@ -278,9 +281,9 @@ namespace UsdStageImporterImpl
 		TranslationContext.CompleteTasks();
 	}
 
-	void CacheCollapsingState( FUsdSchemaTranslationContext& TranslationContext )
+	void CacheCollapsingState(FUsdSchemaTranslationContext& TranslationContext)
 	{
-		if ( !TranslationContext.InfoCache.IsValid() )
+		if (!TranslationContext.InfoCache.IsValid())
 		{
 			TranslationContext.InfoCache = MakeShared<FUsdInfoCache>();
 		}
@@ -292,21 +295,21 @@ namespace UsdStageImporterImpl
 
 	void ImportMaterials(FUsdStageImportContext& ImportContext, FUsdSchemaTranslationContext& TranslationContext)
 	{
-		TRACE_CPUPROFILER_EVENT_SCOPE( ImportMaterials );
+		TRACE_CPUPROFILER_EVENT_SCOPE(ImportMaterials);
 
 		if (!ImportContext.ImportOptions->bImportMaterials)
 		{
 			return;
 		}
 
-		TArray< UE::FUsdPrim > MaterialPrims = UsdUtils::GetAllPrimsOfType( ImportContext.Stage.GetPseudoRoot(), TEXT("UsdShadeMaterial") );
+		TArray<UE::FUsdPrim> MaterialPrims = UsdUtils::GetAllPrimsOfType(ImportContext.Stage.GetPseudoRoot(), TEXT("UsdShadeMaterial"));
 
-		CreateAssetsForPrims(MaterialPrims, TranslationContext, LOCTEXT("CreateMaterials", "Creating materials") );
+		CreateAssetsForPrims(MaterialPrims, TranslationContext, LOCTEXT("CreateMaterials", "Creating materials"));
 	}
 
 	void ImportMeshes(FUsdStageImportContext& ImportContext, FUsdSchemaTranslationContext& TranslationContext)
 	{
-		TRACE_CPUPROFILER_EVENT_SCOPE( ImportMeshes );
+		TRACE_CPUPROFILER_EVENT_SCOPE(ImportMeshes);
 
 #if USE_USD_SDK
 		if (!ImportContext.ImportOptions->bImportGeometry)
@@ -318,7 +321,10 @@ namespace UsdStageImporterImpl
 
 		auto PruneCollapsedMeshes = [&UsdSchemasModule, &TranslationContext](const UE::FUsdPrim& UsdPrim) -> bool
 		{
-			if (TSharedPtr< FUsdSchemaTranslator > SchemaTranslator = UsdSchemasModule.GetTranslatorRegistry().CreateTranslatorForSchema(TranslationContext.AsShared(), UE::FUsdTyped(UsdPrim)))
+			if (TSharedPtr<FUsdSchemaTranslator> SchemaTranslator = UsdSchemasModule.GetTranslatorRegistry().CreateTranslatorForSchema(
+					TranslationContext.AsShared(),
+					UE::FUsdTyped(UsdPrim)
+				))
 			{
 				return SchemaTranslator->CollapsesChildren(ECollapsingType::Assets);
 			}
@@ -326,9 +332,13 @@ namespace UsdStageImporterImpl
 			return false;
 		};
 
-		TArray< UE::FUsdPrim > MeshPrims = UsdUtils::GetAllPrimsOfType( ImportContext.Stage.GetPseudoRoot(), TEXT("UsdGeomXformable"), PruneCollapsedMeshes );
-		CreateAssetsForPrims(MeshPrims, TranslationContext, LOCTEXT( "CreateMeshes", "Creating meshes" ) );
-#endif // #if USE_USD_SDK
+		TArray<UE::FUsdPrim> MeshPrims = UsdUtils::GetAllPrimsOfType(
+			ImportContext.Stage.GetPseudoRoot(),
+			TEXT("UsdGeomXformable"),
+			PruneCollapsedMeshes
+		);
+		CreateAssetsForPrims(MeshPrims, TranslationContext, LOCTEXT("CreateMeshes", "Creating meshes"));
+#endif	  // #if USE_USD_SDK
 	}
 
 	void ImportAnimation(
@@ -354,19 +364,25 @@ namespace UsdStageImporterImpl
 		ImportContext.LevelSequenceHelper.AddPrim(*UsdPrimTwin, bAnimatedVisibility, bHasAnimatedBounds);
 	}
 
-	void ImportActor(FUsdStageImportContext& ImportContext, UE::FUsdPrim& Prim, bool bForceVisibilityAnimationTracks, FUsdSchemaTranslationContext& TranslationContext)
+	void ImportActor(
+		FUsdStageImportContext& ImportContext,
+		UE::FUsdPrim& Prim,
+		bool bForceVisibilityAnimationTracks,
+		FUsdSchemaTranslationContext& TranslationContext
+	)
 	{
-		TRACE_CPUPROFILER_EVENT_SCOPE( ImportActor );
+		TRACE_CPUPROFILER_EVENT_SCOPE(ImportActor);
 
-		IUsdSchemasModule& UsdSchemasModule = FModuleManager::Get().LoadModuleChecked< IUsdSchemasModule >(TEXT("USDSchemas"));
+		IUsdSchemasModule& UsdSchemasModule = FModuleManager::Get().LoadModuleChecked<IUsdSchemasModule>(TEXT("USDSchemas"));
 
-		FScopedSlowTask SlowTask( 3, LOCTEXT("ImportActor", "Importing Actor") );
+		FScopedSlowTask SlowTask(3, LOCTEXT("ImportActor", "Importing Actor"));
 
 		bool bExpandChilren = true;
 		USceneComponent* Component = nullptr;
 
 		// Spawn components and/or actors for this prim
-		if (TSharedPtr<FUsdSchemaTranslator> SchemaTranslator = UsdSchemasModule.GetTranslatorRegistry().CreateTranslatorForSchema(TranslationContext.AsShared(),UE::FUsdTyped(Prim)))
+		if (TSharedPtr<FUsdSchemaTranslator> SchemaTranslator = UsdSchemasModule.GetTranslatorRegistry()
+																	.CreateTranslatorForSchema(TranslationContext.AsShared(), UE::FUsdTyped(Prim)))
 		{
 			Component = SchemaTranslator->CreateComponents();
 
@@ -378,7 +394,7 @@ namespace UsdStageImporterImpl
 		// then we must ensure that we generate visibility tracks for our child prims that don't collapse either, so that
 		// they get the chance to bake their "computed visibilities" and e.g. be hidden whenever a parent prim goes invisible
 #if USE_USD_SDK
-		const bool bAnimatedVisibility = bForceVisibilityAnimationTracks || UsdUtils::HasAnimatedVisibility( Prim );
+		const bool bAnimatedVisibility = bForceVisibilityAnimationTracks || UsdUtils::HasAnimatedVisibility(Prim);
 #else
 		const bool bAnimatedVisibility = bForceVisibilityAnimationTracks;
 #endif
@@ -397,12 +413,12 @@ namespace UsdStageImporterImpl
 			}
 		}
 		SlowTask.EnterProgressFrame();
-		if ( Component )
+		if (Component)
 		{
 			// LightComponents specifically need this to setup static lighting
 			Component->PostEditChange();
 
-			if ( !Component->IsRegistered() )
+			if (!Component->IsRegistered())
 			{
 				Component->RegisterComponent();
 			}
@@ -426,7 +442,7 @@ namespace UsdStageImporterImpl
 					ImportAnimation(ImportContext, Prim, bAnimatedVisibility, HasAnimatedBounds, Component);
 				}
 			}
-#endif // USE_USD_SDK
+#endif	  // USE_USD_SDK
 		}
 		SlowTask.EnterProgressFrame();
 	}
@@ -470,7 +486,7 @@ namespace UsdStageImporterImpl
 			}
 		}
 
-		FString AssetName = FPaths::GetBaseFilename( AssetPath );
+		FString AssetName = FPaths::GetBaseFilename(AssetPath);
 
 		if (UStaticMesh* Mesh = Cast<UStaticMesh>(Asset))
 		{
@@ -517,35 +533,35 @@ namespace UsdStageImporterImpl
 			AssetName = Asset->GetFName().GetPlainNameString();
 			AssetPrefix = TEXT("SKEL_");
 		}
-		else if ( UPhysicsAsset* PhysicsAsset = Cast<UPhysicsAsset>( Asset ) )
+		else if (UPhysicsAsset* PhysicsAsset = Cast<UPhysicsAsset>(Asset))
 		{
 			// See comments above on the case for USkeleton
 			FString TempName = Asset->GetFName().GetPlainNameString();
 
 			// The asset is named after the SkelRoot prim. If we're importing back a scene that was originally exported,
 			// we should clean up these prefixes or else we may end up with something like "PHYS_SK_PrimName"
-			TempName.RemoveFromStart( TEXT( "PHYS_" ), ESearchCase::CaseSensitive );
-			TempName.RemoveFromStart( TEXT( "SK_" ), ESearchCase::CaseSensitive );
-			if ( !TempName.IsEmpty() )
+			TempName.RemoveFromStart(TEXT("PHYS_"), ESearchCase::CaseSensitive);
+			TempName.RemoveFromStart(TEXT("SK_"), ESearchCase::CaseSensitive);
+			if (!TempName.IsEmpty())
 			{
 				AssetName = TempName;
 			}
 
-			AssetPrefix = TEXT( "PHYS_" );
+			AssetPrefix = TEXT("PHYS_");
 		}
-		else if ( UAnimSequence* AnimSequence = Cast<UAnimSequence>( Asset ) )
+		else if (UAnimSequence* AnimSequence = Cast<UAnimSequence>(Asset))
 		{
-			AssetPrefix = TEXT( "AS_" );
+			AssetPrefix = TEXT("AS_");
 		}
 		else if (UMaterialInterface* Material = Cast<UMaterialInterface>(Asset))
 		{
-			if ( Material->IsA<UMaterialInstance>() )
+			if (Material->IsA<UMaterialInstance>())
 			{
-				AssetPrefix = TEXT( "MI_" );
+				AssetPrefix = TEXT("MI_");
 			}
 			else
 			{
-				AssetPrefix = TEXT( "M_" );
+				AssetPrefix = TEXT("M_");
 			}
 
 			// The only materials with no prim path are our auto-generated displayColor materials
@@ -577,27 +593,27 @@ namespace UsdStageImporterImpl
 			// textures inside USDZ files. It's just easier to take from the asset name itself instead
 			AssetName = Texture->GetFName().GetPlainNameString();
 		}
-		else if ( ULevelSequence* LevelSequence = Cast<ULevelSequence>( Asset ) )
+		else if (ULevelSequence* LevelSequence = Cast<ULevelSequence>(Asset))
 		{
-			AssetPrefix = TEXT( "LS_" );
+			AssetPrefix = TEXT("LS_");
 		}
-		else if ( UAnimBlueprint* AnimBP = Cast<UAnimBlueprint>( Asset ) )
+		else if (UAnimBlueprint* AnimBP = Cast<UAnimBlueprint>(Asset))
 		{
 			FString TempName = AssetName;
 
 			// The asset is named after the SkelRoot prim. If we're importing back a scene that was originally exported,
 			// we should clean up these prefixes or else we may end up with something like "ABP_SK_PrimName"
-			TempName.RemoveFromStart( TEXT( "ABP_" ), ESearchCase::CaseSensitive );
-			TempName.RemoveFromStart( TEXT( "SK_" ), ESearchCase::CaseSensitive );
-			if ( !TempName.IsEmpty() )
+			TempName.RemoveFromStart(TEXT("ABP_"), ESearchCase::CaseSensitive);
+			TempName.RemoveFromStart(TEXT("SK_"), ESearchCase::CaseSensitive);
+			if (!TempName.IsEmpty())
 			{
 				AssetName = TempName;
 			}
 
-			AssetPrefix = TEXT( "ABP_" );
+			AssetPrefix = TEXT("ABP_");
 		}
 
-		if ( !AssetName.StartsWith( AssetPrefix ) )
+		if (!AssetName.StartsWith(AssetPrefix))
 		{
 			AssetName = AssetPrefix + AssetName;
 		}
@@ -637,11 +653,7 @@ namespace UsdStageImporterImpl
 		}
 	}
 
-	void UpdateAssetUserData(
-		const TSet<UObject*>& UsedAssetsAndDependencies,
-		const FString& MainFilePath,
-		UUsdStageImportOptions* ImportOptions
-	)
+	void UpdateAssetUserData(const TSet<UObject*>& UsedAssetsAndDependencies, const FString& MainFilePath, UUsdStageImportOptions* ImportOptions)
 	{
 		for (UObject* Asset : UsedAssetsAndDependencies)
 		{
@@ -667,39 +679,39 @@ namespace UsdStageImporterImpl
 
 		EReplaceAssetPolicy ReplacePolicy = ImportContext.ImportOptions->ExistingAssetPolicy;
 		FString TargetPackagePath = UPackageTools::SanitizePackageName(DestFullPackagePath);
-		FString TargetAssetName = FPaths::GetBaseFilename( TargetPackagePath );
+		FString TargetAssetName = FPaths::GetBaseFilename(TargetPackagePath);
 		UObject* ExistingAsset = nullptr;
 		UPackage* ExistingPackage = nullptr;
 
-		if ( ReplacePolicy == EReplaceAssetPolicy::Append )
+		if (ReplacePolicy == EReplaceAssetPolicy::Append)
 		{
-			FAssetToolsModule& AssetToolsModule = FModuleManager::LoadModuleChecked<FAssetToolsModule>( "AssetTools" );
-			AssetToolsModule.Get().CreateUniqueAssetName( TargetPackagePath, TEXT(""), TargetPackagePath, TargetAssetName );
+			FAssetToolsModule& AssetToolsModule = FModuleManager::LoadModuleChecked<FAssetToolsModule>("AssetTools");
+			AssetToolsModule.Get().CreateUniqueAssetName(TargetPackagePath, TEXT(""), TargetPackagePath, TargetAssetName);
 		}
 		else
 		{
 			// See if we have an existing asset/package
-			ExistingPackage = FindPackage( nullptr, *TargetPackagePath );
-			if ( !ExistingPackage && FPackageName::DoesPackageExist( TargetPackagePath ) )
+			ExistingPackage = FindPackage(nullptr, *TargetPackagePath);
+			if (!ExistingPackage && FPackageName::DoesPackageExist(TargetPackagePath))
 			{
-				ExistingPackage = LoadPackage( nullptr, *TargetPackagePath, LOAD_None );
+				ExistingPackage = LoadPackage(nullptr, *TargetPackagePath, LOAD_None);
 			}
-			if ( ExistingPackage )
+			if (ExistingPackage)
 			{
-				FSoftObjectPath ObjectPath( TargetPackagePath );
-				ExistingAsset = static_cast< UObject* >( FindObjectWithOuter( ExistingPackage, Asset->GetClass() ) );
-				if ( !ExistingAsset )
+				FSoftObjectPath ObjectPath(TargetPackagePath);
+				ExistingAsset = static_cast<UObject*>(FindObjectWithOuter(ExistingPackage, Asset->GetClass()));
+				if (!ExistingAsset)
 				{
 					ExistingAsset = ObjectPath.TryLoad();
 				}
 			}
 
 			// If we're ignoring assets that conflict, just abort now
-			if ( ExistingAsset != nullptr && ExistingAsset != Asset && ReplacePolicy == EReplaceAssetPolicy::Ignore )
+			if (ExistingAsset != nullptr && ExistingAsset != Asset && ReplacePolicy == EReplaceAssetPolicy::Ignore)
 			{
 				// Redirect any users of our new transient asset to the old, existing asset
-				ObjectsToRemap.Add( Asset, ExistingAsset );
-				SoftObjectsToRemap.Add( Asset, ExistingAsset );
+				ObjectsToRemap.Add(Asset, ExistingAsset);
+				SoftObjectsToRemap.Add(Asset, ExistingAsset);
 				return nullptr;
 			}
 		}
@@ -716,7 +728,14 @@ namespace UsdStageImporterImpl
 		UPackage* Package = ExistingPackage ? ExistingPackage : CreatePackage(*TargetPackagePath);
 		if (!Package)
 		{
-			FUsdLogManager::LogMessage( EMessageSeverity::Error, FText::Format( LOCTEXT( "PublishFailure", "Failed to get destination package at '{0}' for imported asset '{1}'!" ), FText::FromString( TargetPackagePath ), FText::FromName( Asset->GetFName() ) ) );
+			FUsdLogManager::LogMessage(
+				EMessageSeverity::Error,
+				FText::Format(
+					LOCTEXT("PublishFailure", "Failed to get destination package at '{0}' for imported asset '{1}'!"),
+					FText::FromString(TargetPackagePath),
+					FText::FromName(Asset->GetFName())
+				)
+			);
 			return nullptr;
 		}
 		Package->FullyLoad();
@@ -764,16 +783,9 @@ namespace UsdStageImporterImpl
 			// That can lead to issues if the target names conflict with the names of ExistingAsset's subobjects
 			if (ExistingAsset)
 			{
-				FName UniqueName = MakeUniqueObjectName(
-					GetTransientPackage(),
-					ExistingAsset->GetClass(),
-					ExistingAsset->GetFName()
-				);
-				ExistingAsset->Rename(
-					*UniqueName.ToString(),
-					GetTransientPackage(),
-					REN_DontCreateRedirectors | REN_NonTransactional | REN_DoNotDirty
-				);
+				FName UniqueName = MakeUniqueObjectName(GetTransientPackage(), ExistingAsset->GetClass(), ExistingAsset->GetFName());
+				ExistingAsset
+					->Rename(*UniqueName.ToString(), GetTransientPackage(), REN_DontCreateRedirectors | REN_NonTransactional | REN_DoNotDirty);
 			}
 
 			if (bForceDuplicate)
@@ -812,20 +824,20 @@ namespace UsdStageImporterImpl
 			MovedAsset = Asset;
 		}
 
-		SoftObjectsToRemap.Add( OldPath, MovedAsset );
+		SoftObjectsToRemap.Add(OldPath, MovedAsset);
 		if (MovedAsset != Asset)
 		{
 			ObjectsToRemap.Add(Asset, MovedAsset);
 		}
 
 		// Important as some assets (e.g. material instances) are created with no flags
-		MovedAsset->SetFlags(ImportContext.ImportObjectFlags | EObjectFlags::RF_Public | EObjectFlags::RF_Standalone );
+		MovedAsset->SetFlags(ImportContext.ImportObjectFlags | EObjectFlags::RF_Public | EObjectFlags::RF_Standalone);
 		MovedAsset->ClearFlags(EObjectFlags::RF_Transient | EObjectFlags::RF_DuplicateTransient | EObjectFlags::RF_NonPIEDuplicateTransient);
 
-		// Some subobjects like UStaticMesh::HiResSourceModel->StaticMeshDescriptionBulkData can't be left transient, or else they won't serialize their data.
-		// We probably never want to make them public or standalone if they aren't already though
+		// Some subobjects like UStaticMesh::HiResSourceModel->StaticMeshDescriptionBulkData can't be left transient, or else they won't serialize
+		// their data. We probably never want to make them public or standalone if they aren't already though
 		TArray<UObject*> Subobjects;
-		MovedAsset->GetDefaultSubobjects( Subobjects );
+		MovedAsset->GetDefaultSubobjects(Subobjects);
 		if (UMaterialInterface* Material = Cast<UMaterialInterface>(MovedAsset))
 		{
 			// Materials in particular have EditorOnlyData which behaves like a default subobject but kind of isn't flagged as one...
@@ -836,17 +848,18 @@ namespace UsdStageImporterImpl
 				Subobjects.Add(EditorOnlyData);
 			}
 		}
-		for ( UObject* Subobject : Subobjects )
+		for (UObject* Subobject : Subobjects)
 		{
-			Subobject->ClearFlags( EObjectFlags::RF_Transient | EObjectFlags::RF_DuplicateTransient | EObjectFlags::RF_NonPIEDuplicateTransient );
+			Subobject->ClearFlags(EObjectFlags::RF_Transient | EObjectFlags::RF_DuplicateTransient | EObjectFlags::RF_NonPIEDuplicateTransient);
 		}
 
 		// We need to make sure that "dirtying the final package" is not added to the transaction, because if we undo this transaction
-		// the assets should remain on their final destination, so we still want the packages to remain marked as dirty (as they're really not on the disk yet).
-		// If we didn't suppress, the package would become transactional by this call. When undoing, the assets would still remain on the final package,
-		// but the "dirtying" would be undone, so the engine would think the assets weren't dirty (i.e. were already saved), which is not true
+		// the assets should remain on their final destination, so we still want the packages to remain marked as dirty (as they're really not on the
+		// disk yet). If we didn't suppress, the package would become transactional by this call. When undoing, the assets would still remain on the
+		// final package, but the "dirtying" would be undone, so the engine would think the assets weren't dirty (i.e. were already saved), which is
+		// not true
 		{
-			TGuardValue< ITransaction* > SuppressTransaction{ GUndo, nullptr };
+			TGuardValue<ITransaction*> SuppressTransaction{GUndo, nullptr};
 			Package->MarkPackageDirty();
 		}
 
@@ -921,13 +934,13 @@ namespace UsdStageImporterImpl
 				continue;
 			}
 
-			const static FString TexturesFolder       = TEXT("Textures");
-			const static FString MaterialsFolder      = TEXT("Materials");
-			const static FString StaticMeshesFolder   = TEXT("StaticMeshes");
+			const static FString TexturesFolder = TEXT("Textures");
+			const static FString MaterialsFolder = TEXT("Materials");
+			const static FString StaticMeshesFolder = TEXT("StaticMeshes");
 			const static FString SkeletalMeshesFolder = TEXT("SkeletalMeshes");
 			const static FString LevelSequencesFolder = TEXT("LevelSequences");
 			const static FString GeometryCachesFolder = TEXT("GeometryCaches");
-			const static FString GroomsFolder         = TEXT("Grooms");
+			const static FString GroomsFolder = TEXT("Grooms");
 
 			const FString* AssetTypeFolderPtr = nullptr;
 			if (Asset->IsA(UMaterialInterface::StaticClass()))
@@ -1042,14 +1055,8 @@ namespace UsdStageImporterImpl
 				{
 					FString TargetAssetName = GetUserFriendlyName(Asset, UniqueAssetNames);
 					FString DestPackagePath = FPaths::Combine(ImportContext.PackagePath, *AssetTypeFolder, TargetAssetName);
-					UObject* PublishedAsset = PublishAsset(
-						ImportContext,
-						Asset,
-						DestPackagePath,
-						ObjectsToRemap,
-						SoftObjectsToRemap,
-						OutAssetsToFinalize
-					);
+					UObject*
+						PublishedAsset = PublishAsset(ImportContext, Asset, DestPackagePath, ObjectsToRemap, SoftObjectsToRemap, OutAssetsToFinalize);
 
 					if (PublishedAsset)
 					{
@@ -1077,7 +1084,13 @@ namespace UsdStageImporterImpl
 		PublishAssetType(Textures);
 	}
 
-	void ResolveComponentConflict(USceneComponent* NewRoot, USceneComponent* ExistingRoot, EReplaceActorPolicy ReplacePolicy, TMap<UObject*, UObject*>& ObjectsToRemap, TMap<FSoftObjectPath, FSoftObjectPath>& SoftObjectsToRemap )
+	void ResolveComponentConflict(
+		USceneComponent* NewRoot,
+		USceneComponent* ExistingRoot,
+		EReplaceActorPolicy ReplacePolicy,
+		TMap<UObject*, UObject*>& ObjectsToRemap,
+		TMap<FSoftObjectPath, FSoftObjectPath>& SoftObjectsToRemap
+	)
 	{
 		if (!NewRoot || !ExistingRoot || ReplacePolicy == EReplaceActorPolicy::Append)
 		{
@@ -1121,22 +1134,22 @@ namespace UsdStageImporterImpl
 
 				switch (ReplacePolicy)
 				{
-				case EReplaceActorPolicy::UpdateTransform:
-					(*FoundExistingComponent)->SetRelativeTransform(NewComponent->GetRelativeTransform());
-					(*FoundExistingComponent)->AttachToComponent(NewRoot, FAttachmentTransformRules::KeepRelativeTransform);
-					bRecurse = true;
-					break;
-				case EReplaceActorPolicy::Ignore:
-					// Note how we're iterating the new hierarchy here, so "ignore" means "keep the existing one"
-					NewComponent->DestroyComponent(false);
-					(*FoundExistingComponent)->AttachToComponent(NewRoot, FAttachmentTransformRules::KeepRelativeTransform);
-					bRecurse = false;
-					break;
-				case EReplaceActorPolicy::Replace:
-				default:
-					// Keep NewChild completely, but recurse to replace components and children
-					bRecurse = true;
-					break;
+					case EReplaceActorPolicy::UpdateTransform:
+						(*FoundExistingComponent)->SetRelativeTransform(NewComponent->GetRelativeTransform());
+						(*FoundExistingComponent)->AttachToComponent(NewRoot, FAttachmentTransformRules::KeepRelativeTransform);
+						bRecurse = true;
+						break;
+					case EReplaceActorPolicy::Ignore:
+						// Note how we're iterating the new hierarchy here, so "ignore" means "keep the existing one"
+						NewComponent->DestroyComponent(false);
+						(*FoundExistingComponent)->AttachToComponent(NewRoot, FAttachmentTransformRules::KeepRelativeTransform);
+						bRecurse = false;
+						break;
+					case EReplaceActorPolicy::Replace:
+					default:
+						// Keep NewChild completely, but recurse to replace components and children
+						bRecurse = true;
+						break;
 				}
 
 				if (bRecurse)
@@ -1180,7 +1193,13 @@ namespace UsdStageImporterImpl
 		Actor->GetWorld()->DestroyActor(Actor);
 	}
 
-	void ResolveActorConflict(AActor* NewActor, AActor* ExistingActor, EReplaceActorPolicy ReplacePolicy, TMap<UObject*, UObject*>& ObjectsToRemap, TMap<FSoftObjectPath, FSoftObjectPath>& SoftObjectsToRemap)
+	void ResolveActorConflict(
+		AActor* NewActor,
+		AActor* ExistingActor,
+		EReplaceActorPolicy ReplacePolicy,
+		TMap<UObject*, UObject*>& ObjectsToRemap,
+		TMap<FSoftObjectPath, FSoftObjectPath>& SoftObjectsToRemap
+	)
 	{
 		if (!NewActor || !ExistingActor || ReplacePolicy == EReplaceActorPolicy::Append)
 		{
@@ -1188,7 +1207,7 @@ namespace UsdStageImporterImpl
 		}
 
 		ObjectsToRemap.Add(ExistingActor, NewActor);
-		SoftObjectsToRemap.Add( ExistingActor, NewActor );
+		SoftObjectsToRemap.Add(ExistingActor, NewActor);
 
 		// Collect new and existing actors by label
 		const bool bResetArray = false;
@@ -1221,22 +1240,22 @@ namespace UsdStageImporterImpl
 
 				switch (ReplacePolicy)
 				{
-				case EReplaceActorPolicy::UpdateTransform:
-					(*ExistingChild)->GetRootComponent()->SetRelativeTransform(NewChild->GetRootComponent()->GetRelativeTransform());
-					GEditor->ParentActors( NewActor, *ExistingChild, NAME_None );
-					bRecurse = true;
-					break;
-				case EReplaceActorPolicy::Ignore:
-					// Note how we're iterating the new hierarchy here, so "ignore" means "keep the existing one"
-					RecursiveDestroyActor(NewChild);
-					GEditor->ParentActors(NewActor, *ExistingChild, NAME_None);
-					bRecurse = false;
-					break;
-				case EReplaceActorPolicy::Replace:
-				default:
-					// Keep NewChild, but recurse to replace components and children
-					bRecurse = true;
-					break;
+					case EReplaceActorPolicy::UpdateTransform:
+						(*ExistingChild)->GetRootComponent()->SetRelativeTransform(NewChild->GetRootComponent()->GetRelativeTransform());
+						GEditor->ParentActors(NewActor, *ExistingChild, NAME_None);
+						bRecurse = true;
+						break;
+					case EReplaceActorPolicy::Ignore:
+						// Note how we're iterating the new hierarchy here, so "ignore" means "keep the existing one"
+						RecursiveDestroyActor(NewChild);
+						GEditor->ParentActors(NewActor, *ExistingChild, NAME_None);
+						bRecurse = false;
+						break;
+					case EReplaceActorPolicy::Replace:
+					default:
+						// Keep NewChild, but recurse to replace components and children
+						bRecurse = true;
+						break;
 				}
 
 				if (bRecurse)
@@ -1266,16 +1285,24 @@ namespace UsdStageImporterImpl
 		}
 	}
 
-	void ResolveActorConflicts(FUsdStageImportContext& ImportContext, AActor* ExistingSceneActor, TMap<UObject*, UObject*>& ObjectsToRemap, TMap<FSoftObjectPath, FSoftObjectPath>& SoftObjectsToRemap)
+	void ResolveActorConflicts(
+		FUsdStageImportContext& ImportContext,
+		AActor* ExistingSceneActor,
+		TMap<UObject*, UObject*>& ObjectsToRemap,
+		TMap<FSoftObjectPath, FSoftObjectPath>& SoftObjectsToRemap
+	)
 	{
-		if ( !ImportContext.ImportOptions->bImportActors )
+		if (!ImportContext.ImportOptions->bImportActors)
 		{
 			return;
 		}
 
 		if (!ImportContext.SceneActor)
 		{
-			FUsdLogManager::LogMessage( EMessageSeverity::Error, LOCTEXT( "NoSceneActor", "Failed to publish actors as there was no scene actor available!" ) );
+			FUsdLogManager::LogMessage(
+				EMessageSeverity::Error,
+				LOCTEXT("NoSceneActor", "Failed to publish actors as there was no scene actor available!")
+			);
 			return;
 		}
 
@@ -1328,41 +1355,45 @@ namespace UsdStageImporterImpl
 		}
 	}
 
-	void CopySkeletonAssignment( FUsdStageImportContext& ImportContext, UObject* ExistingAsset, UObject* NewAsset )
+	void CopySkeletonAssignment(FUsdStageImportContext& ImportContext, UObject* ExistingAsset, UObject* NewAsset)
 	{
-		USkeletalMesh* ExistingSkeletalMesh = Cast<USkeletalMesh>( ExistingAsset );
-		USkeletalMesh* NewSkeletalMesh = Cast<USkeletalMesh>( NewAsset );
-		if ( ExistingSkeletalMesh && NewSkeletalMesh )
+		USkeletalMesh* ExistingSkeletalMesh = Cast<USkeletalMesh>(ExistingAsset);
+		USkeletalMesh* NewSkeletalMesh = Cast<USkeletalMesh>(NewAsset);
+		if (ExistingSkeletalMesh && NewSkeletalMesh)
 		{
 			// Never assign a transient skeleton
-			if ( ExistingSkeletalMesh->GetSkeleton() && ExistingSkeletalMesh->GetSkeleton()->GetOutermost() == GetTransientPackage() )
+			if (ExistingSkeletalMesh->GetSkeleton() && ExistingSkeletalMesh->GetSkeleton()->GetOutermost() == GetTransientPackage())
 			{
 				return;
 			}
 
 			// Assign even if ExistingSkeletalMesh has nullptr skeleton because we must be able to cleanup the
 			// abandoned Skeleton in the transient package
-			NewSkeletalMesh->SetSkeleton( ExistingSkeletalMesh->GetSkeleton() );
-			NewSkeletalMesh->SetPhysicsAsset( ExistingSkeletalMesh->GetPhysicsAsset() );
+			NewSkeletalMesh->SetSkeleton(ExistingSkeletalMesh->GetSkeleton());
+			NewSkeletalMesh->SetPhysicsAsset(ExistingSkeletalMesh->GetPhysicsAsset());
 		}
 
-		UAnimSequence* ExistingAnimSequence = Cast<UAnimSequence>( ExistingAsset );
-		UAnimSequence* NewAnimSequence = Cast<UAnimSequence>( NewAsset );
-		if ( ExistingAnimSequence && NewAnimSequence )
+		UAnimSequence* ExistingAnimSequence = Cast<UAnimSequence>(ExistingAsset);
+		UAnimSequence* NewAnimSequence = Cast<UAnimSequence>(NewAsset);
+		if (ExistingAnimSequence && NewAnimSequence)
 		{
 			// Never assign a transient skeleton
 			USkeleton* ExistingSkeleton = ExistingAnimSequence->GetSkeleton();
-			if ( ExistingSkeleton && ExistingSkeleton->GetOutermost() == GetTransientPackage() )
+			if (ExistingSkeleton && ExistingSkeleton->GetOutermost() == GetTransientPackage())
 			{
 				return;
 			}
 
-			NewAnimSequence->SetSkeleton( ExistingSkeleton );
+			NewAnimSequence->SetSkeleton(ExistingSkeleton);
 		}
 	}
 
 	// Adapted from FDatasmithImporterImpl::FixReferencesForObject
-	void RemapReferences(FUsdStageImportContext& ImportContext, const TSet<UObject*>& PublishedObjects, const TMap< UObject*, UObject* >& ObjectsToRemap)
+	void RemapReferences(
+		FUsdStageImportContext& ImportContext,
+		const TSet<UObject*>& PublishedObjects,
+		const TMap<UObject*, UObject*>& ObjectsToRemap
+	)
 	{
 		if (ObjectsToRemap.Num() == 0)
 		{
@@ -1382,28 +1413,29 @@ namespace UsdStageImporterImpl
 
 		// Remap references held by assets that were moved directly to the destination package, and won't be in ProcessedObjectsToRemap
 		TSet<UObject*> Referencers = PublishedObjects;
-		if ( AActor* SceneActor = ImportContext.SceneActor )
+		if (AActor* SceneActor = ImportContext.SceneActor)
 		{
 			// Remap references to spawned actors
-			Referencers.Add( ImportContext.SceneActor->GetWorld()->GetCurrentLevel() );
+			Referencers.Add(ImportContext.SceneActor->GetWorld()->GetCurrentLevel());
 		}
-		for ( const TPair<UObject*, UObject*>& Pair : ProcessedObjectsToRemap )
+		for (const TPair<UObject*, UObject*>& Pair : ProcessedObjectsToRemap)
 		{
 			// Remap internal references between the remapped objects
-			Referencers.Add( Pair.Value );
+			Referencers.Add(Pair.Value);
 		}
 
 		// Fix references between actors and assets (e.g. mesh in final package referencing material in transient package)
 		// Note we don't care if transient assets reference each other, as we'll delete them all at once anyway
-		for ( UObject* Referencer : Referencers )
+		for (UObject* Referencer : Referencers)
 		{
-			if ( !Referencer || Referencer->GetOutermost() == GetTransientPackage() )
+			if (!Referencer || Referencer->GetOutermost() == GetTransientPackage())
 			{
 				continue;
 			}
 
-			constexpr EArchiveReplaceObjectFlags ReplaceFlags = (EArchiveReplaceObjectFlags::IgnoreOuterRef | EArchiveReplaceObjectFlags::IgnoreArchetypeRef);
-			FArchiveReplaceObjectRef< UObject > ArchiveReplaceObjectRefInner(Referencer, ProcessedObjectsToRemap, ReplaceFlags);
+			constexpr EArchiveReplaceObjectFlags
+				ReplaceFlags = (EArchiveReplaceObjectFlags::IgnoreOuterRef | EArchiveReplaceObjectFlags::IgnoreArchetypeRef);
+			FArchiveReplaceObjectRef<UObject> ArchiveReplaceObjectRefInner(Referencer, ProcessedObjectsToRemap, ReplaceFlags);
 		}
 	}
 
@@ -1424,14 +1456,9 @@ namespace UsdStageImporterImpl
 		}
 	}
 
-	void Cleanup(
-		AActor* NewSceneActor,
-		AActor* ExistingSceneActor,
-		EReplaceActorPolicy ReplacePolicy,
-		const TSet<UObject*>& AssetsToFinalize
-	)
+	void Cleanup(AActor* NewSceneActor, AActor* ExistingSceneActor, EReplaceActorPolicy ReplacePolicy, const TSet<UObject*>& AssetsToFinalize)
 	{
-		if ( !NewSceneActor )
+		if (!NewSceneActor)
 		{
 			return;
 		}
@@ -1445,9 +1472,9 @@ namespace UsdStageImporterImpl
 
 		for (TObjectIterator<UMeshComponent> It; It; ++It)
 		{
-			if(UStaticMeshComponent* StaticMeshComponent = Cast<UStaticMeshComponent>(*It))
+			if (UStaticMeshComponent* StaticMeshComponent = Cast<UStaticMeshComponent>(*It))
 			{
-				if(AssetsToFinalize.Contains(StaticMeshComponent->GetStaticMesh()))
+				if (AssetsToFinalize.Contains(StaticMeshComponent->GetStaticMesh()))
 				{
 					StaticMeshComponent->FixupOverrideColorsIfNecessary(true);
 					StaticMeshComponent->InvalidateLightingCache();
@@ -1455,7 +1482,7 @@ namespace UsdStageImporterImpl
 			}
 			else if (USkeletalMeshComponent* SkeletalMeshComponent = Cast<USkeletalMeshComponent>(*It))
 			{
-				if(AssetsToFinalize.Contains(SkeletalMeshComponent->GetSkeletalMeshAsset()))
+				if (AssetsToFinalize.Contains(SkeletalMeshComponent->GetSkeletalMeshAsset()))
 				{
 					SkeletalMeshComponent->ClearAnimScriptInstance();
 				}
@@ -1472,62 +1499,63 @@ namespace UsdStageImporterImpl
 			UnrealUSDWrapper::EraseStageFromCache(ImportContext.Stage);
 		}
 
-		if ( ImportContext.ImportOptions->bOverrideStageOptions )
+		if (ImportContext.ImportOptions->bOverrideStageOptions)
 		{
-			UsdUtils::SetUsdStageMetersPerUnit( ImportContext.Stage, ImportContext.OriginalMetersPerUnit );
-			UsdUtils::SetUsdStageUpAxis( ImportContext.Stage, ImportContext.OriginalUpAxis );
+			UsdUtils::SetUsdStageMetersPerUnit(ImportContext.Stage, ImportContext.OriginalMetersPerUnit);
+			UsdUtils::SetUsdStageUpAxis(ImportContext.Stage, ImportContext.OriginalUpAxis);
 		}
 
 		// Always discard the context's reference to the stage because it may be a persistent import context (like
 		// the non-static data member of UUsdStageImportFactory
 		ImportContext.Stage = UE::FUsdStage();
-		ImportContext.LevelSequenceHelper.Init( UE::FUsdStage() );
-#endif // #if USE_USD_SDK
+		ImportContext.LevelSequenceHelper.Init(UE::FUsdStage());
+#endif	  // #if USE_USD_SDK
 	}
 
 	/**
 	 * UUsdAssetCache2 can track which assets are requested/added to itself during translation, but it may miss some dependencies
 	 * that are only retrieved/added themselves when the original asset is first parsed. This function recursively collects all of those.
-	 * Example: An UMaterialInstance is already in the cache, so when translating we just retrieve the existing asset --> The textures that it's using won't be retrieved or marked as "Used"
-	 * Example: An USkeletalMesh is already in the cache, so in the same way we would miss its USkeleton, materials and textures of those materials
+	 * Example: An UMaterialInstance is already in the cache, so when translating we just retrieve the existing asset --> The textures that it's using
+	 * won't be retrieved or marked as "Used" Example: An USkeletalMesh is already in the cache, so in the same way we would miss its USkeleton,
+	 * materials and textures of those materials
 	 */
-	void CollectUsedAssetDependencies( FUsdStageImportContext& ImportContext, TSet<UObject*>& OutAssetsAndDependencies )
+	void CollectUsedAssetDependencies(FUsdStageImportContext& ImportContext, TSet<UObject*>& OutAssetsAndDependencies)
 	{
-		const int32 ReserveSize = OutAssetsAndDependencies.Num() + ( ImportContext.AssetCache ? ImportContext.AssetCache->GetActiveAssets().Num() : 0 );
+		const int32 ReserveSize = OutAssetsAndDependencies.Num() + (ImportContext.AssetCache ? ImportContext.AssetCache->GetActiveAssets().Num() : 0);
 
 		// We will only emit the level sequences if we have data in the main one.
 		// Keep subsequences even if they have no data as the main sequence/other sequences may reference them
-		if ( ImportContext.ImportOptions->bImportLevelSequences && ImportContext.LevelSequenceHelper.HasData() )
+		if (ImportContext.ImportOptions->bImportLevelSequences && ImportContext.LevelSequenceHelper.HasData())
 		{
 			TArray<ULevelSequence*> SubSequences = ImportContext.LevelSequenceHelper.GetSubSequences();
 			ULevelSequence* MainSequence = ImportContext.LevelSequenceHelper.GetMainLevelSequence();
 
-			OutAssetsAndDependencies.Reserve( ReserveSize + SubSequences.Num() + 1 );
-			OutAssetsAndDependencies.Add( MainSequence );
-			for ( ULevelSequence* SubSequence : SubSequences )
+			OutAssetsAndDependencies.Reserve(ReserveSize + SubSequences.Num() + 1);
+			OutAssetsAndDependencies.Add(MainSequence);
+			for (ULevelSequence* SubSequence : SubSequences)
 			{
-				OutAssetsAndDependencies.Add( SubSequence );
+				OutAssetsAndDependencies.Add(SubSequence);
 			}
 		}
 		else
 		{
-			OutAssetsAndDependencies.Reserve( ReserveSize );
+			OutAssetsAndDependencies.Reserve(ReserveSize);
 		}
 
-		if ( ImportContext.AssetCache )
+		if (ImportContext.AssetCache)
 		{
 			const TSet<UObject*>& InPrimaryAssets = ImportContext.AssetCache->GetActiveAssets();
 			TArray<UObject*> AssetQueue = InPrimaryAssets.Array();
 
-			for ( int32 AssetIndex = 0; AssetIndex < AssetQueue.Num(); ++AssetIndex )
+			for (int32 AssetIndex = 0; AssetIndex < AssetQueue.Num(); ++AssetIndex)
 			{
-				UObject* Asset = AssetQueue[ AssetIndex ];
+				UObject* Asset = AssetQueue[AssetIndex];
 
 				// Only add it as a dependency if it's owned by the asset cache, but still traverse it because
 				// we may be in some strange situation where the material shouldn't be in this list, but one of its used textures should
-				if ( Asset && ImportContext.AssetCache->IsAssetOwnedByCache( Asset->GetPathName() ) )
+				if (Asset && ImportContext.AssetCache->IsAssetOwnedByCache(Asset->GetPathName()))
 				{
-					OutAssetsAndDependencies.Add( Asset );
+					OutAssetsAndDependencies.Add(Asset);
 				}
 
 				TSet<UObject*> Dependencies = IUsdClassesModule::GetAssetDependencies(Asset);
@@ -1540,26 +1568,30 @@ namespace UsdStageImporterImpl
 	 * Remaps asset's soft object pointers to point to the post-publish paths of their target assets.
 	 * It's important to run this *after* RemapReferences, as we will sometimes rely on those references to find our target assets.
 	 */
-	void RemapSoftReferences( const FUsdStageImportContext& ImportContext, const TSet<UObject*>& UsedAssetsAndDependencies, const TMap<FSoftObjectPath, FSoftObjectPath>& SoftObjectsToRemap )
+	void RemapSoftReferences(
+		const FUsdStageImportContext& ImportContext,
+		const TSet<UObject*>& UsedAssetsAndDependencies,
+		const TMap<FSoftObjectPath, FSoftObjectPath>& SoftObjectsToRemap
+	)
 	{
 		TSet<UPackage*> Packages;
-		for ( UObject* Object : UsedAssetsAndDependencies )
+		for (UObject* Object : UsedAssetsAndDependencies)
 		{
-			if ( Object )
+			if (Object)
 			{
-				Packages.Add( Object->GetOutermost() );
+				Packages.Add(Object->GetOutermost());
 			}
 		}
 
-		if ( AActor* SceneActor = ImportContext.SceneActor )
+		if (AActor* SceneActor = ImportContext.SceneActor)
 		{
-			Packages.Add( ImportContext.SceneActor->GetWorld()->GetOutermost() );
+			Packages.Add(ImportContext.SceneActor->GetWorld()->GetOutermost());
 		}
 
 		// In case one our used assets was left on the transient package.
 		// We don't care about anything that was left on the transient package, and doing this may actually cause some reference counting issues
 		// if we try deleting those assets afterwards
-		Packages.Remove( GetTransientPackage() );
+		Packages.Remove(GetTransientPackage());
 
 		// We never want to remap all invalid reference to something. That particularly seems to break LevelSequence bindings somehow
 		TMap<FSoftObjectPath, FSoftObjectPath> ProcessedPaths;
@@ -1572,13 +1604,13 @@ namespace UsdStageImporterImpl
 			}
 		}
 
-		IAssetTools& AssetTools = FModuleManager::LoadModuleChecked<FAssetToolsModule>( "AssetTools" ).Get();
-		AssetTools.RenameReferencingSoftObjectPaths( Packages.Array(), ProcessedPaths );
+		IAssetTools& AssetTools = FModuleManager::LoadModuleChecked<FAssetToolsModule>("AssetTools").Get();
+		AssetTools.RenameReferencingSoftObjectPaths(Packages.Array(), ProcessedPaths);
 	}
 
 	void GetPublishedAssetsAndDependencies(
 		const TSet<UObject*>& UsedAssetsAndDependencies,
-		const TMap< UObject*, UObject* >& ObjectsToRemap,
+		const TMap<UObject*, UObject*>& ObjectsToRemap,
 		const TMap<FSoftObjectPath, FSoftObjectPath>& SoftObjectsToRemap,
 		TSet<UObject*>& OutPublishedAssetsAndDependencies
 	)
@@ -1614,14 +1646,14 @@ namespace UsdStageImporterImpl
 	}
 
 	/** After we remapped everything, notify the AssetRegistry that we created some new assets */
-	void NotifyAssetRegistry( const TSet<UObject*>& UsedAssetsAndDependencies )
+	void NotifyAssetRegistry(const TSet<UObject*>& UsedAssetsAndDependencies)
 	{
-		for ( UObject* Object : UsedAssetsAndDependencies )
+		for (UObject* Object : UsedAssetsAndDependencies)
 		{
 			// If it's still on the transient package it means we abandoned this one (maybe we had asset replace policy ignore and hit a conflict)
-			if ( Object && Object->GetOutermost() != GetTransientPackage() )
+			if (Object && Object->GetOutermost() != GetTransientPackage())
 			{
-				FAssetRegistryModule::AssetCreated( Object );
+				FAssetRegistryModule::AssetCreated(Object);
 			}
 		}
 	}
@@ -1635,23 +1667,23 @@ namespace UsdStageImporterImpl
 	)
 	{
 #if USE_USD_SDK
-		if ( FEngineAnalytics::IsAvailable() )
+		if (FEngineAnalytics::IsAvailable())
 		{
 			TArray<FAnalyticsEventAttribute> EventAttributes;
 
 			FString EventName = Operation;
-			if ( Asset )
+			if (Asset)
 			{
 				FString ClassName = Asset->GetClass()->GetName();
 
 				// e.g. "Reimport.StaticMesh"
-				EventName = FString::Printf( TEXT( "%s.%s" ), *EventName, *ClassName );
-				EventAttributes.Emplace( TEXT( "AssetType" ), ClassName );
+				EventName = FString::Printf(TEXT("%s.%s"), *EventName, *ClassName);
+				EventAttributes.Emplace(TEXT("AssetType"), ClassName);
 			}
 
-			if ( ImportContext.ImportOptions )
+			if (ImportContext.ImportOptions)
 			{
-				UsdUtils::AddAnalyticsAttributes( *ImportContext.ImportOptions, EventAttributes );
+				UsdUtils::AddAnalyticsAttributes(*ImportContext.ImportOptions, EventAttributes);
 			}
 
 			int32 NumStaticMeshes = 0;
@@ -1664,121 +1696,121 @@ namespace UsdStageImporterImpl
 			int32 NumGroomAssets = 0;
 			int32 NumGroomBindings = 0;
 			int32 NumGroomCaches = 0;
-			for ( UObject* ImportedAsset : ImportedAssets )
+			for (UObject* ImportedAsset : ImportedAssets)
 			{
-				if ( !ImportedAsset )
+				if (!ImportedAsset)
 				{
 					continue;
 				}
 
-				if ( ImportedAsset->IsA<UStaticMesh>() )
+				if (ImportedAsset->IsA<UStaticMesh>())
 				{
 					++NumStaticMeshes;
 				}
-				else if ( ImportedAsset->IsA<USkeletalMesh>() )
+				else if (ImportedAsset->IsA<USkeletalMesh>())
 				{
 					++NumSkeletalMeshes;
 				}
-				else if ( ImportedAsset->IsA<UMaterialInterface>() )
+				else if (ImportedAsset->IsA<UMaterialInterface>())
 				{
 					++NumMaterials;
 				}
-				else if ( ImportedAsset->IsA<UAnimSequence>() )
+				else if (ImportedAsset->IsA<UAnimSequence>())
 				{
 					++NumAnimSequences;
 				}
-				else if ( ImportedAsset->IsA<ULevelSequence>() )
+				else if (ImportedAsset->IsA<ULevelSequence>())
 				{
 					++NumLevelSequences;
 				}
-				else if ( ImportedAsset->IsA<UTexture>() )
+				else if (ImportedAsset->IsA<UTexture>())
 				{
 					++NumTextures;
 				}
-				else if ( ImportedAsset->IsA<UGeometryCache>() )
+				else if (ImportedAsset->IsA<UGeometryCache>())
 				{
 					++NumGeometryCaches;
 				}
-				else if ( ImportedAsset->IsA<UGroomAsset>() )
+				else if (ImportedAsset->IsA<UGroomAsset>())
 				{
 					++NumGroomAssets;
 				}
-				else if ( ImportedAsset->IsA<UGroomBindingAsset>() )
+				else if (ImportedAsset->IsA<UGroomBindingAsset>())
 				{
 					++NumGroomBindings;
 				}
-				else if ( ImportedAsset->IsA<UGroomCache>() )
+				else if (ImportedAsset->IsA<UGroomCache>())
 				{
 					++NumGroomCaches;
 				}
 			}
-			EventAttributes.Emplace( TEXT( "NumStaticMeshes" ), LexToString( NumStaticMeshes ) );
-			EventAttributes.Emplace( TEXT( "NumSkeletalMeshes" ), LexToString( NumSkeletalMeshes ) );
-			EventAttributes.Emplace( TEXT( "NumMaterials" ), LexToString( NumMaterials ) );
-			EventAttributes.Emplace( TEXT( "NumAnimSequences" ), LexToString( NumAnimSequences ) );
-			EventAttributes.Emplace( TEXT( "NumLevelSequences" ), LexToString( NumLevelSequences ) );
-			EventAttributes.Emplace( TEXT( "NumTextures" ), LexToString( NumTextures ) );
-			EventAttributes.Emplace( TEXT( "NumGeometryCaches" ), LexToString( NumGeometryCaches ) );
-			EventAttributes.Emplace( TEXT( "NumGroomAssets" ), LexToString( NumGroomAssets ) );
-			EventAttributes.Emplace( TEXT( "NumGroomBindings" ), LexToString( NumGroomBindings ) );
-			EventAttributes.Emplace( TEXT( "NumGroomCaches" ), LexToString( NumGroomCaches ) );
+			EventAttributes.Emplace(TEXT("NumStaticMeshes"), LexToString(NumStaticMeshes));
+			EventAttributes.Emplace(TEXT("NumSkeletalMeshes"), LexToString(NumSkeletalMeshes));
+			EventAttributes.Emplace(TEXT("NumMaterials"), LexToString(NumMaterials));
+			EventAttributes.Emplace(TEXT("NumAnimSequences"), LexToString(NumAnimSequences));
+			EventAttributes.Emplace(TEXT("NumLevelSequences"), LexToString(NumLevelSequences));
+			EventAttributes.Emplace(TEXT("NumTextures"), LexToString(NumTextures));
+			EventAttributes.Emplace(TEXT("NumGeometryCaches"), LexToString(NumGeometryCaches));
+			EventAttributes.Emplace(TEXT("NumGroomAssets"), LexToString(NumGroomAssets));
+			EventAttributes.Emplace(TEXT("NumGroomBindings"), LexToString(NumGroomBindings));
+			EventAttributes.Emplace(TEXT("NumGroomCaches"), LexToString(NumGroomCaches));
 
 			FString RootLayerIdentifier = ImportContext.FilePath;
-			if ( ImportContext.Stage )
+			if (ImportContext.Stage)
 			{
-				if ( RootLayerIdentifier.IsEmpty() )
+				if (RootLayerIdentifier.IsEmpty())
 				{
 					RootLayerIdentifier = ImportContext.Stage.GetRootLayer().GetIdentifier();
 				}
 			}
 
 			IUsdClassesModule::SendAnalytics(
-				MoveTemp( EventAttributes ),
+				MoveTemp(EventAttributes),
 				EventName,
 				ImportContext.bIsAutomated,
 				ElapsedSeconds,
-				UsdUtils::GetUsdStageNumFrames( ImportContext.Stage ),
-				FPaths::GetExtension( RootLayerIdentifier )
+				UsdUtils::GetUsdStageNumFrames(ImportContext.Stage),
+				FPaths::GetExtension(RootLayerIdentifier)
 			);
 		}
-#endif // USE_USD_SDK
+#endif	  // USE_USD_SDK
 	}
 
 	// Removes from AssetsToImport assets that are unwanted according to our import options, and adds entries to
 	// ObjectsToRemap and SoftObjectsToRemap that remaps them to nullptr.
 	// This function is needed because it's not enough to e.g. just prevent new meshes from being imported from
-	// UsdStageImporterImpl::ImportMeshes, because we may want to reuse meshes we already got from the asset cache.
+	// UE::USDStageImporter::Private::ImportMeshes, because we may want to reuse meshes we already got from the asset cache.
 	// Additionally, we'll want to remap even our components away from pointing to these assets
-	void PruneUnwantedAssets( FUsdStageImportContext& ImportContext, TSet<UObject*>& AssetsToImport, TMap<UObject*, UObject*>& ObjectsToRemap, TMap<FSoftObjectPath, FSoftObjectPath>& SoftObjectsToRemap )
+	void PruneUnwantedAssets(
+		FUsdStageImportContext& ImportContext,
+		TSet<UObject*>& AssetsToImport,
+		TMap<UObject*, UObject*>& ObjectsToRemap,
+		TMap<FSoftObjectPath, FSoftObjectPath>& SoftObjectsToRemap
+	)
 	{
 		const bool bImportSkeletalAnimations = ImportContext.ImportOptions->bImportGeometry && ImportContext.ImportOptions->bImportSkeletalAnimations;
 
-		for ( TSet<UObject*>::TIterator It( AssetsToImport ); It; ++It )
+		for (TSet<UObject*>::TIterator It(AssetsToImport); It; ++It)
 		{
 			UObject* Asset = *It;
 
-			if ( !Asset )
+			if (!Asset)
 			{
 				It.RemoveCurrent();
 				continue;
 			}
 
-			if (
-				(!ImportContext.ImportOptions->bImportGeometry && (
-					Asset->IsA<UStaticMesh>() ||
-					Asset->IsA<USkeletalMesh>() ||
-					Asset->IsA<USkeleton>() ||
-					Asset->IsA<UPhysicsAsset>() ||
-					Asset->IsA<UGeometryCache>()
-				)) ||
-				(!bImportSkeletalAnimations && (Asset->IsA<UAnimSequence>())) ||
-				(!ImportContext.ImportOptions->bImportLevelSequences && (Asset->IsA<ULevelSequence>())) ||
-				(!ImportContext.ImportOptions->bImportMaterials && (Asset->IsA<UMaterialInterface>() || Asset->IsA<UTexture>())) ||
-				(!ImportContext.ImportOptions->bImportGroomAssets && (Asset->IsA<UGroomAsset>() || Asset->IsA<UGroomCache>() || Asset->IsA<UGroomBindingAsset>()))
-			)
-			{
-				ObjectsToRemap.Add( Asset, nullptr );
-				SoftObjectsToRemap.Add( Asset, nullptr );
+			if ((!ImportContext.ImportOptions->bImportGeometry
+				 && (Asset->IsA<UStaticMesh>() || Asset->IsA<USkeletalMesh>() || Asset->IsA<USkeleton>() || Asset->IsA<UPhysicsAsset>()
+					 || Asset->IsA<UGeometryCache>()))
+				|| (!bImportSkeletalAnimations && (Asset->IsA<UAnimSequence>()))
+				|| (!ImportContext.ImportOptions->bImportLevelSequences && (Asset->IsA<ULevelSequence>()))
+				|| (!ImportContext.ImportOptions->bImportMaterials && (Asset->IsA<UMaterialInterface>() || Asset->IsA<UTexture>()))
+				|| (!ImportContext.ImportOptions->bImportGroomAssets
+					&& (Asset->IsA<UGroomAsset>() || Asset->IsA<UGroomCache>() || Asset->IsA<UGroomBindingAsset>())))
+			{
+				ObjectsToRemap.Add(Asset, nullptr);
+				SoftObjectsToRemap.Add(Asset, nullptr);
 				It.RemoveCurrent();
 			}
 		}
@@ -1788,24 +1820,24 @@ namespace UsdStageImporterImpl
 	// Also, skeletal mesh components need to be manually ticked, or else they may be showing an animated state of an animation that
 	// we chose not to import, and wouldn't update otherwise until manually ticked by the user (or after save/reload), which may look
 	// like a bug
-	void RefreshComponents( AActor* RootSceneActor, bool bImportAtSpecificTimeCode )
+	void RefreshComponents(AActor* RootSceneActor, bool bImportAtSpecificTimeCode)
 	{
-		if ( !RootSceneActor )
+		if (!RootSceneActor)
 		{
 			return;
 		}
 
 		TArray<USceneComponent*> Components;
 		const bool bIncludeAllDescendants = true;
-		RootSceneActor->GetRootComponent()->GetChildrenComponents( bIncludeAllDescendants, Components );
+		RootSceneActor->GetRootComponent()->GetChildrenComponents(bIncludeAllDescendants, Components);
 
-		for ( USceneComponent* Component : Components )
+		for (USceneComponent* Component : Components)
 		{
-			if ( USkeletalMeshComponent* SkeletalMeshComponent = Cast<USkeletalMeshComponent>( Component ) )
+			if (USkeletalMeshComponent* SkeletalMeshComponent = Cast<USkeletalMeshComponent>(Component))
 			{
-				if ( SkeletalMeshComponent->AnimationData.AnimToPlay == nullptr )
+				if (SkeletalMeshComponent->AnimationData.AnimToPlay == nullptr)
 				{
-					SkeletalMeshComponent->TickAnimation( 0.f, false );
+					SkeletalMeshComponent->TickAnimation(0.f, false);
 					SkeletalMeshComponent->RefreshBoneTransforms();
 					SkeletalMeshComponent->RefreshFollowerComponents();
 					SkeletalMeshComponent->UpdateComponentToWorld();
@@ -1813,7 +1845,7 @@ namespace UsdStageImporterImpl
 					SkeletalMeshComponent->MarkRenderTransformDirty();
 					SkeletalMeshComponent->MarkRenderDynamicDataDirty();
 				}
-				else if ( bImportAtSpecificTimeCode )
+				else if (bImportAtSpecificTimeCode)
 				{
 					// The asset we return from the import factories may lead to
 					// USkeletalMesh::PostEditChangeProperty being called. The FMultiComponentReregisterContext in
@@ -1830,12 +1862,7 @@ namespace UsdStageImporterImpl
 					const bool bIsLooping = false;
 					const bool bIsPlaying = false;
 					const float Position = SkeletalMeshComponent->GetPosition();
-					SkeletalMeshComponent->OverrideAnimationData(
-						SkeletalMeshComponent->AnimationData.AnimToPlay,
-						bIsLooping,
-						bIsPlaying,
-						Position
-					);
+					SkeletalMeshComponent->OverrideAnimationData(SkeletalMeshComponent->AnimationData.AnimToPlay, bIsLooping, bIsPlaying, Position);
 				}
 
 				// It does need us to manually set this to dirty regardless or else it won't update in case we changed material
@@ -1844,14 +1871,17 @@ namespace UsdStageImporterImpl
 			}
 		}
 	}
-}
+}	 // namespace UE::USDStageImporter::Private
 
 void UUsdStageImporter::ImportFromFile(FUsdStageImportContext& ImportContext)
 {
 #if USE_USD_SDK
 	if (!ImportContext.World)
 	{
-		FUsdLogManager::LogMessage( EMessageSeverity::Error, LOCTEXT( "NoWorldError", "Failed to import USD Stage because the target UWorld is invalid!" ) );
+		FUsdLogManager::LogMessage(
+			EMessageSeverity::Error,
+			LOCTEXT("NoWorldError", "Failed to import USD Stage because the target UWorld is invalid!")
+		);
 		return;
 	}
 
@@ -1875,36 +1905,35 @@ void UUsdStageImporter::ImportFromFile(FUsdStageImportContext& ImportContext)
 		ImportContext.ImportOptions = NewObject<UUsdStageImportOptions>();
 	}
 
-	if ( ImportContext.ImportOptions->PrimsToImport.Num() == 0 )
+	if (ImportContext.ImportOptions->PrimsToImport.Num() == 0)
 	{
 		return;
 	}
 
-	const bool bNeedsMasking =
-		ImportContext.ImportOptions->PrimsToImport != TArray<FString>{ UE::FSdfPath::AbsoluteRootPath().GetString() };
+	const bool bNeedsMasking = ImportContext.ImportOptions->PrimsToImport != TArray<FString>{UE::FSdfPath::AbsoluteRootPath().GetString()};
 
-	if ( !ImportContext.Stage || bNeedsMasking )
+	if (!ImportContext.Stage || bNeedsMasking)
 	{
-		UsdStageImporterImpl::OpenStage( ImportContext, bNeedsMasking );
+		UE::USDStageImporter::Private::OpenStage(ImportContext, bNeedsMasking);
 	}
 
-	if ( !ImportContext.Stage )
+	if (!ImportContext.Stage)
 	{
-		FUsdLogManager::LogMessage( EMessageSeverity::Error, LOCTEXT( "NoStageError", "Failed to open the USD Stage!" ) );
+		FUsdLogManager::LogMessage(EMessageSeverity::Error, LOCTEXT("NoStageError", "Failed to open the USD Stage!"));
 		return;
 	}
 
-	UsdStageImporterImpl::SetupSceneActor( ImportContext );
-	if ( !ImportContext.SceneActor && ImportContext.ImportOptions->bImportActors )
+	UE::USDStageImporter::Private::SetupSceneActor(ImportContext);
+	if (!ImportContext.SceneActor && ImportContext.ImportOptions->bImportActors)
 	{
 		return;
 	}
 
-	FUsdDelegates::OnPreUsdImport.Broadcast( ImportContext.FilePath );
+	FUsdDelegates::OnPreUsdImport.Broadcast(ImportContext.FilePath);
 
-	AActor* ExistingSceneActor = UsdStageImporterImpl::GetExistingSceneActor( ImportContext );
+	AActor* ExistingSceneActor = UE::USDStageImporter::Private::GetExistingSceneActor(ImportContext);
 
-	UsdStageImporterImpl::SetupStageForImport( ImportContext );
+	UE::USDStageImporter::Private::SetupStageForImport(ImportContext);
 
 	TMap<FSoftObjectPath, FSoftObjectPath> SoftObjectsToRemap;
 	TMap<UObject*, UObject*> ObjectsToRemap;
@@ -1914,24 +1943,21 @@ void UUsdStageImporter::ImportFromFile(FUsdStageImportContext& ImportContext)
 	UsdUtils::FBlendShapeMap BlendShapesByPath;
 
 	// Ensure a valid asset cache
-	if ( !ImportContext.AssetCache )
+	if (!ImportContext.AssetCache)
 	{
-		UE_LOG(LogUsd, Log, TEXT("Generating a temporary USD Asset Cache when importing '%s'."),
-			*ImportContext.FilePath
-		);
+		UE_LOG(LogUsd, Log, TEXT("Generating a temporary USD Asset Cache when importing '%s'."), *ImportContext.FilePath);
 		ImportContext.AssetCache = NewObject<UUsdAssetCache2>();
 	}
 	TSharedPtr<FUsdInfoCache> InfoCache = MakeShared<FUsdInfoCache>();
 	ImportContext.AssetCache->MarkAssetsAsStale();
 	ImportContext.LevelSequenceHelper.SetInfoCache(InfoCache);
-	ImportContext.LevelSequenceHelper.Init( ImportContext.Stage );  // Must happen after the context gets an InfoCache!
+	ImportContext.LevelSequenceHelper.Init(ImportContext.Stage);	// Must happen after the context gets an InfoCache!
 	ImportContext.LevelSequenceHelper.SetRootMotionHandling(ImportContext.ImportOptions->RootMotionHandling);
 
 	EUsdPurpose PurposesToImport = static_cast<EUsdPurpose>(ImportContext.ImportOptions->PurposesToImport);
 
-	float ImportTime = ImportContext.ImportOptions->bImportAtSpecificTimeCode
-						   ? ImportContext.ImportOptions->ImportTimeCode
-						   : static_cast<float>(UsdUtils::GetDefaultTimeCode());
+	float ImportTime = ImportContext.ImportOptions->bImportAtSpecificTimeCode ? ImportContext.ImportOptions->ImportTimeCode
+																			  : static_cast<float>(UsdUtils::GetDefaultTimeCode());
 
 	if (!ImportContext.BBoxCache)
 	{
@@ -1954,10 +1980,14 @@ void UUsdStageImporter::ImportFromFile(FUsdStageImportContext& ImportContext)
 	}
 	ImportContext.LevelSequenceHelper.SetBBoxCache(ImportContext.BBoxCache);
 
-	// Shotgun approach to recreate all render states because we may want to reimport/delete/reassign a material/static/skeletalmesh while it is currently being drawn
+	// Shotgun approach to recreate all render states because we may want to reimport/delete/reassign a material/static/skeletalmesh while it is
+	// currently being drawn
 	FGlobalComponentRecreateRenderStateContext RecreateRenderStateContext;
 
-	TSharedRef<FUsdSchemaTranslationContext> TranslationContext = MakeShared<FUsdSchemaTranslationContext>( ImportContext.Stage, *ImportContext.AssetCache );
+	TSharedRef<FUsdSchemaTranslationContext> TranslationContext = MakeShared<FUsdSchemaTranslationContext>(
+		ImportContext.Stage,
+		*ImportContext.AssetCache
+	);
 	TranslationContext->bIsImporting = true;
 	TranslationContext->Level = ImportContext.World->GetCurrentLevel();
 	TranslationContext->ObjectFlags = ImportContext.ImportObjectFlags;
@@ -1970,11 +2000,12 @@ void UUsdStageImporter::ImportFromFile(FUsdStageImportContext& ImportContext)
 	TranslationContext->SubdivisionLevel = ImportContext.ImportOptions->SubdivisionLevel;
 	TranslationContext->MetadataOptions = ImportContext.ImportOptions->MetadataOptions;
 	TranslationContext->ParentComponent = ImportContext.SceneActor ? ImportContext.SceneActor->GetRootComponent() : nullptr;
-	TranslationContext->KindsToCollapse = ( EUsdDefaultKind ) ImportContext.ImportOptions->KindsToCollapse;
+	TranslationContext->KindsToCollapse = (EUsdDefaultKind)ImportContext.ImportOptions->KindsToCollapse;
 	TranslationContext->bMergeIdenticalMaterialSlots = ImportContext.ImportOptions->bMergeIdenticalMaterialSlots;
 	TranslationContext->bReuseIdenticalAssets = ImportContext.ImportOptions->bReuseIdenticalAssets;
 	TranslationContext->bAllowInterpretingLODs = ImportContext.ImportOptions->bInterpretLODs;
-	TranslationContext->bAllowParsingSkeletalAnimations = ImportContext.ImportOptions->bImportGeometry && ImportContext.ImportOptions->bImportSkeletalAnimations;
+	TranslationContext->bAllowParsingSkeletalAnimations = ImportContext.ImportOptions->bImportGeometry
+														  && ImportContext.ImportOptions->bImportSkeletalAnimations;
 	TranslationContext->bAllowParsingGroomAssets = ImportContext.ImportOptions->bImportGroomAssets;
 	TranslationContext->bTranslateOnlyUsedMaterials = ImportContext.ImportOptions->bImportOnlyUsedMaterials;
 	TranslationContext->InfoCache = InfoCache;
@@ -1982,25 +2013,35 @@ void UUsdStageImporter::ImportFromFile(FUsdStageImportContext& ImportContext)
 	TranslationContext->BlendShapesByPath = &BlendShapesByPath;
 	TranslationContext->GroomInterpolationSettings = ImportContext.ImportOptions->GroomInterpolationSettings;
 	{
-		UsdStageImporterImpl::CacheCollapsingState( TranslationContext.Get() );
-		UsdStageImporterImpl::ImportMaterials( ImportContext, TranslationContext.Get() );
-		UsdStageImporterImpl::ImportMeshes( ImportContext, TranslationContext.Get() );
-		UsdStageImporterImpl::ImportActors( ImportContext, TranslationContext.Get() );
+		UE::USDStageImporter::Private::CacheCollapsingState(TranslationContext.Get());
+		UE::USDStageImporter::Private::ImportMaterials(ImportContext, TranslationContext.Get());
+		UE::USDStageImporter::Private::ImportMeshes(ImportContext, TranslationContext.Get());
+		UE::USDStageImporter::Private::ImportActors(ImportContext, TranslationContext.Get());
 	}
 	TranslationContext->CompleteTasks();
 
-	UsdStageImporterImpl::CollectUsedAssetDependencies( ImportContext, UsedAssetsAndDependencies );
-	UsdStageImporterImpl::PruneUnwantedAssets( ImportContext, UsedAssetsAndDependencies, ObjectsToRemap, SoftObjectsToRemap );
-	UsdStageImporterImpl::UpdateAssetUserData( UsedAssetsAndDependencies, ImportContext.FilePath, ImportContext.ImportOptions );
-	UsdStageImporterImpl::ResolveActorConflicts( ImportContext, ExistingSceneActor, ObjectsToRemap, SoftObjectsToRemap );
-	UsdStageImporterImpl::PublishAssets( ImportContext, UsedAssetsAndDependencies, ObjectsToRemap, SoftObjectsToRemap, AssetsToFinalize );
-	UsdStageImporterImpl::RemapReferences( ImportContext, UsedAssetsAndDependencies, ObjectsToRemap );
-	UsdStageImporterImpl::RemapSoftReferences( ImportContext, UsedAssetsAndDependencies, SoftObjectsToRemap );
-	UsdStageImporterImpl::Cleanup( ImportContext.SceneActor, ExistingSceneActor, ImportContext.ImportOptions->ExistingActorPolicy, AssetsToFinalize );
-	UsdStageImporterImpl::GetPublishedAssetsAndDependencies(UsedAssetsAndDependencies, ObjectsToRemap, SoftObjectsToRemap, PublishedAssetsAndDependencies);
-	UsdStageImporterImpl::CallAssetsPostEditChange(PublishedAssetsAndDependencies);
-	UsdStageImporterImpl::NotifyAssetRegistry(PublishedAssetsAndDependencies);
-	UsdStageImporterImpl::RefreshComponents( ImportContext.SceneActor, ImportContext.ImportOptions->bImportAtSpecificTimeCode );
+	UE::USDStageImporter::Private::CollectUsedAssetDependencies(ImportContext, UsedAssetsAndDependencies);
+	UE::USDStageImporter::Private::PruneUnwantedAssets(ImportContext, UsedAssetsAndDependencies, ObjectsToRemap, SoftObjectsToRemap);
+	UE::USDStageImporter::Private::UpdateAssetUserData(UsedAssetsAndDependencies, ImportContext.FilePath, ImportContext.ImportOptions);
+	UE::USDStageImporter::Private::ResolveActorConflicts(ImportContext, ExistingSceneActor, ObjectsToRemap, SoftObjectsToRemap);
+	UE::USDStageImporter::Private::PublishAssets(ImportContext, UsedAssetsAndDependencies, ObjectsToRemap, SoftObjectsToRemap, AssetsToFinalize);
+	UE::USDStageImporter::Private::RemapReferences(ImportContext, UsedAssetsAndDependencies, ObjectsToRemap);
+	UE::USDStageImporter::Private::RemapSoftReferences(ImportContext, UsedAssetsAndDependencies, SoftObjectsToRemap);
+	UE::USDStageImporter::Private::Cleanup(
+		ImportContext.SceneActor,
+		ExistingSceneActor,
+		ImportContext.ImportOptions->ExistingActorPolicy,
+		AssetsToFinalize
+	);
+	UE::USDStageImporter::Private::GetPublishedAssetsAndDependencies(
+		UsedAssetsAndDependencies,
+		ObjectsToRemap,
+		SoftObjectsToRemap,
+		PublishedAssetsAndDependencies
+	);
+	UE::USDStageImporter::Private::CallAssetsPostEditChange(PublishedAssetsAndDependencies);
+	UE::USDStageImporter::Private::NotifyAssetRegistry(PublishedAssetsAndDependencies);
+	UE::USDStageImporter::Private::RefreshComponents(ImportContext.SceneActor, ImportContext.ImportOptions->bImportAtSpecificTimeCode);
 
 	if (IncludedPurposesToRevertBBoxCacheTo.IsSet())
 	{
@@ -2011,19 +2052,19 @@ void UUsdStageImporter::ImportFromFile(FUsdStageImportContext& ImportContext)
 		ImportContext.BBoxCache->SetTime(TimeToRevertBBoxCacheTo.GetValue());
 	}
 
-	FUsdDelegates::OnPostUsdImport.Broadcast( ImportContext.FilePath );
+	FUsdDelegates::OnPostUsdImport.Broadcast(ImportContext.FilePath);
 
 	// Analytics
 	{
-		double ElapsedSeconds = FPlatformTime::ToSeconds64( FPlatformTime::Cycles64() - StartTime );
-		UsdStageImporterImpl::SendAnalytics( ImportContext, nullptr, TEXT("Import"), PublishedAssetsAndDependencies, ElapsedSeconds);
+		double ElapsedSeconds = FPlatformTime::ToSeconds64(FPlatformTime::Cycles64() - StartTime);
+		UE::USDStageImporter::Private::SendAnalytics(ImportContext, nullptr, TEXT("Import"), PublishedAssetsAndDependencies, ElapsedSeconds);
 		UE_LOG(LogUsd, Log, TEXT("Imported '%s' in %.3f seconds."), *ImportContext.FilePath, ElapsedSeconds);
 
 		UsdUtils::CollectSchemaAnalytics(ImportContext.Stage, TEXT("Import"));
 	}
 
-	UsdStageImporterImpl::CloseStageIfNeeded( ImportContext );
-#endif // #if USE_USD_SDK
+	UE::USDStageImporter::Private::CloseStageIfNeeded(ImportContext);
+#endif	  // #if USE_USD_SDK
 }
 
 bool UUsdStageImporter::ReimportSingleAsset(
@@ -2045,17 +2086,16 @@ bool UUsdStageImporter::ReimportSingleAsset(
 	}
 
 	// TODO: Maybe change this whole reimporting approach to just taking advantage of a population mask instead?
-	const bool bNeedsMasking =
-		ImportContext.ImportOptions->PrimsToImport != TArray<FString>{ UE::FSdfPath::AbsoluteRootPath().GetString() };
+	const bool bNeedsMasking = ImportContext.ImportOptions->PrimsToImport != TArray<FString>{UE::FSdfPath::AbsoluteRootPath().GetString()};
 
-	if ( !ImportContext.Stage || bNeedsMasking )
+	if (!ImportContext.Stage || bNeedsMasking)
 	{
-		UsdStageImporterImpl::OpenStage( ImportContext, bNeedsMasking );
+		UE::USDStageImporter::Private::OpenStage(ImportContext, bNeedsMasking);
 	}
 
-	if ( !ImportContext.Stage )
+	if (!ImportContext.Stage)
 	{
-		FUsdLogManager::LogMessage( EMessageSeverity::Error, LOCTEXT( "NoStageError", "Failed to open the USD Stage!" ) );
+		FUsdLogManager::LogMessage(EMessageSeverity::Error, LOCTEXT("NoStageError", "Failed to open the USD Stage!"));
 		return bSuccess;
 	}
 
@@ -2063,9 +2103,9 @@ bool UUsdStageImporter::ReimportSingleAsset(
 
 	// We still need the scene actor to remap all other users of the mesh to the new reimported one. It's not critical if we fail though,
 	// the goal is to just reimport the asset
-	UsdStageImporterImpl::SetupSceneActor(ImportContext);
+	UE::USDStageImporter::Private::SetupSceneActor(ImportContext);
 
-	UsdStageImporterImpl::SetupStageForImport( ImportContext );
+	UE::USDStageImporter::Private::SetupStageForImport(ImportContext);
 
 	TMap<FSoftObjectPath, FSoftObjectPath> SoftObjectsToRemap;
 	TMap<UObject*, UObject*> ObjectsToRemap;
@@ -2073,24 +2113,21 @@ bool UUsdStageImporter::ReimportSingleAsset(
 	UsdUtils::FBlendShapeMap BlendShapesByPath;
 
 	// Ensure a valid asset cache
-	if ( !ImportContext.AssetCache )
+	if (!ImportContext.AssetCache)
 	{
-		UE_LOG(LogUsd, Log, TEXT("Generating a temporary USD Asset Cache when importing '%s'."),
-			*ImportContext.FilePath
-		);
+		UE_LOG(LogUsd, Log, TEXT("Generating a temporary USD Asset Cache when importing '%s'."), *ImportContext.FilePath);
 		ImportContext.AssetCache = NewObject<UUsdAssetCache2>();
 	}
 	TSharedPtr<FUsdInfoCache> InfoCache = MakeShared<FUsdInfoCache>();
 	ImportContext.AssetCache->MarkAssetsAsStale();
 	ImportContext.LevelSequenceHelper.SetInfoCache(InfoCache);
-	ImportContext.LevelSequenceHelper.Init(ImportContext.Stage);  // Must happen after the context gets an InfoCache!
+	ImportContext.LevelSequenceHelper.Init(ImportContext.Stage);	// Must happen after the context gets an InfoCache!
 	ImportContext.LevelSequenceHelper.SetRootMotionHandling(ImportContext.ImportOptions->RootMotionHandling);
 
 	EUsdPurpose PurposesToImport = static_cast<EUsdPurpose>(ImportContext.ImportOptions->PurposesToImport);
 
-	float ImportTime = ImportContext.ImportOptions->bImportAtSpecificTimeCode
-						   ? ImportContext.ImportOptions->ImportTimeCode
-						   : static_cast<float>(UsdUtils::GetDefaultTimeCode());
+	float ImportTime = ImportContext.ImportOptions->bImportAtSpecificTimeCode ? ImportContext.ImportOptions->ImportTimeCode
+																			  : static_cast<float>(UsdUtils::GetDefaultTimeCode());
 
 	if (!ImportContext.BBoxCache)
 	{
@@ -2112,10 +2149,14 @@ bool UUsdStageImporter::ReimportSingleAsset(
 	}
 	ImportContext.LevelSequenceHelper.SetBBoxCache(ImportContext.BBoxCache);
 
-	// Shotgun approach to recreate all render states because we may want to reimport/delete/reassign a material/static/skeletalmesh while it is currently being drawn
+	// Shotgun approach to recreate all render states because we may want to reimport/delete/reassign a material/static/skeletalmesh while it is
+	// currently being drawn
 	FGlobalComponentRecreateRenderStateContext RecreateRenderStateContext;
 
-	TSharedRef<FUsdSchemaTranslationContext> TranslationContext = MakeShared<FUsdSchemaTranslationContext>( ImportContext.Stage, *ImportContext.AssetCache );
+	TSharedRef<FUsdSchemaTranslationContext> TranslationContext = MakeShared<FUsdSchemaTranslationContext>(
+		ImportContext.Stage,
+		*ImportContext.AssetCache
+	);
 	TranslationContext->bIsImporting = true;
 	TranslationContext->Level = ImportContext.World->GetCurrentLevel();
 	TranslationContext->ObjectFlags = ImportContext.ImportObjectFlags;
@@ -2127,11 +2168,12 @@ bool UUsdStageImporter::ReimportSingleAsset(
 	TranslationContext->RootMotionHandling = ImportContext.ImportOptions->RootMotionHandling;
 	TranslationContext->SubdivisionLevel = ImportContext.ImportOptions->SubdivisionLevel;
 	TranslationContext->MetadataOptions = ImportContext.ImportOptions->MetadataOptions;
-	TranslationContext->KindsToCollapse = ( EUsdDefaultKind ) ImportContext.ImportOptions->KindsToCollapse;
+	TranslationContext->KindsToCollapse = (EUsdDefaultKind)ImportContext.ImportOptions->KindsToCollapse;
 	TranslationContext->bMergeIdenticalMaterialSlots = ImportContext.ImportOptions->bMergeIdenticalMaterialSlots;
 	TranslationContext->bReuseIdenticalAssets = ImportContext.ImportOptions->bReuseIdenticalAssets;
 	TranslationContext->bAllowInterpretingLODs = ImportContext.ImportOptions->bInterpretLODs;
-	TranslationContext->bAllowParsingSkeletalAnimations = ImportContext.ImportOptions->bImportGeometry && ImportContext.ImportOptions->bImportSkeletalAnimations;
+	TranslationContext->bAllowParsingSkeletalAnimations = ImportContext.ImportOptions->bImportGeometry
+														  && ImportContext.ImportOptions->bImportSkeletalAnimations;
 	TranslationContext->bAllowParsingGroomAssets = ImportContext.ImportOptions->bImportGroomAssets;
 	TranslationContext->bTranslateOnlyUsedMaterials = ImportContext.ImportOptions->bImportOnlyUsedMaterials;
 	TranslationContext->InfoCache = InfoCache;
@@ -2139,19 +2181,19 @@ bool UUsdStageImporter::ReimportSingleAsset(
 	TranslationContext->BlendShapesByPath = &BlendShapesByPath;
 	TranslationContext->GroomInterpolationSettings = ImportContext.ImportOptions->GroomInterpolationSettings;
 	{
-		UsdStageImporterImpl::CacheCollapsingState( TranslationContext.Get() );
+		UE::USDStageImporter::Private::CacheCollapsingState(TranslationContext.Get());
 
 		UE::FUsdPrim TargetPrim = ImportContext.Stage.GetPrimAtPath(UE::FSdfPath(*OriginalPrimPath));
-		if ( TargetPrim )
+		if (TargetPrim)
 		{
-			UsdStageImporterImpl::CreateAssetsForPrims({TargetPrim}, TranslationContext.Get(), LOCTEXT("CreateAssets", "Creating assets"));
+			UE::USDStageImporter::Private::CreateAssetsForPrims({TargetPrim}, TranslationContext.Get(), LOCTEXT("CreateAssets", "Creating assets"));
 		}
 	}
 	TranslationContext->CompleteTasks();
 
 	// Look for our reimported asset in the assets cache as we may have multiple assets with the same prim path
 	UObject* ReimportedObject = nullptr;
-	for ( UObject* Asset : ImportContext.AssetCache->GetActiveAssets() )
+	for (UObject* Asset : ImportContext.AssetCache->GetActiveAssets())
 	{
 		UUsdAssetUserData* UserData = nullptr;
 		if (IInterface_AssetUserData* UserDataInterface = Cast<IInterface_AssetUserData>(Asset))
@@ -2159,26 +2201,24 @@ bool UUsdStageImporter::ReimportSingleAsset(
 			UserData = UserDataInterface->GetAssetUserData<UUsdAssetUserData>();
 		}
 
-		if ( Asset &&
-			 UserData &&
-			 Asset->GetClass() == OriginalAsset->GetClass() &&
-			 UserData->PrimPaths.Contains(OriginalPrimPath))
+		if (Asset && UserData && Asset->GetClass() == OriginalAsset->GetClass() && UserData->PrimPaths.Contains(OriginalPrimPath))
 		{
 			ReimportedObject = Asset;
 			break;
 		}
 	}
 
-	if ( ReimportedObject )
+	if (ReimportedObject)
 	{
-		UsdStageImporterImpl::UpdateAssetImportData(ReimportedObject, ImportContext.FilePath, ImportContext.ImportOptions);
+		UE::USDStageImporter::Private::UpdateAssetImportData(ReimportedObject, ImportContext.FilePath, ImportContext.ImportOptions);
 
 		// Assign things from the original assets before we publish the reimported asset, overwriting it
-		UsdStageImporterImpl::CopyOriginalMaterialAssignment(ImportContext, OriginalAsset, ReimportedObject );
-		UsdStageImporterImpl::CopySkeletonAssignment(ImportContext, OriginalAsset, ReimportedObject );
+		UE::USDStageImporter::Private::CopyOriginalMaterialAssignment(ImportContext, OriginalAsset, ReimportedObject);
+		UE::USDStageImporter::Private::CopySkeletonAssignment(ImportContext, OriginalAsset, ReimportedObject);
 
-		// Just publish the one asset we wanted to reimport. Note that we may have other assets here too, but we'll ignore those e.g. a displayColor material or a skeleton
-		OutReimportedAsset = UsdStageImporterImpl::PublishAsset(
+		// Just publish the one asset we wanted to reimport. Note that we may have other assets here too, but we'll ignore those e.g. a displayColor
+		// material or a skeleton
+		OutReimportedAsset = UE::USDStageImporter::Private::PublishAsset(
 			ImportContext,
 			ReimportedObject,
 			OriginalAsset->GetOutermost()->GetPathName(),
@@ -2189,22 +2229,22 @@ bool UUsdStageImporter::ReimportSingleAsset(
 
 		TSet<UObject*> UsedAssetsAndDependencies = ImportContext.AssetCache->GetActiveAssets();
 		TSet<UObject*> PublishedAssetsAndDependencies;
-		UsdStageImporterImpl::RemapReferences(ImportContext, UsedAssetsAndDependencies, ObjectsToRemap);
-		UsdStageImporterImpl::RemapSoftReferences(ImportContext, UsedAssetsAndDependencies, SoftObjectsToRemap);
-		UsdStageImporterImpl::GetPublishedAssetsAndDependencies(UsedAssetsAndDependencies, ObjectsToRemap, SoftObjectsToRemap, PublishedAssetsAndDependencies);
-		UsdStageImporterImpl::CallAssetsPostEditChange(PublishedAssetsAndDependencies);
+		UE::USDStageImporter::Private::RemapReferences(ImportContext, UsedAssetsAndDependencies, ObjectsToRemap);
+		UE::USDStageImporter::Private::RemapSoftReferences(ImportContext, UsedAssetsAndDependencies, SoftObjectsToRemap);
+		UE::USDStageImporter::Private::GetPublishedAssetsAndDependencies(
+			UsedAssetsAndDependencies,
+			ObjectsToRemap,
+			SoftObjectsToRemap,
+			PublishedAssetsAndDependencies
+		);
+		UE::USDStageImporter::Private::CallAssetsPostEditChange(PublishedAssetsAndDependencies);
 
-		bSuccess = OutReimportedAsset != nullptr && ImportContext.AssetCache->GetActiveAssets().Contains( ReimportedObject );
+		bSuccess = OutReimportedAsset != nullptr && ImportContext.AssetCache->GetActiveAssets().Contains(ReimportedObject);
 	}
 
-	UsdStageImporterImpl::Cleanup(
-		ImportContext.SceneActor,
-		nullptr,
-		ImportContext.ImportOptions->ExistingActorPolicy,
-		AssetsToFinalize
-	);
-	UsdStageImporterImpl::NotifyAssetRegistry( { ReimportedObject } );
-	UsdStageImporterImpl::RefreshComponents( ImportContext.SceneActor, ImportContext.ImportOptions->bImportAtSpecificTimeCode );
+	UE::USDStageImporter::Private::Cleanup(ImportContext.SceneActor, nullptr, ImportContext.ImportOptions->ExistingActorPolicy, AssetsToFinalize);
+	UE::USDStageImporter::Private::NotifyAssetRegistry({ReimportedObject});
+	UE::USDStageImporter::Private::RefreshComponents(ImportContext.SceneActor, ImportContext.ImportOptions->bImportAtSpecificTimeCode);
 
 	if (IncludedPurposesToRevertTo.IsSet())
 	{
@@ -2219,18 +2259,17 @@ bool UUsdStageImporter::ReimportSingleAsset(
 
 	// Analytics
 	{
-		double ElapsedSeconds = FPlatformTime::ToSeconds64( FPlatformTime::Cycles64() - StartTime );
-		UsdStageImporterImpl::SendAnalytics( ImportContext, ReimportedObject, TEXT( "Reimport" ), { ReimportedObject }, ElapsedSeconds );
+		double ElapsedSeconds = FPlatformTime::ToSeconds64(FPlatformTime::Cycles64() - StartTime);
+		UE::USDStageImporter::Private::SendAnalytics(ImportContext, ReimportedObject, TEXT("Reimport"), {ReimportedObject}, ElapsedSeconds);
 		UE_LOG(LogUsd, Log, TEXT("Re-imported '%s' in %.3f seconds."), *ImportContext.FilePath, ElapsedSeconds);
 
 		UsdUtils::CollectSchemaAnalytics(ImportContext.Stage, TEXT("Reimport"));
 	}
 
-	UsdStageImporterImpl::CloseStageIfNeeded( ImportContext );
+	UE::USDStageImporter::Private::CloseStageIfNeeded(ImportContext);
 
-#endif // #if USE_USD_SDK
+#endif	  // #if USE_USD_SDK
 	return bSuccess;
 }
 
 #undef LOCTEXT_NAMESPACE
-
diff --git a/Engine/Plugins/Importers/USDImporter/Source/USDStageImporter/Private/USDStageImporterModule.cpp b/Engine/Plugins/Importers/USDImporter/Source/USDStageImporter/Private/USDStageImporterModule.cpp
--- a/Engine/Plugins/Importers/USDImporter/Source/USDStageImporter/Private/USDStageImporterModule.cpp
+++ b/Engine/Plugins/Importers/USDImporter/Source/USDStageImporter/Private/USDStageImporterModule.cpp
@@ -7,7 +7,6 @@
 #include "USDStageImporter.h"
 #include "USDStageImportOptionsCustomization.h"
 
-#include "Misc/Paths.h"
 #include "Modules/ModuleManager.h"
 #include "PropertyEditorModule.h"
 #include "Templates/UniquePtr.h"
@@ -22,19 +21,22 @@ public:
 #if USE_USD_SDK
 		LLM_SCOPE_BYTAG(Usd);
 
-		IUnrealUSDWrapperModule& UnrealUSDWrapperModule = FModuleManager::Get().LoadModuleChecked< IUnrealUSDWrapperModule >(TEXT("UnrealUSDWrapper"));
+		FModuleManager::Get().LoadModuleChecked<IUnrealUSDWrapperModule>(TEXT("UnrealUSDWrapper"));
 
-		FPropertyEditorModule& PropertyModule = FModuleManager::LoadModuleChecked<FPropertyEditorModule>( TEXT( "PropertyEditor" ) );
-		PropertyModule.RegisterCustomClassLayout( TEXT( "UsdStageImportOptions" ), FOnGetDetailCustomizationInstance::CreateStatic( &FUsdStageImportOptionsCustomization::MakeInstance ) );
+		FPropertyEditorModule& PropertyModule = FModuleManager::LoadModuleChecked<FPropertyEditorModule>(TEXT("PropertyEditor"));
+		PropertyModule.RegisterCustomClassLayout(
+			TEXT("UsdStageImportOptions"),
+			FOnGetDetailCustomizationInstance::CreateStatic(&FUsdStageImportOptionsCustomization::MakeInstance)
+		);
 
 		USDStageImporter = MakeUnique<UUsdStageImporter>();
-#endif // #if USE_USD_SDK
+#endif	  // #if USE_USD_SDK
 	}
 
 	virtual void ShutdownModule() override
 	{
-		FPropertyEditorModule& PropertyModule = FModuleManager::LoadModuleChecked< FPropertyEditorModule >( TEXT( "PropertyEditor" ) );
-		PropertyModule.UnregisterCustomClassLayout( TEXT( "UsdStageImportOptions" ) );
+		FPropertyEditorModule& PropertyModule = FModuleManager::LoadModuleChecked<FPropertyEditorModule>(TEXT("PropertyEditor"));
+		PropertyModule.UnregisterCustomClassLayout(TEXT("UsdStageImportOptions"));
 
 		USDStageImporter.Reset();
 	}
diff --git a/Engine/Plugins/Importers/USDImporter/Source/USDStageImporter/Private/Widgets/SUSDStagePreviewTree.cpp b/Engine/Plugins/Importers/USDImporter/Source/USDStageImporter/Private/Widgets/SUSDStagePreviewTree.cpp
--- a/Engine/Plugins/Importers/USDImporter/Source/USDStageImporter/Private/Widgets/SUSDStagePreviewTree.cpp
+++ b/Engine/Plugins/Importers/USDImporter/Source/USDStageImporter/Private/Widgets/SUSDStagePreviewTree.cpp
@@ -18,16 +18,16 @@ namespace UE::UsdStagePreviewTree::Private
 {
 	const float CheckBoxColumnWidth = 24.0f;
 
-	bool AreAllViewsCheckedRecursive( const SUsdStagePreviewTree& Tree, const TArray<FUsdPrimPreviewModelViewRef>& Views )
+	bool AreAllViewsCheckedRecursive(const SUsdStagePreviewTree& Tree, const TArray<FUsdPrimPreviewModelViewRef>& Views)
 	{
-		for ( const FUsdPrimPreviewModelViewRef& View : Views )
+		for (const FUsdPrimPreviewModelViewRef& View : Views)
 		{
-			if ( !Tree.GetModel( View->ModelIndex ).bShouldImport )
+			if (!Tree.GetModel(View->ModelIndex).bShouldImport)
 			{
 				return false;
 			}
 
-			AreAllViewsCheckedRecursive( Tree, View->Children );
+			AreAllViewsCheckedRecursive(Tree, View->Children);
 		}
 
 		return true;
@@ -37,58 +37,56 @@ namespace UE::UsdStagePreviewTree::Private
 class FUsdPreviewTreeImportColumn : public FUsdTreeViewColumn
 {
 public:
-	virtual TSharedRef< SWidget > GenerateWidget(
-		const TSharedPtr< IUsdTreeViewItem > InTreeItem,
-		const TSharedPtr< ITableRow > TableRow
-	) override
+	virtual TSharedRef<SWidget> GenerateWidget(const TSharedPtr<IUsdTreeViewItem> InTreeItem, const TSharedPtr<ITableRow> TableRow) override
 	{
-		const FUsdPrimPreviewModelViewWeak WeakTreeItem = StaticCastSharedPtr< FUsdPrimPreviewModelView >( InTreeItem );
-		TSharedPtr< SUsdStagePreviewTreeRow > Row = StaticCastSharedPtr< SUsdStagePreviewTreeRow >( TableRow );
-		TWeakPtr< SUsdStagePreviewTree > WeakTree = Row ? Row->GetOwnerTree() : nullptr;
-
-		return SNew( SBox )
-			.VAlign( VAlign_Center )
-			.HAlign( HAlign_Center )
-			.WidthOverride( UE::UsdStagePreviewTree::Private::CheckBoxColumnWidth )
+		const FUsdPrimPreviewModelViewWeak WeakTreeItem = StaticCastSharedPtr<FUsdPrimPreviewModelView>(InTreeItem);
+		TSharedPtr<SUsdStagePreviewTreeRow> Row = StaticCastSharedPtr<SUsdStagePreviewTreeRow>(TableRow);
+		TWeakPtr<SUsdStagePreviewTree> WeakTree = Row ? Row->GetOwnerTree() : nullptr;
+
+		// clang-format off
+		return SNew(SBox)
+			.VAlign(VAlign_Center)
+			.HAlign(HAlign_Center)
+			.WidthOverride(UE::UsdStagePreviewTree::Private::CheckBoxColumnWidth)
 			[
-				SNew( SCheckBox )
-				.IsChecked_Lambda( [WeakTreeItem, WeakTree]() -> ECheckBoxState
+				SNew(SCheckBox)
+				.IsChecked_Lambda([WeakTreeItem, WeakTree]() -> ECheckBoxState
 				{
 					FUsdPrimPreviewModelViewPtr PinnedTreeItem = WeakTreeItem.Pin();
-					TSharedPtr< SUsdStagePreviewTree > PinnedTree = WeakTree.Pin();
-					if( PinnedTreeItem && PinnedTree )
+					TSharedPtr<SUsdStagePreviewTree> PinnedTree = WeakTree.Pin();
+					if(PinnedTreeItem && PinnedTree)
 					{
-						FUsdPrimPreviewModel& Model = PinnedTree->GetModel( PinnedTreeItem->ModelIndex );
+						FUsdPrimPreviewModel& Model = PinnedTree->GetModel(PinnedTreeItem->ModelIndex);
 						return Model.bShouldImport ? ECheckBoxState::Checked : ECheckBoxState::Unchecked;
 					}
 
 					return ECheckBoxState::Undetermined;
 				})
-				.OnCheckStateChanged_Lambda([WeakTreeItem, WeakTree]( ECheckBoxState NewCheckedState )
+				.OnCheckStateChanged_Lambda([WeakTreeItem, WeakTree](ECheckBoxState NewCheckedState)
 				{
 					FUsdPrimPreviewModelViewPtr PinnedTreeItem = WeakTreeItem.Pin();
-					TSharedPtr< SUsdStagePreviewTree > PinnedTree = WeakTree.Pin();
-					if ( PinnedTreeItem && PinnedTree )
+					TSharedPtr<SUsdStagePreviewTree> PinnedTree = WeakTree.Pin();
+					if (PinnedTreeItem && PinnedTree)
 					{
-						if ( NewCheckedState == ECheckBoxState::Checked )
+						if (NewCheckedState == ECheckBoxState::Checked)
 						{
 							// Check recursively when shift clicking (tree views support recursive expansion when
 							// shift-clicking the expander arrows, so lets do that too on the checkboxes!)
-							if ( FSlateApplication::Get().GetModifierKeys().IsShiftDown() )
+							if (FSlateApplication::Get().GetModifierKeys().IsShiftDown())
 							{
 								const bool bChecked = true;
-								PinnedTree->CheckItemsRecursively( { PinnedTreeItem.ToSharedRef() }, bChecked );
+								PinnedTree->CheckItemsRecursively({PinnedTreeItem.ToSharedRef()}, bChecked);
 							}
 							else
 							{
-								FUsdPrimPreviewModel& Model = PinnedTree->GetModel( PinnedTreeItem->ModelIndex );
+								FUsdPrimPreviewModel& Model = PinnedTree->GetModel(PinnedTreeItem->ModelIndex);
 								Model.bShouldImport = true;
 
 								// If we want to import some prim, we must import all ancestors too
 								int32 ParentIndex = Model.ParentIndex;
-								while ( ParentIndex != INDEX_NONE )
+								while (ParentIndex != INDEX_NONE)
 								{
-									FUsdPrimPreviewModel& ParentModel = PinnedTree->GetModel( ParentIndex );
+									FUsdPrimPreviewModel& ParentModel = PinnedTree->GetModel(ParentIndex);
 									ParentModel.bShouldImport = true;
 									ParentIndex = ParentModel.ParentIndex;
 								}
@@ -97,37 +95,37 @@ public:
 						else
 						{
 							const bool bChecked = false;
-							PinnedTree->CheckItemsRecursively( { PinnedTreeItem.ToSharedRef() }, bChecked );
+							PinnedTree->CheckItemsRecursively({PinnedTreeItem.ToSharedRef()}, bChecked);
 						}
 					}
 				})
 			];
+		// clang-format on
 	}
 };
 
 class FUsdPreviewTreeNameColumn : public FUsdTreeViewColumn
 {
 public:
-	virtual TSharedRef< SWidget > GenerateWidget(
-		const TSharedPtr< IUsdTreeViewItem > InTreeItem,
-		const TSharedPtr< ITableRow > TableRow
-	) override
+	virtual TSharedRef<SWidget> GenerateWidget(const TSharedPtr<IUsdTreeViewItem> InTreeItem, const TSharedPtr<ITableRow> TableRow) override
 	{
-		FUsdPrimPreviewModelViewPtr TreeItem = StaticCastSharedPtr< FUsdPrimPreviewModelView >( InTreeItem );
-		TSharedPtr< SUsdStagePreviewTreeRow > Row = StaticCastSharedPtr< SUsdStagePreviewTreeRow >( TableRow );
-		TSharedPtr< SUsdStagePreviewTree > Tree = Row ? Row->GetOwnerTree() : nullptr;
+		FUsdPrimPreviewModelViewPtr TreeItem = StaticCastSharedPtr<FUsdPrimPreviewModelView>(InTreeItem);
+		TSharedPtr<SUsdStagePreviewTreeRow> Row = StaticCastSharedPtr<SUsdStagePreviewTreeRow>(TableRow);
+		TSharedPtr<SUsdStagePreviewTree> Tree = Row ? Row->GetOwnerTree() : nullptr;
 
-		if ( Tree && TreeItem )
+		if (Tree && TreeItem)
 		{
-			FUsdPrimPreviewModel& Model = Tree->GetModel( TreeItem->ModelIndex );
+			FUsdPrimPreviewModel& Model = Tree->GetModel(TreeItem->ModelIndex);
 
-			return SNew( SBox )
-				.VAlign( VAlign_Center )
+			// clang-format off
+			return SNew(SBox)
+				.VAlign(VAlign_Center)
 				[
-					SNew( STextBlock )
-					.HighlightText( Tree->GetFilterText() )
-					.Text( Model.Name )
+					SNew(STextBlock)
+					.HighlightText(Tree->GetFilterText())
+					.Text(Model.Name)
 				];
+			// clang-format on
 		}
 
 		return SNullWidget::NullWidget;
@@ -137,82 +135,81 @@ public:
 class FUsdPreviewTreeTypeColumn : public FUsdTreeViewColumn
 {
 public:
-	virtual TSharedRef< SWidget > GenerateWidget(
-		const TSharedPtr< IUsdTreeViewItem > InTreeItem,
-		const TSharedPtr< ITableRow > TableRow
-	) override
+	virtual TSharedRef<SWidget> GenerateWidget(const TSharedPtr<IUsdTreeViewItem> InTreeItem, const TSharedPtr<ITableRow> TableRow) override
 	{
-		FUsdPrimPreviewModelViewPtr TreeItem = StaticCastSharedPtr< FUsdPrimPreviewModelView >( InTreeItem );
-		TSharedPtr< SUsdStagePreviewTreeRow > Row = StaticCastSharedPtr< SUsdStagePreviewTreeRow >( TableRow );
-		TSharedPtr< SUsdStagePreviewTree > Tree = Row ? Row->GetOwnerTree() : nullptr;
+		FUsdPrimPreviewModelViewPtr TreeItem = StaticCastSharedPtr<FUsdPrimPreviewModelView>(InTreeItem);
+		TSharedPtr<SUsdStagePreviewTreeRow> Row = StaticCastSharedPtr<SUsdStagePreviewTreeRow>(TableRow);
+		TSharedPtr<SUsdStagePreviewTree> Tree = Row ? Row->GetOwnerTree() : nullptr;
 
-		if ( Tree && TreeItem )
+		if (Tree && TreeItem)
 		{
-			FUsdPrimPreviewModel& Model = Tree->GetModel( TreeItem->ModelIndex );
+			FUsdPrimPreviewModel& Model = Tree->GetModel(TreeItem->ModelIndex);
 
-			return SNew( SBox )
-				.VAlign( VAlign_Center )
+			// clang-format off
+			return SNew(SBox)
+				.VAlign(VAlign_Center)
 				[
-					SNew( STextBlock )
-					.HighlightText( Tree->GetFilterText() )
-					.Text( Model.TypeName )
+					SNew(STextBlock)
+					.HighlightText(Tree->GetFilterText())
+					.Text(Model.TypeName)
 				];
+			// clang-format on
 		}
 
 		return SNullWidget::NullWidget;
 	}
 };
 
-void SUsdStagePreviewTree::Construct( const FArguments& InArgs, const UE::FUsdStage& InStage )
+void SUsdStagePreviewTree::Construct(const FArguments& InArgs, const UE::FUsdStage& InStage)
 {
-	SUsdTreeView::Construct( SUsdTreeView::FArguments() );
+	SUsdTreeView::Construct(SUsdTreeView::FArguments());
 
-	OnContextMenuOpening = FOnContextMenuOpening::CreateSP( this, &SUsdStagePreviewTree::ConstructPrimContextMenu );
+	OnContextMenuOpening = FOnContextMenuOpening::CreateSP(this, &SUsdStagePreviewTree::ConstructPrimContextMenu);
 	OnExpansionChanged = FOnExpansionChanged::CreateLambda(
-		[this]( FUsdPrimPreviewModelViewRef Item, bool bIsExpanded )
+		[this](FUsdPrimPreviewModelViewRef Item, bool bIsExpanded)
 		{
-			ItemModels[ Item->ModelIndex ].bIsExpanded = bIsExpanded;
+			ItemModels[Item->ModelIndex].bIsExpanded = bIsExpanded;
 		}
 	);
 	OnSetExpansionRecursive = FOnSetExpansionRecursive::CreateLambda(
-		[this]( FUsdPrimPreviewModelViewRef Item, bool bShouldExpand )
+		[this](FUsdPrimPreviewModelViewRef Item, bool bShouldExpand)
 		{
-			ExpandItemsRecursively( { Item }, bShouldExpand );
+			ExpandItemsRecursively({Item}, bShouldExpand);
 		}
 	);
 
-	if ( InStage )
+	if (InStage)
 	{
-		if ( UE::FUsdPrim RootPrim = InStage.GetPseudoRoot() )
+		if (UE::FUsdPrim RootPrim = InStage.GetPseudoRoot())
 		{
-			TFunction<int32( const UE::FUsdPrim& )> RecursivelyAddModels;
-			RecursivelyAddModels = [this, &RecursivelyAddModels]( const UE::FUsdPrim& Prim ) -> int32
+			TFunction<int32(const UE::FUsdPrim&)> RecursivelyAddModels;
+			RecursivelyAddModels = [this, &RecursivelyAddModels](const UE::FUsdPrim& Prim) -> int32
 			{
 				int32 NewIndex = ItemModels.Num();
 				FUsdPrimPreviewModel& NewPrimModel = ItemModels.Emplace_GetRef();
-				NewPrimModel.Name = FText::FromName( Prim.GetName() );
-				NewPrimModel.TypeName = FText::FromName( Prim.GetTypeName() );
+				NewPrimModel.Name = FText::FromName(Prim.GetName());
+				NewPrimModel.TypeName = FText::FromName(Prim.GetTypeName());
 
 				TArray<UE::FUsdPrim> ChildPrims = Prim.GetChildren();
-				NewPrimModel.ChildIndices.Reset( ChildPrims.Num() );
-				ItemModels.Reserve( ItemModels.Num() + ChildPrims.Num() );
-				for ( const UE::FUsdPrim& ChildPrim : ChildPrims )
+				NewPrimModel.ChildIndices.Reset(ChildPrims.Num());
+				ItemModels.Reserve(ItemModels.Num() + ChildPrims.Num());
+				for (const UE::FUsdPrim& ChildPrim : ChildPrims)
 				{
 					// WARNING: This will likely invalidate the NewPrimModel reference! Don't use it past this line!
-					int32 NewChildIndex = RecursivelyAddModels( ChildPrim );
+					int32 NewChildIndex = RecursivelyAddModels(ChildPrim);
 
-					ItemModels[ NewChildIndex ].ParentIndex = NewIndex;
-					ItemModels[ NewIndex ].ChildIndices.Add( NewChildIndex );
+					ItemModels[NewChildIndex].ParentIndex = NewIndex;
+					ItemModels[NewIndex].ChildIndices.Add(NewChildIndex);
 				}
 
 				return NewIndex;
 			};
 
 			TArray<UE::FUsdPrim> TopLevelPrims = RootPrim.GetChildren();
-			ItemModels.Reset( TopLevelPrims.Num() ); // Not the best guess but better than nothing
-			for ( const UE::FUsdPrim& TopLevelPrim : TopLevelPrims )
+			ItemModels.Reset(TopLevelPrims.Num());	  // Not the best guess but better than nothing
+			for (const UE::FUsdPrim& TopLevelPrim : TopLevelPrims)
 			{
-				RecursivelyAddModels( TopLevelPrim );
+				RecursivelyAddModels(TopLevelPrim);
 			}
 
 			RebuildModelViews();
@@ -228,36 +225,33 @@ TArray<FString> SUsdStagePreviewTree::GetSelectedFullPrimPaths() const
 	// This basically amounts to the rule that if a prim P has all of its descendants selected, we put P on the list
 	// and no other descendant. Otherwise, we list all of the individual descendants that are enabled
 
-	const static FString Divider = TEXT( "/" );
+	const static FString Divider = TEXT("/");
 
 	// Since this is the default case, it may be worth it to try a quick scan to see if we're just trying to import
 	// everything. If that's the case all we need to do is return the root prim path
 	bool bShouldImportAllPrims = true;
-	for ( const FUsdPrimPreviewModel& Model : ItemModels )
+	for (const FUsdPrimPreviewModel& Model : ItemModels)
 	{
-		if ( !Model.bShouldImport )
+		if (!Model.bShouldImport)
 		{
 			bShouldImportAllPrims = false;
 			break;
 		}
 	}
-	if ( bShouldImportAllPrims )
+	if (bShouldImportAllPrims)
 	{
-		return { Divider };
+		return {Divider};
 	}
 
-	TFunction<void( const TArray<int32>&, FString, TArray<FString>&, bool& )> RecursivelyCollectPaths;
-	RecursivelyCollectPaths = [this, &RecursivelyCollectPaths](
-		const TArray<int32>& Indices,
-		FString PathSoFar,
-		TArray<FString>& OutChildPaths,
-		bool& bOutAllDescendantsImported
-	)
+	TFunction<void(const TArray<int32>&, FString, TArray<FString>&, bool&)> RecursivelyCollectPaths;
+	RecursivelyCollectPaths =
+		[this,
+		 &RecursivelyCollectPaths](const TArray<int32>& Indices, FString PathSoFar, TArray<FString>& OutChildPaths, bool& bOutAllDescendantsImported)
 	{
-		for ( int32 Index : Indices )
+		for (int32 Index : Indices)
 		{
-			const FUsdPrimPreviewModel& Model = ItemModels[ Index ];
-			if ( !Model.bShouldImport )
+			const FUsdPrimPreviewModel& Model = ItemModels[Index];
+			if (!Model.bShouldImport)
 			{
 				bOutAllDescendantsImported = false;
 				continue;
@@ -267,24 +261,19 @@ TArray<FString> SUsdStagePreviewTree::GetSelectedFullPrimPaths() const
 
 			TArray<FString> ChildPaths;
 			bool bAllDescendantsImported = true;
-			RecursivelyCollectPaths(
-				Model.ChildIndices,
-				PrimFullPath + Divider,
-				ChildPaths,
-				bAllDescendantsImported
-			);
+			RecursivelyCollectPaths(Model.ChildIndices, PrimFullPath + Divider, ChildPaths, bAllDescendantsImported);
 
 			// If all of the Prim's descendents were imported, fully discard ChildPaths and just add Prim directly
 			// to the list.
 			// Alternatively, if we imported *nothing else* but we made it this far (i.e. we want to import this prim)
 			// then we actually want to import this intermediate prim directly, so also add it directly to the list.
-			if ( bAllDescendantsImported || ChildPaths.Num() == 0 )
+			if (bAllDescendantsImported || ChildPaths.Num() == 0)
 			{
-				OutChildPaths.Add( MoveTemp( PrimFullPath ) );
+				OutChildPaths.Add(MoveTemp(PrimFullPath));
 			}
 			else
 			{
-				OutChildPaths.Append( MoveTemp( ChildPaths ) );
+				OutChildPaths.Append(MoveTemp(ChildPaths));
 			}
 
 			bOutAllDescendantsImported &= bAllDescendantsImported;
@@ -292,36 +281,30 @@ TArray<FString> SUsdStagePreviewTree::GetSelectedFullPrimPaths() const
 	};
 
 	TArray<int32> RootItemModels;
-	RootItemModels.Reserve( RootItems.Num() );
-	for ( const FUsdPrimPreviewModelViewRef& RootItem : RootItems )
+	RootItemModels.Reserve(RootItems.Num());
+	for (const FUsdPrimPreviewModelViewRef& RootItem : RootItems)
 	{
-		RootItemModels.Add( RootItem->ModelIndex );
+		RootItemModels.Add(RootItem->ModelIndex);
 	}
 
 	TArray<FString> Result;
-	bool bUnusedDummy = false; // At this point we know we won't import *everything* anyway, don't bother checking
-	RecursivelyCollectPaths( RootItemModels, Divider, Result, bUnusedDummy );
+	bool bUnusedDummy = false;	  // At this point we know we won't import *everything* anyway, don't bother checking
+	RecursivelyCollectPaths(RootItemModels, Divider, Result, bUnusedDummy);
 	return Result;
 }
 
-void SUsdStagePreviewTree::SetFilterText( const FText& NewText )
+void SUsdStagePreviewTree::SetFilterText(const FText& NewText)
 {
 	CurrentFilterText = NewText;
 	RebuildModelViews();
 }
 
-TSharedRef< ITableRow > SUsdStagePreviewTree::OnGenerateRow(
-	FUsdPrimPreviewModelViewRef InDisplayNode,
-	const TSharedRef<STableViewBase>& OwnerTable
-)
+TSharedRef<ITableRow> SUsdStagePreviewTree::OnGenerateRow(FUsdPrimPreviewModelViewRef InDisplayNode, const TSharedRef<STableViewBase>& OwnerTable)
 {
-	return SNew( SUsdStagePreviewTreeRow, InDisplayNode, OwnerTable, SharedData );
+	return SNew(SUsdStagePreviewTreeRow, InDisplayNode, OwnerTable, SharedData);
 }
 
-void SUsdStagePreviewTree::OnGetChildren(
-	FUsdPrimPreviewModelViewRef InParent,
-	TArray< FUsdPrimPreviewModelViewRef >& OutChildren
-) const
+void SUsdStagePreviewTree::OnGetChildren(FUsdPrimPreviewModelViewRef InParent, TArray<FUsdPrimPreviewModelViewRef>& OutChildren) const
 {
 	OutChildren = InParent->Children;
 }
@@ -332,114 +315,118 @@ void SUsdStagePreviewTree::SetupColumns()
 
 	// Import checkbox column with a checkbox on the header itself (so we can't just use AddColumn)
 	{
-		const static FName ImportColumnKey = TEXT( "ImportColumn" );
+		const static FName ImportColumnKey = TEXT("ImportColumn");
 
-		SharedData->Columns.Add( ImportColumnKey, MakeShared< FUsdPreviewTreeImportColumn >() );
+		SharedData->Columns.Add(ImportColumnKey, MakeShared<FUsdPreviewTreeImportColumn>());
 
+		// clang-format off
 		HeaderRowWidget->AddColumn(
-			SHeaderRow::Column( ImportColumnKey )
-			.FixedWidth( UE::UsdStagePreviewTree::Private::CheckBoxColumnWidth )
-			.HeaderContentPadding( FMargin( 0 ) )
+			SHeaderRow::Column(ImportColumnKey)
+			.FixedWidth(UE::UsdStagePreviewTree::Private::CheckBoxColumnWidth)
+			.HeaderContentPadding(FMargin(0))
 			[
-				SNew( SBox )
-				.VAlign( VAlign_Center )
-				.HAlign( HAlign_Center )
-				.WidthOverride( UE::UsdStagePreviewTree::Private::CheckBoxColumnWidth )
+				SNew(SBox)
+				.VAlign(VAlign_Center)
+				.HAlign(HAlign_Center)
+				.WidthOverride(UE::UsdStagePreviewTree::Private::CheckBoxColumnWidth)
 				[
-					SNew( SCheckBox )
-					.ToolTipText( LOCTEXT( "HeaderCheckboxToolTip", "Toggle all currently displayed items" ) )
-					.IsChecked_Lambda( [this]() -> ECheckBoxState
+					SNew(SCheckBox)
+					.ToolTipText(LOCTEXT("HeaderCheckboxToolTip", "Toggle all currently displayed items"))
+					.IsChecked_Lambda([this]() -> ECheckBoxState
 					{
-						return UE::UsdStagePreviewTree::Private::AreAllViewsCheckedRecursive( *this, RootItems )
+						return UE::UsdStagePreviewTree::Private::AreAllViewsCheckedRecursive(*this, RootItems)
 							? ECheckBoxState::Checked
 							: ECheckBoxState::Unchecked;
 					})
-					.OnCheckStateChanged_Lambda( [this]( ECheckBoxState NewState )
+					.OnCheckStateChanged_Lambda([this](ECheckBoxState NewState)
 					{
 						const bool bChecked = NewState == ECheckBoxState::Checked;
-						CheckItemsRecursively( RootItems, bChecked );
+						CheckItemsRecursively(RootItems, bChecked);
 					})
 				]
 			]
 		);
+		// clang-format on
 	}
 
 	// Prim name column
 	{
-		TSharedRef< FUsdPreviewTreeNameColumn > PrimNameColumn = MakeShared< FUsdPreviewTreeNameColumn >();
+		TSharedRef<FUsdPreviewTreeNameColumn> PrimNameColumn = MakeShared<FUsdPreviewTreeNameColumn>();
 		PrimNameColumn->bIsMainColumn = true;
 
 		SHeaderRow::FColumn::FArguments Arguments;
-		Arguments.FillWidth( 70.f );
+		Arguments.FillWidth(70.f);
 
-		AddColumn( TEXT("Prim"), LOCTEXT( "Prim", "Prim" ), PrimNameColumn, Arguments );
+		AddColumn(TEXT("Prim"), LOCTEXT("Prim", "Prim"), PrimNameColumn, Arguments);
 	}
 
 	// Prim type column
 	{
 		SHeaderRow::FColumn::FArguments Arguments;
-		Arguments.FillWidth( 20.f );
-		AddColumn( TEXT( "Type" ), LOCTEXT( "Type", "Type" ), MakeShared< FUsdPreviewTreeTypeColumn >(), Arguments );
+		Arguments.FillWidth(20.f);
+		AddColumn(TEXT("Type"), LOCTEXT("Type", "Type"), MakeShared<FUsdPreviewTreeTypeColumn>(), Arguments);
 	}
 }
 
 TArray<FUsdPrimPreviewModelViewRef> SUsdStagePreviewTree::GetAncestorSelectedViews()
 {
-	TArray< FUsdPrimPreviewModelViewRef > SelectedViews = GetSelectedItems();
+	TArray<FUsdPrimPreviewModelViewRef> SelectedViews = GetSelectedItems();
 
 	TSet<int32> SelectedIndices;
-	SelectedIndices.Reserve( SelectedViews.Num() );
-	for ( const FUsdPrimPreviewModelViewRef& SelectedItem : SelectedViews )
+	SelectedIndices.Reserve(SelectedViews.Num());
+	for (const FUsdPrimPreviewModelViewRef& SelectedItem : SelectedViews)
 	{
-		SelectedIndices.Add( SelectedItem->ModelIndex );
+		SelectedIndices.Add(SelectedItem->ModelIndex);
 	}
 
-	TArray< FUsdPrimPreviewModelViewRef > AncestorSelectedViews;
-	AncestorSelectedViews.Reserve( SelectedViews.Num() );
+	TArray<FUsdPrimPreviewModelViewRef> AncestorSelectedViews;
+	AncestorSelectedViews.Reserve(SelectedViews.Num());
 
-	for ( const FUsdPrimPreviewModelViewRef& SelectedView : SelectedViews )
+	for (const FUsdPrimPreviewModelViewRef& SelectedView : SelectedViews)
 	{
 		bool bIsAncestor = true;
 
-		int32 ParentIndex = ItemModels[ SelectedView->ModelIndex ].ParentIndex;
-		while ( ParentIndex != INDEX_NONE )
+		int32 ParentIndex = ItemModels[SelectedView->ModelIndex].ParentIndex;
+		while (ParentIndex != INDEX_NONE)
 		{
-			if ( SelectedIndices.Contains( ParentIndex ) )
+			if (SelectedIndices.Contains(ParentIndex))
 			{
 				bIsAncestor = false;
 				break;
 			}
 
-			ParentIndex = ItemModels[ ParentIndex ].ParentIndex;
+			ParentIndex = ItemModels[ParentIndex].ParentIndex;
 		}
 
-		if ( bIsAncestor )
+		if (bIsAncestor)
 		{
-			AncestorSelectedViews.Add( SelectedView );
+			AncestorSelectedViews.Add(SelectedView);
 		}
 	}
 
 	return AncestorSelectedViews;
 }
 
-TSharedPtr< SWidget > SUsdStagePreviewTree::ConstructPrimContextMenu()
+TSharedPtr<SWidget> SUsdStagePreviewTree::ConstructPrimContextMenu()
 {
 	const bool bInShouldCloseWindowAfterMenuSelection = true;
-	const TSharedPtr< const FUICommandList > CommandList = nullptr;
-	FMenuBuilder MenuBuilder{ bInShouldCloseWindowAfterMenuSelection, CommandList };
+	const TSharedPtr<const FUICommandList> CommandList = nullptr;
+	FMenuBuilder MenuBuilder{bInShouldCloseWindowAfterMenuSelection, CommandList};
 
-	MenuBuilder.BeginSection( "Import", LOCTEXT( "ImportText", "Import" ) );
+	MenuBuilder.BeginSection("Import", LOCTEXT("ImportText", "Import"));
 	{
 		MenuBuilder.AddMenuEntry(
-			LOCTEXT( "CheckDescendants_Text", "Check Recursively" ),
-			LOCTEXT( "CheckDescendants_ToolTip", "Check this prim and all displayed descendants recursively" ),
+			LOCTEXT("CheckDescendants_Text", "Check Recursively"),
+			LOCTEXT("CheckDescendants_ToolTip", "Check this prim and all displayed descendants recursively"),
 			FSlateIcon(),
 			FUIAction(
-				FExecuteAction::CreateLambda( [this]()
-				{
-					const bool bShouldSelect = true;
-					CheckItemsRecursively( GetAncestorSelectedViews(), bShouldSelect );
-				}),
+				FExecuteAction::CreateLambda(
+					[this]()
+					{
+						const bool bShouldSelect = true;
+						CheckItemsRecursively(GetAncestorSelectedViews(), bShouldSelect);
+					}
+				),
 				FCanExecuteAction()
 			),
 			NAME_None,
@@ -447,15 +434,17 @@ TSharedPtr< SWidget > SUsdStagePreviewTree::ConstructPrimContextMenu()
 		);
 
 		MenuBuilder.AddMenuEntry(
-			LOCTEXT( "UncheckDescendants_Text", "Uncheck Recursively" ),
-			LOCTEXT( "UncheckDescendants_ToolTip", "Uncheck this prim and all displayed descendants recursively" ),
+			LOCTEXT("UncheckDescendants_Text", "Uncheck Recursively"),
+			LOCTEXT("UncheckDescendants_ToolTip", "Uncheck this prim and all displayed descendants recursively"),
 			FSlateIcon(),
 			FUIAction(
-				FExecuteAction::CreateLambda( [this]()
-				{
-					const bool bShouldSelect = false;
-					CheckItemsRecursively( GetAncestorSelectedViews(), bShouldSelect );
-				}),
+				FExecuteAction::CreateLambda(
+					[this]()
+					{
+						const bool bShouldSelect = false;
+						CheckItemsRecursively(GetAncestorSelectedViews(), bShouldSelect);
+					}
+				),
 				FCanExecuteAction()
 			),
 			NAME_None,
@@ -464,18 +453,20 @@ TSharedPtr< SWidget > SUsdStagePreviewTree::ConstructPrimContextMenu()
 	}
 	MenuBuilder.EndSection();
 
-	MenuBuilder.BeginSection( "Expansion", LOCTEXT( "ExpansionText", "Expansion" ) );
+	MenuBuilder.BeginSection("Expansion", LOCTEXT("ExpansionText", "Expansion"));
 	{
 		MenuBuilder.AddMenuEntry(
-			LOCTEXT( "ExpandDescendants_Text", "Expand Recursively" ),
-			LOCTEXT( "ExpandDescendants_ToolTip", "Expand this prim and all displayed descendants recursively" ),
+			LOCTEXT("ExpandDescendants_Text", "Expand Recursively"),
+			LOCTEXT("ExpandDescendants_ToolTip", "Expand this prim and all displayed descendants recursively"),
 			FSlateIcon(),
 			FUIAction(
-				FExecuteAction::CreateLambda( [this]()
-				{
-					const bool bShouldExpand = true;
-					ExpandItemsRecursively( GetAncestorSelectedViews(), bShouldExpand );
-				}),
+				FExecuteAction::CreateLambda(
+					[this]()
+					{
+						const bool bShouldExpand = true;
+						ExpandItemsRecursively(GetAncestorSelectedViews(), bShouldExpand);
+					}
+				),
 				FCanExecuteAction()
 			),
 			NAME_None,
@@ -483,15 +474,17 @@ TSharedPtr< SWidget > SUsdStagePreviewTree::ConstructPrimContextMenu()
 		);
 
 		MenuBuilder.AddMenuEntry(
-			LOCTEXT( "CollapseDescendants_Text", "Collapse Recursively" ),
-			LOCTEXT( "CollapseDescendants_ToolTip", "Collapse this prim and all displayed descendants recursively" ),
+			LOCTEXT("CollapseDescendants_Text", "Collapse Recursively"),
+			LOCTEXT("CollapseDescendants_ToolTip", "Collapse this prim and all displayed descendants recursively"),
 			FSlateIcon(),
 			FUIAction(
-				FExecuteAction::CreateLambda( [this]()
-				{
-					const bool bShouldExpand = false;
-					ExpandItemsRecursively( GetAncestorSelectedViews(), bShouldExpand );
-				}),
+				FExecuteAction::CreateLambda(
+					[this]()
+					{
+						const bool bShouldExpand = false;
+						ExpandItemsRecursively(GetAncestorSelectedViews(), bShouldExpand);
+					}
+				),
 				FCanExecuteAction()
 			),
 			NAME_None,
@@ -508,20 +501,19 @@ void SUsdStagePreviewTree::RebuildModelViews()
 	const FString CurrentFilterString = CurrentFilterText.ToString().ToLower();
 
 	// Update our models caching whether they pass the filter or not
-	for ( FUsdPrimPreviewModel& Model : ItemModels )
+	for (FUsdPrimPreviewModel& Model : ItemModels)
 	{
-		Model.bPassesFilter = CurrentFilterString.IsEmpty()
-			|| Model.Name.ToString().ToLower().Contains( CurrentFilterString )
-			|| Model.TypeName.ToString().ToLower().Contains( CurrentFilterString );
+		Model.bPassesFilter = CurrentFilterString.IsEmpty() || Model.Name.ToString().ToLower().Contains(CurrentFilterString)
+							  || Model.TypeName.ToString().ToLower().Contains(CurrentFilterString);
 
 		// If this model passes the filter, we need to ensure all of its ancestors are marked as passing too
-		if ( Model.bPassesFilter )
+		if (Model.bPassesFilter)
 		{
 			int32 ParentIndex = Model.ParentIndex;
-			while ( ParentIndex != INDEX_NONE )
+			while (ParentIndex != INDEX_NONE)
 			{
-				FUsdPrimPreviewModel& ParentModel = ItemModels[ ParentIndex ];
-				if ( ParentModel.bPassesFilter )
+				FUsdPrimPreviewModel& ParentModel = ItemModels[ParentIndex];
+				if (ParentModel.bPassesFilter)
 				{
 					break;
 				}
@@ -533,114 +525,111 @@ void SUsdStagePreviewTree::RebuildModelViews()
 
 	// Rebuild and keep track of this mapping as its very cheap to do right now
 	ItemModelsToViews.Reset();
-	ItemModelsToViews.Reserve( ItemModels.Num() );
+	ItemModelsToViews.Reserve(ItemModels.Num());
 
 	// Create a view for each model *that passes the filter*, having nullptr instead if they don't
-	TArray< FUsdPrimPreviewModelViewPtr > LinearViews;
-	LinearViews.SetNum( ItemModels.Num() );
-	for ( int32 Index = 0; Index < ItemModels.Num(); ++Index )
+	TArray<FUsdPrimPreviewModelViewPtr> LinearViews;
+	LinearViews.SetNum(ItemModels.Num());
+	for (int32 Index = 0; Index < ItemModels.Num(); ++Index)
 	{
-		const FUsdPrimPreviewModel& Model = ItemModels[ Index ];
+		const FUsdPrimPreviewModel& Model = ItemModels[Index];
 
-		if ( Model.bPassesFilter )
+		if (Model.bPassesFilter)
 		{
 			FUsdPrimPreviewModelViewRef View = MakeShared<FUsdPrimPreviewModelView>();
 			View->ModelIndex = Index;
 
-			ItemModelsToViews.Add( Index, View );
-			LinearViews[ Index ] = MoveTemp( View );
+			ItemModelsToViews.Add(Index, View);
+			LinearViews[Index] = MoveTemp(View);
 		}
 	}
 
 	// Now that we have views created for each model we want, connect parent/child based on model indices
 	RootItems.Reset();
 	SparseItemInfos.Reset();
-	for ( int32 Index = 0; Index < ItemModels.Num(); ++Index )
+	for (int32 Index = 0; Index < ItemModels.Num(); ++Index)
 	{
-		const FUsdPrimPreviewModel& Model = ItemModels[ Index ];
-		FUsdPrimPreviewModelViewPtr& View = LinearViews[ Index ];
-		if ( !View )
+		const FUsdPrimPreviewModel& Model = ItemModels[Index];
+		FUsdPrimPreviewModelViewPtr& View = LinearViews[Index];
+		if (!View)
 		{
 			// Model didn't pass the filter
 			continue;
 		}
 
 		// Set the view's parent
-		if ( Model.ParentIndex == INDEX_NONE )
+		if (Model.ParentIndex == INDEX_NONE)
 		{
-			RootItems.Add( View.ToSharedRef() );
+			RootItems.Add(View.ToSharedRef());
 		}
 		else
 		{
-			View->Parent = LinearViews[ Model.ParentIndex ];
+			View->Parent = LinearViews[Model.ParentIndex];
 		}
 
 		// Add child models it they pass the filter, also already tracking if they're expanded or not
 		bool bHasExpandedChildren = false;
-		View->Children.Reserve( Model.ChildIndices.Num() );
-		for ( int32 ChildIndex : Model.ChildIndices )
+		View->Children.Reserve(Model.ChildIndices.Num());
+		for (int32 ChildIndex : Model.ChildIndices)
 		{
-			const FUsdPrimPreviewModel& ChildModel = ItemModels[ ChildIndex ];
-			if ( ChildModel.bPassesFilter )
+			const FUsdPrimPreviewModel& ChildModel = ItemModels[ChildIndex];
+			if (ChildModel.bPassesFilter)
 			{
-				View->Children.Add( LinearViews[ ChildIndex ].ToSharedRef() );
+				View->Children.Add(LinearViews[ChildIndex].ToSharedRef());
 				bHasExpandedChildren |= ChildModel.bIsExpanded;
 			}
 		}
 
-		const bool bShouldExpand = true;
-		SparseItemInfos.Add( View.ToSharedRef(), FSparseItemInfo( Model.bIsExpanded, bHasExpandedChildren ) );
+		SparseItemInfos.Add(View.ToSharedRef(), FSparseItemInfo(Model.bIsExpanded, bHasExpandedChildren));
 	}
 
 	RequestTreeRefresh();
 }
 
-void SUsdStagePreviewTree::CheckItemsRecursively( const TArray<FUsdPrimPreviewModelViewRef>& Items, bool bCheck )
+void SUsdStagePreviewTree::CheckItemsRecursively(const TArray<FUsdPrimPreviewModelViewRef>& Items, bool bCheck)
 {
-	TFunction< void(const TArray<FUsdPrimPreviewModelViewRef>&)> CheckViewsRecursively;
-	CheckViewsRecursively = [this, &CheckViewsRecursively]( const TArray<FUsdPrimPreviewModelViewRef>& Views )
+	TFunction<void(const TArray<FUsdPrimPreviewModelViewRef>&)> CheckViewsRecursively;
+	CheckViewsRecursively = [this, &CheckViewsRecursively](const TArray<FUsdPrimPreviewModelViewRef>& Views)
 	{
-		for ( const FUsdPrimPreviewModelViewRef& View : Views )
+		for (const FUsdPrimPreviewModelViewRef& View : Views)
 		{
-			ItemModels[ View->ModelIndex ].bShouldImport = true;
-			CheckViewsRecursively( View->Children );
+			ItemModels[View->ModelIndex].bShouldImport = true;
+			CheckViewsRecursively(View->Children);
 		}
 	};
 
 	TArray<FString> UncheckedInvisiblePrims;
-	TFunction< void( const TArray<int32>& )> UncheckModelsRecursively;
-	UncheckModelsRecursively = [this, &UncheckModelsRecursively, &UncheckedInvisiblePrims](
-		const TArray<int32>& ModelIndices
-	)
+	TFunction<void(const TArray<int32>&)> UncheckModelsRecursively;
+	UncheckModelsRecursively = [this, &UncheckModelsRecursively, &UncheckedInvisiblePrims](const TArray<int32>& ModelIndices)
 	{
-		for ( int32 ModelIndex : ModelIndices )
+		for (int32 ModelIndex : ModelIndices)
 		{
-			FUsdPrimPreviewModel& Model = ItemModels[ ModelIndex ];
+			FUsdPrimPreviewModel& Model = ItemModels[ModelIndex];
 
-			if ( Model.bShouldImport && !ItemModelsToViews.Contains( ModelIndex ) )
+			if (Model.bShouldImport && !ItemModelsToViews.Contains(ModelIndex))
 			{
-				UncheckedInvisiblePrims.Add( Model.Name.ToString() );
+				UncheckedInvisiblePrims.Add(Model.Name.ToString());
 			}
 
 			Model.bShouldImport = false;
-			UncheckModelsRecursively( Model.ChildIndices );
+			UncheckModelsRecursively(Model.ChildIndices);
 		}
 	};
 
 	// Attention here: For consistency, if we are unchecking we must guarantee that no descendant models are ever left
 	// checked, whether we have views for them or not. On the other hand if we are just checking something, it is fine
 	// if not all of the model descendants are checked in the process, so let's only check what is shown on the screen
-	if ( bCheck )
+	if (bCheck)
 	{
-		CheckViewsRecursively( Items );
+		CheckViewsRecursively(Items);
 
 		// If we want to check a view, all of its ancestor *models* must also be checked, regardless of what is shown
-		for ( const FUsdPrimPreviewModelViewRef& View : Items )
+		for (const FUsdPrimPreviewModelViewRef& View : Items)
 		{
-			int32 ParentIndex = ItemModels[ View->ModelIndex ].ParentIndex;
-			while ( ParentIndex != INDEX_NONE )
+			int32 ParentIndex = ItemModels[View->ModelIndex].ParentIndex;
+			while (ParentIndex != INDEX_NONE)
 			{
-				FUsdPrimPreviewModel& ParentModel = ItemModels[ ParentIndex ];
+				FUsdPrimPreviewModel& ParentModel = ItemModels[ParentIndex];
 				ParentModel.bShouldImport = true;
 				ParentIndex = ParentModel.ParentIndex;
 			}
@@ -649,73 +638,70 @@ void SUsdStagePreviewTree::CheckItemsRecursively( const TArray<FUsdPrimPreviewMo
 	else
 	{
 		TArray<int32> ModelIndices;
-		ModelIndices.Reserve( Items.Num() );
-		for ( const FUsdPrimPreviewModelViewRef& Item : Items )
+		ModelIndices.Reserve(Items.Num());
+		for (const FUsdPrimPreviewModelViewRef& Item : Items)
 		{
-			ModelIndices.Add( Item->ModelIndex );
+			ModelIndices.Add(Item->ModelIndex);
 		}
-		UncheckModelsRecursively( ModelIndices );
+		UncheckModelsRecursively(ModelIndices);
 
 		// Notify user in case we unchecked some prims that currently are not being shown
-		if( UncheckedInvisiblePrims.Num() > 0 )
+		if (UncheckedInvisiblePrims.Num() > 0)
 		{
 			const int32 MaxNumPrimsToList = 15;
-			const static FString Delimiter = TEXT( ", " );
+			const static FString Delimiter = TEXT(", ");
 			FString PrimListString;
-			for ( int32 Index = 0; Index < FMath::Min( UncheckedInvisiblePrims.Num(), MaxNumPrimsToList ); ++Index )
+			for (int32 Index = 0; Index < FMath::Min(UncheckedInvisiblePrims.Num(), MaxNumPrimsToList); ++Index)
 			{
-				PrimListString += UncheckedInvisiblePrims[ Index ] + Delimiter;
+				PrimListString += UncheckedInvisiblePrims[Index] + Delimiter;
 			}
-			if ( UncheckedInvisiblePrims.Num() > MaxNumPrimsToList )
+			if (UncheckedInvisiblePrims.Num() > MaxNumPrimsToList)
 			{
-				PrimListString += TEXT( "..." );
+				PrimListString += TEXT("...");
 			}
 			else
 			{
-				PrimListString.RemoveFromEnd( Delimiter );
+				PrimListString.RemoveFromEnd(Delimiter);
 			}
 
-			FNotificationInfo Toast( LOCTEXT( "UncheckedInvisibleTitle", "USD Prims to import" ) );
+			FNotificationInfo Toast(LOCTEXT("UncheckedInvisibleTitle", "USD Prims to import"));
 			Toast.SubText = FText::Format(
 				UncheckedInvisiblePrims.Num() == 1
-					? LOCTEXT( "UncheckedInvisibleSingular", "This prim is not currently shown, but was also unchecked:\n\n{0}" )
-					: LOCTEXT( "UncheckedInvisiblePlural", "These prims are not currently shown, but were also unchecked:\n\n{0}" ),
-				FText::FromString( PrimListString )
+					? LOCTEXT("UncheckedInvisibleSingular", "This prim is not currently shown, but was also unchecked:\n\n{0}")
+					: LOCTEXT("UncheckedInvisiblePlural", "These prims are not currently shown, but were also unchecked:\n\n{0}"),
+				FText::FromString(PrimListString)
 			);
-			Toast.Image = FCoreStyle::Get().GetBrush( TEXT( "MessageLog.Warning" ) );
+			Toast.Image = FCoreStyle::Get().GetBrush(TEXT("MessageLog.Warning"));
 			Toast.bUseLargeFont = false;
 			Toast.bFireAndForget = true;
 			Toast.FadeOutDuration = 0.5f;
 			Toast.ExpireDuration = 4.5f;
 			Toast.bUseThrobber = true;
 			Toast.bUseSuccessFailIcons = false;
-			FSlateNotificationManager::Get().AddNotification( Toast );
+			FSlateNotificationManager::Get().AddNotification(Toast);
 		}
 	}
 }
 
-void SUsdStagePreviewTree::ExpandItemsRecursively(
-	const TArray<FUsdPrimPreviewModelViewRef>& Items,
-	bool bExpand
-)
+void SUsdStagePreviewTree::ExpandItemsRecursively(const TArray<FUsdPrimPreviewModelViewRef>& Items, bool bExpand)
 {
-	TFunction< void( const TArray<FUsdPrimPreviewModelViewRef>& )> ExpandRecursively;
-	ExpandRecursively = [this, &ExpandRecursively, bExpand]( const TArray<FUsdPrimPreviewModelViewRef>& Views )
+	TFunction<void(const TArray<FUsdPrimPreviewModelViewRef>&)> ExpandRecursively;
+	ExpandRecursively = [this, &ExpandRecursively, bExpand](const TArray<FUsdPrimPreviewModelViewRef>& Views)
 	{
-		for ( const FUsdPrimPreviewModelViewRef& View : Views )
+		for (const FUsdPrimPreviewModelViewRef& View : Views)
 		{
-			FUsdPrimPreviewModel& Model = ItemModels[ View->ModelIndex ];
+			FUsdPrimPreviewModel& Model = ItemModels[View->ModelIndex];
 			Model.bIsExpanded = bExpand;
 
 			// We're going to expand recursively anyway, so we know it will have expanded children
 			const bool bHasExpandedChildren = bExpand && View->Children.Num() > 0;
 			const bool bIsExpanded = bExpand;
-			SparseItemInfos.Add( View, FSparseItemInfo( bIsExpanded, bHasExpandedChildren ) );
+			SparseItemInfos.Add(View, FSparseItemInfo(bIsExpanded, bHasExpandedChildren));
 
-			ExpandRecursively( View->Children );
+			ExpandRecursively(View->Children);
 		}
 	};
-	ExpandRecursively( Items );
+	ExpandRecursively(Items);
 
 	RequestTreeRefresh();
 }
@@ -727,8 +713,8 @@ void SUsdStagePreviewTreeRow::Construct(
 	TSharedPtr<FSharedUsdTreeData> InSharedData
 )
 {
-	SUsdTreeRow< FUsdPrimPreviewModelViewRef >::Construct(
-		typename SUsdTreeRow< FUsdPrimPreviewModelViewRef >::FArguments(),
+	SUsdTreeRow<FUsdPrimPreviewModelViewRef>::Construct(
+		typename SUsdTreeRow<FUsdPrimPreviewModelViewRef>::FArguments(),
 		InTreeItem,
 		OwnerTable,
 		InSharedData
@@ -737,7 +723,7 @@ void SUsdStagePreviewTreeRow::Construct(
 
 TSharedPtr<SUsdStagePreviewTree> SUsdStagePreviewTreeRow::GetOwnerTree() const
 {
-	return StaticCastSharedPtr< SUsdStagePreviewTree>( OwnerTablePtr.Pin() );
+	return StaticCastSharedPtr<SUsdStagePreviewTree>(OwnerTablePtr.Pin());
 }
 
-#undef LOCTEXT_NAMESPACE
\ No newline at end of file
+#undef LOCTEXT_NAMESPACE
diff --git a/Engine/Plugins/Importers/USDImporter/Source/USDStageImporter/Public/USDStageImportContext.h b/Engine/Plugins/Importers/USDImporter/Source/USDStageImporter/Public/USDStageImportContext.h
--- a/Engine/Plugins/Importers/USDImporter/Source/USDStageImporter/Public/USDStageImportContext.h
+++ b/Engine/Plugins/Importers/USDImporter/Source/USDStageImporter/Public/USDStageImportContext.h
@@ -2,21 +2,19 @@
 
 #pragma once
 
-#include "UnrealUSDWrapper.h"
 #include "USDInfoCache.h"
 #include "USDLevelSequenceHelper.h"
-#include "USDMemory.h"
 #include "USDStageOptions.h"
 
 #include "UsdWrappers/UsdStage.h"
 
-#include "Logging/TokenizedMessage.h"
 #include "CoreMinimal.h"
 
 #include "USDStageImportContext.generated.h"
 
 class UUsdAssetCache2;
 class UUsdStageImportOptions;
+class FTokenizedMessage;
 
 USTRUCT()
 struct USDSTAGEIMPORTER_API FUsdStageImportContext
@@ -71,7 +69,7 @@ struct USDSTAGEIMPORTER_API FUsdStageImportContext
 	 * When parsing materials, we keep track of which primvar we mapped to which UV channel.
 	 * When parsing meshes later, we use this data to place the correct primvar values in each UV channel.
 	 */
-	TMap< FString, TMap< FString, int32 > > MaterialToPrimvarToUVIndex;
+	TMap<FString, TMap<FString, int32>> MaterialToPrimvarToUVIndex;
 
 	/** USD Stage to import */
 	UE::FUsdStage Stage;
@@ -91,7 +89,8 @@ struct USDSTAGEIMPORTER_API FUsdStageImportContext
 	/** If we're reading from the stage cache and the stage was originally open, it will be left open when the import is completed */
 	bool bStageWasOriginallyOpenInCache;
 
-	/** We modify the stage with our meters per unit import option on import. If the stage was already open, we use this to undo the changes after import */
+	/** We modify the stage with our meters per unit import option on import. If the stage was already open, we use this to undo the changes after
+	 * import */
 	float OriginalMetersPerUnit;
 	EUsdUpAxis OriginalUpAxis;
 
@@ -107,7 +106,15 @@ public:
 	FUsdStageImportContext& operator=(const FUsdStageImportContext& Other) = default;
 	PRAGMA_ENABLE_DEPRECATION_WARNINGS
 
-	bool Init(const FString& InName, const FString& InFilePath, const FString& InInitialPackagePath, EObjectFlags InFlags, bool bInIsAutomated, bool bIsReimport = false, bool bAllowActorImport = true);
+	bool Init(
+		const FString& InName,
+		const FString& InFilePath,
+		const FString& InInitialPackagePath,
+		EObjectFlags InFlags,
+		bool bInIsAutomated,
+		bool bIsReimport = false,
+		bool bAllowActorImport = true
+	);
 	void Reset();
 
 private:
diff --git a/Engine/Plugins/Importers/USDImporter/Source/USDStageImporter/Public/USDStageImportOptions.h b/Engine/Plugins/Importers/USDImporter/Source/USDStageImporter/Public/USDStageImportOptions.h
--- a/Engine/Plugins/Importers/USDImporter/Source/USDStageImporter/Public/USDStageImportOptions.h
+++ b/Engine/Plugins/Importers/USDImporter/Source/USDStageImporter/Public/USDStageImportOptions.h
@@ -7,10 +7,7 @@
 #include "USDStageOptions.h"
 
 #include "CoreMinimal.h"
-#include "Engine/EngineTypes.h"
-#include "Factories/MaterialImportHelpers.h"
 #include "GroomAssetInterpolation.h"
-#include "UObject/ObjectMacros.h"
 
 #include "USDStageImportOptions.generated.h"
 
@@ -58,10 +55,16 @@ public:
 	bool bImportGeometry;
 
 	/** Whether to try importing UAnimSequence skeletal animation assets for each encountered UsdSkelAnimQuery */
-	UPROPERTY(BlueprintReadWrite, config, EditAnywhere, Category = "DataToImport", meta = (EditCondition = bImportGeometry, DisplayName = "Skeletal Animations"))
+	UPROPERTY(
+		BlueprintReadWrite,
+		config,
+		EditAnywhere,
+		Category = "DataToImport",
+		meta = (EditCondition = bImportGeometry, DisplayName = "Skeletal Animations")
+	)
 	bool bImportSkeletalAnimations;
 
-	UPROPERTY( BlueprintReadWrite, config, EditAnywhere, Category = "DataToImport", meta = (DisplayName = "LevelSequences" ) )
+	UPROPERTY(BlueprintReadWrite, config, EditAnywhere, Category = "DataToImport", meta = (DisplayName = "LevelSequences"))
 	bool bImportLevelSequences;
 
 	UPROPERTY(BlueprintReadWrite, config, EditAnywhere, Category = "DataToImport", meta = (DisplayName = "Materials & Textures"))
@@ -83,27 +86,33 @@ public:
 	 * Importing a prim will import its entire subtree.
 	 * If this list contains the root prim path the entire stage will be imported (default value).
 	 */
-	UPROPERTY( BlueprintReadWrite, EditAnywhere, Category = "Prims to Import" )
-	TArray<FString> PrimsToImport = TArray<FString>{ TEXT( "/" ) };
+	UPROPERTY(BlueprintReadWrite, EditAnywhere, Category = "Prims to Import")
+	TArray<FString> PrimsToImport = TArray<FString>{TEXT("/")};
 
 	/** Only import prims with these specific purposes from the USD file */
-	UPROPERTY(BlueprintReadWrite, config, EditAnywhere, Category= "USD options", meta = (Bitmask, BitmaskEnum="/Script/UnrealUSDWrapper.EUsdPurpose"))
+	UPROPERTY(
+		BlueprintReadWrite,
+		config,
+		EditAnywhere,
+		Category = "USD options",
+		meta = (Bitmask, BitmaskEnum = "/Script/UnrealUSDWrapper.EUsdPurpose")
+	)
 	int32 PurposesToImport;
 
 	/** Try enabling Nanite for static meshes that are generated with at least this many triangles */
-	UPROPERTY(BlueprintReadWrite, config, EditAnywhere, Category = "USD options", meta = ( NoSpinbox = "true", UIMin = "0", ClampMin = "0" ) )
+	UPROPERTY(BlueprintReadWrite, config, EditAnywhere, Category = "USD options", meta = (NoSpinbox = "true", UIMin = "0", ClampMin = "0"))
 	int32 NaniteTriangleThreshold;
 
 	/** Specifies which set of shaders to use when parsing USD materials, in addition to the universal render context. */
-	UPROPERTY(BlueprintReadWrite, config, EditAnywhere, Category= "USD options")
+	UPROPERTY(BlueprintReadWrite, config, EditAnywhere, Category = "USD options")
 	FName RenderContextToImport;
 
 	/** Specifies which material purpose to use when parsing USD material bindings, in addition to the "allPurpose" fallback */
-	UPROPERTY( BlueprintReadWrite, config, EditAnywhere, Category = "USD options" )
+	UPROPERTY(BlueprintReadWrite, config, EditAnywhere, Category = "USD options")
 	FName MaterialPurpose;
 
 	// Describes what to add to the root bone animation within generated AnimSequences, if anything
-	UPROPERTY( BlueprintReadWrite, config, EditAnywhere, Category = "USD options", meta = ( EditCondition = bImportSkeletalAnimations ) )
+	UPROPERTY(BlueprintReadWrite, config, EditAnywhere, Category = "USD options", meta = (EditCondition = bImportSkeletalAnimations))
 	EUsdRootMotionHandling RootMotionHandling = EUsdRootMotionHandling::NoAdditionalRootMotion;
 
 	/** Subdivision level to use for all subdivision meshes on the opened stage. 0 means "don't subdivide" */
@@ -115,32 +124,30 @@ public:
 	FUsdMetadataImportOptions MetadataOptions;
 
 	/** Whether to use the specified StageOptions instead of the stage's own settings */
-	UPROPERTY( BlueprintReadWrite, config, EditAnywhere, Category = "USD options" )
+	UPROPERTY(BlueprintReadWrite, config, EditAnywhere, Category = "USD options")
 	bool bOverrideStageOptions;
 
 	/** Custom StageOptions to use for the stage */
-	UPROPERTY( EditAnywhere, config, BlueprintReadWrite, Category = "USD options", meta = ( EditCondition = bOverrideStageOptions ) )
+	UPROPERTY(EditAnywhere, config, BlueprintReadWrite, Category = "USD options", meta = (EditCondition = bOverrideStageOptions))
 	FUsdStageOptions StageOptions;
 
 	/**
 	 * When true the stage will be evaluated at ImportTimeCode for the import.
 	 * When false, the stage will be evaluated at the default (non-animated) timecode
 	 */
-	UPROPERTY( BlueprintReadWrite, config, EditAnywhere, Category = "USD options" )
+	UPROPERTY(BlueprintReadWrite, config, EditAnywhere, Category = "USD options")
 	bool bImportAtSpecificTimeCode;
 
 	/** TimeCode to evaluate the stage for import, in case bImportAtSpecificTimeCode is enabled */
-	UPROPERTY( BlueprintReadWrite, config, EditAnywhere, Category = "USD options", meta = ( EditCondition = bImportAtSpecificTimeCode ) )
+	UPROPERTY(BlueprintReadWrite, config, EditAnywhere, Category = "USD options", meta = (EditCondition = bImportAtSpecificTimeCode))
 	float ImportTimeCode;
 
-
 	/** Groom group interpolation settings */
 	UPROPERTY(EditAnywhere, config, BlueprintReadWrite, Category = "Groom")
 	TArray<FHairGroupsInterpolation> GroomInterpolationSettings;
 
-
 	/** What should happen when imported actors and components try to overwrite existing actors and components */
-	UPROPERTY(BlueprintReadWrite, config, EditAnywhere, Category = "Collision", meta=(EditCondition=bImportActors))
+	UPROPERTY(BlueprintReadWrite, config, EditAnywhere, Category = "Collision", meta = (EditCondition = bImportActors))
 	EReplaceActorPolicy ExistingActorPolicy;
 
 	/** What should happen when imported assets try to overwrite existing assets */
@@ -155,7 +162,6 @@ public:
 	UPROPERTY(BlueprintReadWrite, config, EditAnywhere, Category = "Collision")
 	bool bReuseIdenticalAssets;
 
-
 	/**
 	 * When enabled, assets will be imported into a content folder structure according to their prim path. When disabled,
 	 * assets are imported into content folders according to asset type (e.g. 'Materials', 'StaticMeshes', etc).
@@ -167,7 +173,13 @@ public:
 	 * Whether to try to combine individual assets and components of the same type on a kind-per-kind basis,
 	 * like multiple Mesh prims into a single Static Mesh
 	 */
-	UPROPERTY( BlueprintReadWrite, config, EditAnywhere, Category = "Processing", meta = ( Bitmask, BitmaskEnum = "/Script/UnrealUSDWrapper.EUsdDefaultKind" ) )
+	UPROPERTY(
+		BlueprintReadWrite,
+		config,
+		EditAnywhere,
+		Category = "Processing",
+		meta = (Bitmask, BitmaskEnum = "/Script/UnrealUSDWrapper.EUsdDefaultKind")
+	)
 	int32 KindsToCollapse;
 
 	/**
@@ -182,12 +194,20 @@ public:
 	 * If false, will cause us to use HierarchicalInstancedStaticMeshComponents to replicate the instancing behavior.
 	 * Point instancers inside other point instancer prototypes are *always* collapsed into the prototype's static mesh.
 	 */
-	UE_DEPRECATED( 5.2, "This option is now controlled via the cvar 'USD.CollapseTopLevelPointInstancers'" )
+	UE_DEPRECATED(5.2, "This option is now controlled via the cvar 'USD.CollapseTopLevelPointInstancers'")
 	UPROPERTY()
 	bool bCollapseTopLevelPointInstancers;
 
-	/** When true, if a prim has a "LOD" variant set with variants named "LOD0", "LOD1", etc. where each contains a UsdGeomMesh, the importer will attempt to parse the meshes as separate LODs of a single UStaticMesh. When false, only the selected variant will be parsed as LOD0 of the UStaticMesh.  */
-	UPROPERTY(BlueprintReadWrite, config, EditAnywhere, Category="Processing", meta=(DisplayName="Interpret LOD variant sets", EditCondition=bImportGeometry) )
+	/** When true, if a prim has a "LOD" variant set with variants named "LOD0", "LOD1", etc. where each contains a UsdGeomMesh, the importer will
+	 * attempt to parse the meshes as separate LODs of a single UStaticMesh. When false, only the selected variant will be parsed as LOD0 of the
+	 * UStaticMesh.  */
+	UPROPERTY(
+		BlueprintReadWrite,
+		config,
+		EditAnywhere,
+		Category = "Processing",
+		meta = (DisplayName = "Interpret LOD variant sets", EditCondition = bImportGeometry)
+	)
 	bool bInterpretLODs;
 
 public:
@@ -197,8 +217,5 @@ public:
 
 namespace UsdUtils
 {
-	USDSTAGEIMPORTER_API void AddAnalyticsAttributes(
-		const UUsdStageImportOptions& Options,
-		TArray< FAnalyticsEventAttribute >& InOutAttributes
-	);
+	USDSTAGEIMPORTER_API void AddAnalyticsAttributes(const UUsdStageImportOptions& Options, TArray<FAnalyticsEventAttribute>& InOutAttributes);
 }
diff --git a/Engine/Plugins/Importers/USDImporter/Source/USDStageImporter/Public/USDStageImporter.h b/Engine/Plugins/Importers/USDImporter/Source/USDStageImporter/Public/USDStageImporter.h
--- a/Engine/Plugins/Importers/USDImporter/Source/USDStageImporter/Public/USDStageImporter.h
+++ b/Engine/Plugins/Importers/USDImporter/Source/USDStageImporter/Public/USDStageImporter.h
@@ -2,8 +2,6 @@
 
 #pragma once
 
-#include "USDStageActor.h"
-
 #include "CoreMinimal.h"
 
 struct FUsdStageImportContext;
diff --git a/Engine/Plugins/Importers/USDImporter/Source/USDStageImporter/Public/USDStageImporterModule.h b/Engine/Plugins/Importers/USDImporter/Source/USDStageImporter/Public/USDStageImporterModule.h
--- a/Engine/Plugins/Importers/USDImporter/Source/USDStageImporter/Public/USDStageImporterModule.h
+++ b/Engine/Plugins/Importers/USDImporter/Source/USDStageImporter/Public/USDStageImporterModule.h
@@ -2,17 +2,17 @@
 
 #pragma once
 
-#include "USDStageImporter.h"
-
 #include "Modules/ModuleInterface.h"
 #include "Modules/ModuleManager.h"
 
+class UUsdStageImporter;
+
 class IUsdStageImporterModule : public IModuleInterface
 {
 public:
 	static inline IUsdStageImporterModule& Get()
 	{
-		return FModuleManager::LoadModuleChecked< IUsdStageImporterModule >("UsdStageImporter");
+		return FModuleManager::LoadModuleChecked<IUsdStageImporterModule>("UsdStageImporter");
 	}
 
 	static inline bool IsAvailable()
diff --git a/Engine/Plugins/Importers/USDImporter/Source/USDStageImporter/Public/Widgets/IUSDTreeViewItem.h b/Engine/Plugins/Importers/USDImporter/Source/USDStageImporter/Public/Widgets/IUSDTreeViewItem.h
--- a/Engine/Plugins/Importers/USDImporter/Source/USDStageImporter/Public/Widgets/IUSDTreeViewItem.h
+++ b/Engine/Plugins/Importers/USDImporter/Source/USDStageImporter/Public/Widgets/IUSDTreeViewItem.h
@@ -2,10 +2,9 @@
 
 #pragma once
 
-#include "HAL/Platform.h"
 #include "Templates/SharedPointer.h"
 
-class USDSTAGEIMPORTER_API IUsdTreeViewItem : public TSharedFromThis< IUsdTreeViewItem >
+class USDSTAGEIMPORTER_API IUsdTreeViewItem : public TSharedFromThis<IUsdTreeViewItem>
 {
 public:
 	virtual ~IUsdTreeViewItem() = default;
diff --git a/Engine/Plugins/Importers/USDImporter/Source/USDStageImporter/Public/Widgets/SUSDTreeView.h b/Engine/Plugins/Importers/USDImporter/Source/USDStageImporter/Public/Widgets/SUSDTreeView.h
--- a/Engine/Plugins/Importers/USDImporter/Source/USDStageImporter/Public/Widgets/SUSDTreeView.h
+++ b/Engine/Plugins/Importers/USDImporter/Source/USDStageImporter/Public/Widgets/SUSDTreeView.h
@@ -15,7 +15,7 @@ class FUsdTreeViewColumn
 public:
 	virtual ~FUsdTreeViewColumn() = default;
 
-	virtual TSharedRef< SWidget > GenerateWidget( const TSharedPtr< IUsdTreeViewItem > TreeItem, const TSharedPtr< ITableRow > TableRow ) = 0;
+	virtual TSharedRef<SWidget> GenerateWidget(const TSharedPtr<IUsdTreeViewItem> TreeItem, const TSharedPtr<ITableRow> TableRow) = 0;
 
 public:
 	bool bIsMainColumn = false;
@@ -23,128 +23,146 @@ public:
 
 struct FSharedUsdTreeData
 {
-	TMap< FName, TSharedRef< FUsdTreeViewColumn > > Columns;
+	TMap<FName, TSharedRef<FUsdTreeViewColumn>> Columns;
 };
 
-template< typename ItemType >
-class SUsdTreeView : public STreeView< ItemType >
+template<typename ItemType>
+class SUsdTreeView : public STreeView<ItemType>
 {
 public:
-	SLATE_BEGIN_ARGS( SUsdTreeView ) {}
+	SLATE_BEGIN_ARGS(SUsdTreeView)
+	{
+	}
 	SLATE_END_ARGS()
 
-	void Construct( const FArguments& InArgs );
+	void Construct(const FArguments& InArgs);
 
-	void AddColumn( FName ColumnKey, FText ColumnLabel, TSharedRef< FUsdTreeViewColumn > Column, SHeaderRow::FColumn::FArguments ColumnArguments = SHeaderRow::FColumn::FArguments() );
+	void AddColumn(
+		FName ColumnKey,
+		FText ColumnLabel,
+		TSharedRef<FUsdTreeViewColumn> Column,
+		SHeaderRow::FColumn::FArguments ColumnArguments = SHeaderRow::FColumn::FArguments()
+	);
 	virtual void SetupColumns() = 0;
 
-	virtual TSharedRef< ITableRow > OnGenerateRow( ItemType InDisplayNode, const TSharedRef< STableViewBase >& OwnerTable ) = 0;
-	virtual void OnGetChildren( ItemType InParent, TArray< ItemType >& OutChildren ) const = 0;
+	virtual TSharedRef<ITableRow> OnGenerateRow(ItemType InDisplayNode, const TSharedRef<STableViewBase>& OwnerTable) = 0;
+	virtual void OnGetChildren(ItemType InParent, TArray<ItemType>& OutChildren) const = 0;
 
-	virtual void OnTreeItemScrolledIntoView( ItemType TreeItem, const TSharedPtr< ITableRow >& Widget ) {}
+	virtual void OnTreeItemScrolledIntoView(ItemType TreeItem, const TSharedPtr<ITableRow>& Widget)
+	{
+	}
 
 protected:
-	TSharedPtr< SHeaderRow > HeaderRowWidget;
-	TArray< ItemType > RootItems;
+	TSharedPtr<SHeaderRow> HeaderRowWidget;
+	TArray<ItemType> RootItems;
 
-	TSharedPtr< FSharedUsdTreeData > SharedData;
+	TSharedPtr<FSharedUsdTreeData> SharedData;
 };
 
-template< typename ItemType >
-inline void SUsdTreeView< ItemType >::AddColumn( FName ColumnKey, FText ColumnLabel, TSharedRef< FUsdTreeViewColumn > Column, SHeaderRow::FColumn::FArguments ColumnArguments )
+template<typename ItemType>
+inline void SUsdTreeView<ItemType>::AddColumn(
+	FName ColumnKey,
+	FText ColumnLabel,
+	TSharedRef<FUsdTreeViewColumn> Column,
+	SHeaderRow::FColumn::FArguments ColumnArguments
+)
 {
-	ColumnArguments.ColumnId( ColumnKey );
-	ColumnArguments.DefaultLabel( ColumnLabel );
+	ColumnArguments.ColumnId(ColumnKey);
+	ColumnArguments.DefaultLabel(ColumnLabel);
 
-	HeaderRowWidget->AddColumn( ColumnArguments );
+	HeaderRowWidget->AddColumn(ColumnArguments);
 
-	SharedData->Columns.Add( ColumnKey, MoveTemp( Column ) );
+	SharedData->Columns.Add(ColumnKey, MoveTemp(Column));
 }
 
-template< typename ItemType >
-inline void SUsdTreeView< ItemType >::Construct( const FArguments& InArgs )
+template<typename ItemType>
+inline void SUsdTreeView<ItemType>::Construct(const FArguments& InArgs)
 {
-	SharedData = MakeShared< FSharedUsdTreeData >();
+	SharedData = MakeShared<FSharedUsdTreeData>();
 
-	HeaderRowWidget =
-			SNew( SHeaderRow )
-				.Visibility( EVisibility::Visible );
+	HeaderRowWidget = SNew(SHeaderRow).Visibility(EVisibility::Visible);
 
 	SetupColumns();
 
-	STreeView< ItemType >::Construct
-	(
-		typename STreeView< ItemType >::FArguments()
-		.TreeItemsSource( &RootItems )
-		.OnGenerateRow( this, &SUsdTreeView< ItemType >::OnGenerateRow )
-		.OnGetChildren( this, &SUsdTreeView< ItemType >::OnGetChildren )
-		.OnItemScrolledIntoView( this, &SUsdTreeView< ItemType >::OnTreeItemScrolledIntoView )
-		.HeaderRow( HeaderRowWidget )
-	);
+	STreeView<ItemType>::Construct(typename STreeView<ItemType>::FArguments()
+									   .TreeItemsSource(&RootItems)
+									   .OnGenerateRow(this, &SUsdTreeView<ItemType>::OnGenerateRow)
+									   .OnGetChildren(this, &SUsdTreeView<ItemType>::OnGetChildren)
+									   .OnItemScrolledIntoView(this, &SUsdTreeView<ItemType>::OnTreeItemScrolledIntoView)
+									   .HeaderRow(HeaderRowWidget));
 }
 
-template< typename ItemType >
-class SUsdTreeRow : public SMultiColumnTableRow< ItemType >
+template<typename ItemType>
+class SUsdTreeRow : public SMultiColumnTableRow<ItemType>
 {
 public:
 	typedef typename STableRow<ItemType>::FOnCanAcceptDrop FUsdOnCanAcceptDrop;
 	typedef typename STableRow<ItemType>::FOnAcceptDrop FUsdOnAcceptDrop;
 
-	SLATE_BEGIN_ARGS( SUsdTreeRow ) {}
-		SLATE_EVENT( FUsdOnCanAcceptDrop, OnCanAcceptDrop )
-		SLATE_EVENT( FUsdOnAcceptDrop, OnAcceptDrop )
-		SLATE_EVENT( FOnDragDetected, OnDragDetected )
-		SLATE_EVENT( FOnTableRowDragEnter, OnDragEnter )
-		SLATE_EVENT( FOnTableRowDragLeave, OnDragLeave )
-		SLATE_EVENT( FOnTableRowDrop, OnDrop )
+	SLATE_BEGIN_ARGS(SUsdTreeRow)
+	{
+	}
+	SLATE_EVENT(FUsdOnCanAcceptDrop, OnCanAcceptDrop)
+	SLATE_EVENT(FUsdOnAcceptDrop, OnAcceptDrop)
+	SLATE_EVENT(FOnDragDetected, OnDragDetected)
+	SLATE_EVENT(FOnTableRowDragEnter, OnDragEnter)
+	SLATE_EVENT(FOnTableRowDragLeave, OnDragLeave)
+	SLATE_EVENT(FOnTableRowDrop, OnDrop)
 	SLATE_END_ARGS()
 
-	void Construct( const FArguments& InArgs, ItemType InTreeItem, const TSharedRef< STableViewBase >& OwnerTable, TSharedPtr< FSharedUsdTreeData > InSharedData )
+	void Construct(
+		const FArguments& InArgs,
+		ItemType InTreeItem,
+		const TSharedRef<STableViewBase>& OwnerTable,
+		TSharedPtr<FSharedUsdTreeData> InSharedData
+	)
 	{
 		TreeItem = InTreeItem;
 		SharedData = InSharedData;
 
-		SMultiColumnTableRow< ItemType >::Construct(
-			typename SMultiColumnTableRow< ItemType >::FArguments()
-				.OnCanAcceptDrop( InArgs._OnCanAcceptDrop )
-				.OnAcceptDrop( InArgs._OnAcceptDrop )
-				.OnDragDetected( InArgs._OnDragDetected )
-				.OnDragEnter( InArgs._OnDragEnter )
-				.OnDragLeave( InArgs._OnDragLeave )
-				.OnDrop( InArgs._OnDrop ),
+		SMultiColumnTableRow<ItemType>::Construct(
+			typename SMultiColumnTableRow<ItemType>::FArguments()
+				.OnCanAcceptDrop(InArgs._OnCanAcceptDrop)
+				.OnAcceptDrop(InArgs._OnAcceptDrop)
+				.OnDragDetected(InArgs._OnDragDetected)
+				.OnDragEnter(InArgs._OnDragEnter)
+				.OnDragLeave(InArgs._OnDragLeave)
+				.OnDrop(InArgs._OnDrop),
 			OwnerTable
 		);
 	}
 
 	/** Overridden from SMultiColumnTableRow.  Generates a widget for this column of the tree row. */
-	virtual TSharedRef< SWidget > GenerateWidgetForColumn( const FName& ColumnName ) override;
+	virtual TSharedRef<SWidget> GenerateWidgetForColumn(const FName& ColumnName) override;
 
 private:
-	TSharedPtr< typename ItemType::ElementType > TreeItem;
-	TSharedPtr< FSharedUsdTreeData > SharedData;
+	TSharedPtr<typename ItemType::ElementType> TreeItem;
+	TSharedPtr<FSharedUsdTreeData> SharedData;
 };
 
-template< typename ItemType >
-inline TSharedRef< SWidget > SUsdTreeRow< ItemType >::GenerateWidgetForColumn( const FName& ColumnName )
+template<typename ItemType>
+inline TSharedRef<SWidget> SUsdTreeRow<ItemType>::GenerateWidgetForColumn(const FName& ColumnName)
 {
-	TSharedRef< SWidget > ColumnWidget = SNullWidget::NullWidget;
+	TSharedRef<SWidget> ColumnWidget = SNullWidget::NullWidget;
 
-	if ( SharedData && SharedData->Columns.Contains( ColumnName ) )
+	if (SharedData && SharedData->Columns.Contains(ColumnName))
 	{
-		ColumnWidget = SharedData->Columns[ ColumnName ]->GenerateWidget(
-			StaticCastSharedPtr< IUsdTreeViewItem >( TSharedPtr< typename ItemType::ElementType >( TreeItem ) ),
+		ColumnWidget = SharedData->Columns[ColumnName]->GenerateWidget(
+			StaticCastSharedPtr<IUsdTreeViewItem>(TSharedPtr<typename ItemType::ElementType>(TreeItem)),
 			this->SharedThis(this)
 		);
 	}
 
-	if ( SharedData->Columns[ ColumnName ]->bIsMainColumn )
+	// clang-format off
+	if (SharedData->Columns[ColumnName]->bIsMainColumn)
 	{
-		return SNew( SHorizontalBox )
+		return SNew(SHorizontalBox)
 			+SHorizontalBox::Slot()
 			.AutoWidth()
 			.Padding(6, 0, 0, 0)
 			[
-				SNew( SExpanderArrow, this->SharedThis(this) ).IndentAmount(12)
+				SNew(SExpanderArrow, this->SharedThis(this))
+				.IndentAmount(12)
 			]
 
 			+SHorizontalBox::Slot()
@@ -155,11 +173,12 @@ inline TSharedRef< SWidget > SUsdTreeRow< ItemType >::GenerateWidgetForColumn( c
 	}
 	else
 	{
-		return SNew( SHorizontalBox )
+		return SNew(SHorizontalBox)
 			+SHorizontalBox::Slot()
 			.AutoWidth()
 			[
 				ColumnWidget
 			];
 	}
+	// clang-format on
 }
diff --git a/Engine/Plugins/Importers/USDImporter/Source/USDStageImporter/USDStageImporter.Build.cs b/Engine/Plugins/Importers/USDImporter/Source/USDStageImporter/USDStageImporter.Build.cs
--- a/Engine/Plugins/Importers/USDImporter/Source/USDStageImporter/USDStageImporter.Build.cs
+++ b/Engine/Plugins/Importers/USDImporter/Source/USDStageImporter/USDStageImporter.Build.cs
@@ -20,7 +20,6 @@ namespace UnrealBuildTool.Rules
 					"Engine",
 					"GeometryCache",
 					"InputCore",
-					"InterchangePipelines",
 					"JsonUtilities",
 					"LevelSequence",
 					"MainFrame",
