USD: Fix compilation error on QAGameEditor.

#rnx

diff --git a/Engine/Plugins/Importers/USDImporter/Source/USDStageEditorViewModels/Public/USDStageEditorViewModelsModule.h b/Engine/Plugins/Importers/USDImporter/Source/USDStageEditorViewModels/Public/USDStageEditorViewModelsModule.h
--- a/Engine/Plugins/Importers/USDImporter/Source/USDStageEditorViewModels/Public/USDStageEditorViewModelsModule.h
+++ b/Engine/Plugins/Importers/USDImporter/Source/USDStageEditorViewModels/Public/USDStageEditorViewModelsModule.h
@@ -4,6 +4,6 @@
 
 #include "Modules/ModuleInterface.h"
 
-class USDSTAGEEDITORVIEWMODELS_API IUsdStageEditorViewModelsModule : public IModuleInterface
+class IUsdStageEditorViewModelsModule : public IModuleInterface
 {
 };
diff --git a/Engine/Plugins/Importers/USDImporter/Source/UnrealUSDWrapper/Public/USDMemory.h b/Engine/Plugins/Importers/USDImporter/Source/UnrealUSDWrapper/Public/USDMemory.h
--- a/Engine/Plugins/Importers/USDImporter/Source/UnrealUSDWrapper/Public/USDMemory.h
+++ b/Engine/Plugins/Importers/USDImporter/Source/UnrealUSDWrapper/Public/USDMemory.h
@@ -7,6 +7,7 @@
 #include "Misc/Optional.h"
 #include "Misc/ScopeRWLock.h"
 #include "Modules/Boilerplate/ModuleBoilerplate.h"
+#include "Modules/ModuleManager.h"
 #include "Templates/SharedPointer.h"
 #include "Templates/UnrealTemplate.h"
 
