USD: Fix crash when importing anim sequences repeatedly in some scenarios.

diff --git a/Engine/Plugins/Importers/USDImporter/Source/USDStageImporter/Private/USDStageImporter.cpp b/Engine/Plugins/Importers/USDImporter/Source/USDStageImporter/Private/USDStageImporter.cpp
--- a/Engine/Plugins/Importers/USDImporter/Source/USDStageImporter/Private/USDStageImporter.cpp
+++ b/Engine/Plugins/Importers/USDImporter/Source/USDStageImporter/Private/USDStageImporter.cpp
@@ -1835,17 +1835,7 @@ namespace UE::USDStageImporter::Private
 		{
 			if (USkeletalMeshComponent* SkeletalMeshComponent = Cast<USkeletalMeshComponent>(Component))
 			{
-				if (SkeletalMeshComponent->AnimationData.AnimToPlay == nullptr)
-				{
-					SkeletalMeshComponent->TickAnimation(0.f, false);
-					SkeletalMeshComponent->RefreshBoneTransforms();
-					SkeletalMeshComponent->RefreshFollowerComponents();
-					SkeletalMeshComponent->UpdateComponentToWorld();
-					SkeletalMeshComponent->FinalizeBoneTransform();
-					SkeletalMeshComponent->MarkRenderTransformDirty();
-					SkeletalMeshComponent->MarkRenderDynamicDataDirty();
-				}
-				else if (bImportAtSpecificTimeCode)
+				if (bImportAtSpecificTimeCode)
 				{
 					// The asset we return from the import factories may lead to
 					// USkeletalMesh::PostEditChangeProperty being called. The FMultiComponentReregisterContext in
@@ -1863,6 +1853,21 @@ namespace UE::USDStageImporter::Private
 					const bool bIsPlaying = false;
 					const float Position = SkeletalMeshComponent->GetPosition();
 					SkeletalMeshComponent->OverrideAnimationData(SkeletalMeshComponent->AnimationData.AnimToPlay, bIsLooping, bIsPlaying, Position);
+
+					const bool bForceReinit = true;
+					SkeletalMeshComponent->InitAnim(bForceReinit);
+				}
+				else
+				{
+					const bool bForceReinit = true;
+					SkeletalMeshComponent->InitAnim(bForceReinit);
+					SkeletalMeshComponent->TickAnimation(0.f, false);
+					SkeletalMeshComponent->RefreshBoneTransforms();
+					SkeletalMeshComponent->RefreshFollowerComponents();
+					SkeletalMeshComponent->UpdateComponentToWorld();
+					SkeletalMeshComponent->FinalizeBoneTransform();
+					SkeletalMeshComponent->MarkRenderTransformDirty();
+					SkeletalMeshComponent->MarkRenderDynamicDataDirty();
 				}
 
 				// It does need us to manually set this to dirty regardless or else it won't update in case we changed material
