Fix warning V789: Iterators for the container, used in the range-based for loop, become invalid upon the call of the 'Remove' function.
#rnx

diff --git a/Engine/Plugins/Experimental/Web/WebAPI/Source/WebAPI/Private/WebAPISubsystem.cpp b/Engine/Plugins/Experimental/Web/WebAPI/Source/WebAPI/Private/WebAPISubsystem.cpp
--- a/Engine/Plugins/Experimental/Web/WebAPI/Source/WebAPI/Private/WebAPISubsystem.cpp
+++ b/Engine/Plugins/Experimental/Web/WebAPI/Source/WebAPI/Private/WebAPISubsystem.cpp
@@ -150,8 +150,9 @@ bool UWebAPISubsystem::IsAllowedToTick() const
 
 void UWebAPISubsystem::Tick(float DeltaTime)
 {
-	for(TSharedRef<IHttpRequest>& Request : RequestBuffer)
+	for(auto It = RequestBuffer.CreateIterator(); It; ++It)
 	{
+		TSharedRef<IHttpRequest>& Request = *It;
 		const EHttpRequestStatus::Type RequestStatus = Request->GetStatus();
 		if(RequestStatus == EHttpRequestStatus::NotStarted)
 		{
@@ -180,13 +181,13 @@ void UWebAPISubsystem::Tick(float DeltaTime)
 							HostRequestBuffer.Add(Host);
 						}
 
-						RequestBuffer.Remove(Request);
+						It.RemoveCurrent();
 
 						continue;
 					}
 				}
-				
-				RequestBuffer.Remove(Request);			
+
+				It.RemoveCurrent();
 			}
 		}
 	}
diff --git a/Engine/Plugins/Importers/USDImporter/Source/USDClasses/Private/USDAssetCache2.cpp b/Engine/Plugins/Importers/USDImporter/Source/USDClasses/Private/USDAssetCache2.cpp
--- a/Engine/Plugins/Importers/USDImporter/Source/USDClasses/Private/USDAssetCache2.cpp
+++ b/Engine/Plugins/Importers/USDImporter/Source/USDClasses/Private/USDAssetCache2.cpp
@@ -1547,8 +1547,9 @@ bool UUsdAssetCache2::TryUnloadAsset(FCachedAssetInfo& InOutInfo)
 		if (TObjectPtr<UObject> Asset = AssetStorage.FindRef(InOutInfo.Hash))
 		{
 			// Check to see if all of its Consumers can be unloaded too
-			for (const FSoftObjectPath& Consumer : InOutInfo.Consumers)
+			for (auto It = InOutInfo.Consumers.CreateIterator(); It; ++It)
 			{
+				const FSoftObjectPath& Consumer = *It;
 				if (FCachedAssetInfo* ConsumerInfo = LRUCache.Find(Consumer))
 				{
 					ensure(ConsumerInfo->Dependencies.Contains(Asset));
@@ -1563,7 +1564,7 @@ bool UUsdAssetCache2::TryUnloadAsset(FCachedAssetInfo& InOutInfo)
 				// We don't know about this Consumer anymore (maybe it was evicted?)
 				else
 				{
-					InOutInfo.Consumers.Remove(Consumer);
+					It.RemoveCurrent();
 				}
 			}
 
diff --git a/Engine/Plugins/VirtualProduction/RemoteControl/Source/RemoteControl/Private/RemoteControlModule.cpp b/Engine/Plugins/VirtualProduction/RemoteControl/Source/RemoteControl/Private/RemoteControlModule.cpp
--- a/Engine/Plugins/VirtualProduction/RemoteControl/Source/RemoteControl/Private/RemoteControlModule.cpp
+++ b/Engine/Plugins/VirtualProduction/RemoteControl/Source/RemoteControl/Private/RemoteControlModule.cpp
@@ -2516,11 +2516,11 @@ void FRemoteControlModule::OnAssetRenamed(const FAssetData& AssetData, const FSt
 				// If it's not valid or doesn't have a valid id, search for it and remove it.
 				if (!bRemovedId)
 				{
-					for (const auto& PairInner : CachedPresetNamesById)
+					for (auto It = CachedPresetNamesById.CreateIterator(); It; ++It)
 					{
-						if (PairInner.Value == Pair.Key)
+						if (It.Value() == Pair.Key)
 						{
-							CachedPresetNamesById.Remove(PairInner.Key);
+							It.RemoveCurrent();
 						}
 					}
 				}
diff --git a/Engine/Source/Runtime/Datasmith/CADKernel/Private/CADKernel/Core/Group.cpp b/Engine/Source/Runtime/Datasmith/CADKernel/Private/CADKernel/Core/Group.cpp
--- a/Engine/Source/Runtime/Datasmith/CADKernel/Private/CADKernel/Core/Group.cpp
+++ b/Engine/Source/Runtime/Datasmith/CADKernel/Private/CADKernel/Core/Group.cpp
@@ -40,12 +40,13 @@ void FGroup::ReplaceEntitiesWithMap(const TMap<TSharedPtr<FEntity>, TSharedPtr<F
 
 void FGroup::RemoveNonTopologicalEntities()
 {
-	for (const TSharedPtr<FEntity>& Entity : Entities)
+	for (auto It = Entities.CreateIterator(); It; ++It)
 	{
+		const TSharedPtr<FEntity>& Entity = *It;
 		EEntity Type = Entity->GetEntityType();
 		if (Type != EEntity::TopologicalFace && Type != EEntity::TopologicalEdge && Type != EEntity::TopologicalVertex) 
 		{
-			Entities.Remove(Entity);
+			It.RemoveCurrent();
 		}
 	}
 }
