USD: Fixed USD PhysicsMeshCollision API approximation attribute set to None (or when only the PhysicsCollision API is applied and enabled) to behave as CTF_UseComplexAsSimple.

diff --git a/Engine/Plugins/Importers/USDImporter/Source/USDSchemas/Private/USDGeomMeshTranslator.cpp b/Engine/Plugins/Importers/USDImporter/Source/USDSchemas/Private/USDGeomMeshTranslator.cpp
--- a/Engine/Plugins/Importers/USDImporter/Source/USDSchemas/Private/USDGeomMeshTranslator.cpp
+++ b/Engine/Plugins/Importers/USDImporter/Source/USDSchemas/Private/USDGeomMeshTranslator.cpp
@@ -1315,7 +1315,7 @@ namespace UE::UsdCollision::Private
 	// Approximation types from PhysicsMeshCollisionAPI and UE-specific approximations
 	enum class EUsdCollisionType : uint8
 	{
-		None,
+		None, // no approximation so equivalent to CTF_UseComplexAsSimple
 		ConvexDecomposition,
 		ConvexHull,
 		Sphere,
@@ -1508,18 +1508,24 @@ namespace UE::UsdCollision::Private
 		BodySetup->bNeverNeedsCookedCollisionData = false;
 		BodySetup->DefaultInstance.SetCollisionEnabled(ECollisionEnabled::QueryAndPhysics);
 		BodySetup->DefaultInstance.SetCollisionProfileName(UCollisionProfile::BlockAllDynamic_ProfileName);
+		BodySetup->CollisionTraceFlag = CTF_UseDefault;
 
 		switch (Approximation)
 		{
 			case EUsdCollisionType::None:
 			{
 				// No approximation but allow for custom collision shapes
+				// If there are no custom collision shapes, then use complex as simple
 				FKAggregateGeom CollisionShapes;
 				CollectCustomCollisionShapes(UsdPrim, CollisionShapes);
 				if (CollisionShapes.GetElementCount())
 				{
 					BodySetup->AddCollisionFrom(CollisionShapes);
 				}
+				else
+				{
+					BodySetup->CollisionTraceFlag = CTF_UseComplexAsSimple;
+				}
 				break;
 			}
 #if WITH_EDITOR
