USD Export: Fixed animation constraints evaluation to support changes in the parenting of the constraint.

diff --git a/Engine/Plugins/Importers/USDImporter/Source/USDExporter/Private/LevelSequenceExporterUSD.cpp b/Engine/Plugins/Importers/USDImporter/Source/USDExporter/Private/LevelSequenceExporterUSD.cpp
--- a/Engine/Plugins/Importers/USDImporter/Source/USDExporter/Private/LevelSequenceExporterUSD.cpp
+++ b/Engine/Plugins/Importers/USDImporter/Source/USDExporter/Private/LevelSequenceExporterUSD.cpp
@@ -1153,8 +1153,6 @@ namespace UE::LevelSequenceExporterUSD::Private
 		UsdStage.SetEndTimeCode(EndTimeCode);
 
 		const FConstraintsManagerController& Controller = FConstraintsManagerController::Get(Context.World);
-		static constexpr bool bSorted = true;
-		const TArray<TWeakObjectPtr<UTickableConstraint>> AllConstraints = Controller.GetAllConstraints(bSorted);
 
 		// Collect and sort the bakers: We need all skeletal animations evaluated first, as we need to manually force
 		// the component to update and attached components will only get the correct values if they are evaluated
@@ -1201,13 +1199,8 @@ namespace UE::LevelSequenceExporterUSD::Private
 			Context.Sequencer->ForceEvaluate();
 
 			// Evaluate constraints (these run on tick in the editor, so here we must trigger them manually)
-			for (const TWeakObjectPtr<UTickableConstraint>& Constraint : AllConstraints)
-			{
-				if (Constraint.IsValid())
-				{
-					Constraint->Evaluate();
-				}
-			}
+			// Can't iterate through a pre-sorted list since the parenting of the constraints can change between frames
+			Controller.EvaluateAllConstraints();
 
 			FFrameTime KeyTime = FFrameRate::Snap(EvalTime, Resolution, DisplayRate).FloorToFrame();
 			double UsdTimeCode = FFrameRate::TransformTime(KeyTime, Resolution, StageFrameRate).AsDecimal();
