USD: Fix crashing when setting a stage state track on the transient level sequence.

diff --git a/Engine/Plugins/Importers/USDImporter/Source/USDStage/Private/USDStageActor.cpp b/Engine/Plugins/Importers/USDImporter/Source/USDStage/Private/USDStageActor.cpp
--- a/Engine/Plugins/Importers/USDImporter/Source/USDStage/Private/USDStageActor.cpp
+++ b/Engine/Plugins/Importers/USDImporter/Source/USDStage/Private/USDStageActor.cpp
@@ -34,6 +34,8 @@
 #include "UsdWrappers/UsdGeomXformable.h"
 #include "UsdWrappers/UsdStage.h"
 
+#include "Async/Async.h"
+#include "Async/ParallelFor.h"
 #include "CineCameraActor.h"
 #include "CineCameraComponent.h"
 #include "Components/DirectionalLightComponent.h"
@@ -232,26 +234,6 @@ struct FUsdStageActorImpl
 #endif	  // WITH_EDITOR
 	}
 
-	template<typename ObjectPtr>
-	static void CloseEditorsForAssets(const TMap<FString, ObjectPtr>& AssetsCache)
-	{
-#if WITH_EDITOR
-		if (GIsEditor && GEditor)
-		{
-			if (UAssetEditorSubsystem* AssetEditorSubsysttem = GEditor->GetEditorSubsystem<UAssetEditorSubsystem>())
-			{
-				for (const TPair<FString, ObjectPtr>& Pair : AssetsCache)
-				{
-					if (UObject* Asset = Pair.Value)
-					{
-						AssetEditorSubsysttem->CloseAllEditorsForAsset(Asset);
-					}
-				}
-			}
-		}
-#endif	  // WITH_EDITOR
-	}
-
 	static void DiscardStage(const UE::FUsdStage& Stage, AUsdStageActor* DiscardingActor)
 	{
 		if (!Stage || !DiscardingActor)
@@ -3268,22 +3250,35 @@ void AUsdStageActor::UnloadUsdStage()
 	// and could lead to stage changes
 	BlockMonitoringLevelSequenceForThisTransaction();
 
-	// Close the Sequencer before dropping the info cache, as the Sequencer closing may
-	// trigger one last SetTime call on the stage actor (to revert to the preanimated state),
-	// and if we try animating things and calling UpdateComponents we may need the info cache
 	if (LevelSequence)
 	{
 #if WITH_EDITOR
 		// CloseAllEditorsForAsset crashes if called when the engine is closing
 		if (GEditor && !IsEngineExitRequested())
 		{
-			GEditor->GetEditorSubsystem<UAssetEditorSubsystem>()->CloseAllEditorsForAsset(LevelSequence);
+			// We'll only close the Sequencer via a delayed task. This because the Sequencer can't itself
+			// close from the callstack of its LevelSequence being evaluated (for example, imagine we had a track to
+			// set StateState to Closed: Sequencer evaluates the track -> Calls SetStageState -> Ends up here -> We try
+			// destroying the Sequencer -> Crash).
+			// Note that an AsyncTask has it run on the same tick and so within the same frame, while FTSTicker would end
+			// up on the next tick, showing a frame of the broken sequence bindings on the UI, which doesn't look great...
+			ULevelSequence* LevelSequencePtr = LevelSequence;
+			AsyncTask(
+				ENamedThreads::GameThread,
+				[LevelSequencePtr]()
+				{
+					GEditor->GetEditorSubsystem<UAssetEditorSubsystem>()->CloseAllEditorsForAsset(LevelSequencePtr);
+				}
+			);
 		}
 #endif	  // WITH_EDITOR
 		LevelSequence = nullptr;
 	}
 	LevelSequenceHelper.Clear();
 
+	// Resetting ObjectsToWatch before dropping the info cache, as the Sequencer closing may
+	// trigger one last SetTime call on the stage actor (to revert to the preanimated state),
+	// and if we try animating things and calling UpdateComponents we may need the info cache
 	ObjectsToWatch.Reset();
 	BlendShapesByPath.Reset();
 	MaterialToPrimvarToUVIndex.Reset();
@@ -3468,20 +3463,8 @@ void AUsdStageActor::ReloadAnimations()
 	if (!IsTemplate())
 	{
 #if WITH_EDITOR
-		TArray<TSharedPtr<ISequencer>> SequencersToReset;
-		if (IUsdStageModule* UsdStageModule = FModuleManager::Get().GetModulePtr<IUsdStageModule>(TEXT("UsdStage")))
-		{
-			for (const TWeakPtr<ISequencer>& ExistingSequencer : UsdStageModule->GetExistingSequencers())
-			{
-				if (TSharedPtr<ISequencer> PinnedSequencer = ExistingSequencer.Pin())
-				{
-					if (PinnedSequencer->GetRootMovieSceneSequence() == LevelSequence)
-					{
-						SequencersToReset.Add(PinnedSequencer);
-					}
-				}
-			}
-		}
+		UAssetEditorSubsystem* AssetEditorSubsystem = GEditor->GetEditorSubsystem<UAssetEditorSubsystem>();
+		const bool bSequencerWasOpened = AssetEditorSubsystem && AssetEditorSubsystem->FindEditorsForAssetAndSubObjects(LevelSequence).Num() > 0;
 #endif	  // WITH_EDITOR
 
 		// We need to guarantee we'll record our change of LevelSequence into the transaction, as Init() will create a new one
@@ -3492,12 +3475,11 @@ void AUsdStageActor::ReloadAnimations()
 		LevelSequenceHelper.BindToUsdStageActor(this);
 
 #if WITH_EDITOR
-		for (TSharedPtr<ISequencer>& Sequencer : SequencersToReset)
+		if (bSequencerWasOpened && LevelSequence && AssetEditorSubsystem)
 		{
-			if (LevelSequence && Sequencer->GetRootMovieSceneSequence() != LevelSequence)
-			{
-				Sequencer->ResetToNewRootSequence(*LevelSequence);
-			}
+			// Open the LevelSequence editor via the UAssetEditorSubsystem, otherwise calls to CloseAllEditorsForAsset
+			// won't be able to close it
+			AssetEditorSubsystem->OpenEditorForAsset(LevelSequence);
 		}
 #endif	  // WITH_EDITOR
 	}
