USD: Fix support for exporting animations with attachment tracks.

diff --git a/Engine/Plugins/Importers/USDImporter/Source/USDExporter/Private/LevelSequenceExporterUSD.cpp b/Engine/Plugins/Importers/USDImporter/Source/USDExporter/Private/LevelSequenceExporterUSD.cpp
--- a/Engine/Plugins/Importers/USDImporter/Source/USDExporter/Private/LevelSequenceExporterUSD.cpp
+++ b/Engine/Plugins/Importers/USDImporter/Source/USDExporter/Private/LevelSequenceExporterUSD.cpp
@@ -42,6 +42,7 @@
 #include "Selection.h"
 #include "Sequencer/MovieSceneControlRigParameterTrack.h"
 #include "Subsystems/AssetEditorSubsystem.h"
+#include "Tracks/MovieScene3DAttachTrack.h"
 #include "Tracks/MovieScenePropertyTrack.h"
 #include "Tracks/MovieSceneSkeletalAnimationTrack.h"
 #include "Tracks/MovieSceneSpawnTrack.h"
@@ -1038,6 +1039,12 @@ namespace UE::LevelSequenceExporterUSD::Private
 							UnrealToUsd::CreateSkeletalAnimationBaker(SkeletonPrim, SkelAnimPrim, *SkeletalBoundComponent, Baker);
 						}
 					}
+					// If we have an attach track that attaches the object to somewhere else, then we'll need to bake in that transform
+					// change, as we can't export "hierarchy changes" otherwise
+					else if (Track->IsA<UMovieScene3DAttachTrack>())
+					{
+						UnrealToUsd::CreateComponentPropertyBaker(Prim, *BoundComponent, TEXT("Transform"), Baker);
+					}
 
 					// If we made a baker and we don't have one of this type for this component yet, add its lambda to the array
 					FCombinedComponentBakers& ExistingBakers = InOutComponentBakers.FindOrAdd(BoundComponent);
diff --git a/Engine/Plugins/Importers/USDImporter/Source/USDUtilities/Private/USDPrimConversion.cpp b/Engine/Plugins/Importers/USDImporter/Source/USDUtilities/Private/USDPrimConversion.cpp
--- a/Engine/Plugins/Importers/USDImporter/Source/USDUtilities/Private/USDPrimConversion.cpp
+++ b/Engine/Plugins/Importers/USDImporter/Source/USDUtilities/Private/USDPrimConversion.cpp
@@ -1144,7 +1144,7 @@ bool UsdToUnreal::ConvertBoundsTimeSamples(
 			FVector UESpaceUSDMax = UsdToUnreal::ConvertVector(StageInfo, Box.GetMax());
 			UEBox = FBox{
 				TArray<FVector>{UESpaceUSDMin, UESpaceUSDMax}
-			};
+			 };
 		}
 
 		MinXValues.Emplace_GetRef(UEBox.Min.X).InterpMode = InterpMode;
@@ -1812,7 +1812,7 @@ bool UsdToUnreal::ConvertBounds(const pxr::UsdPrim& Prim, UUsdDrawModeComponent*
 		FVector UESpaceUSDMax = UsdToUnreal::ConvertVector(StageInfo, Box.GetMax());
 		FBox UEBox{
 			TArray<FVector>{UESpaceUSDMin, UESpaceUSDMax}
-		};
+		 };
 
 		BoundsComponent->SetBoundsMin(UEBox.Min);
 		BoundsComponent->SetBoundsMax(UEBox.Max);
@@ -2840,8 +2840,8 @@ bool UnrealToUsd::ConvertSceneComponent(const pxr::UsdStageRefPtr& Stage, const
 	}
 
 	// Compensate different orientation for light or camera components:
-	// In USD cameras shoot towards local - Z, with + Y up.Lights also emit towards local - Z, with + Y up
-	// In UE cameras shoot towards local + X, with + Z up.Lights also emit towards local + X, with + Z up
+	// In USD cameras shoot towards local - Z, with + Y up. Lights also emit towards local - Z, with + Y up
+	// In UE cameras shoot towards local + X, with + Z up. Lights also emit towards local + X, with + Z up
 	// Note that this wouldn't have worked in case we collapsed light and camera components, but these always get their own
 	// actors, so we know that we don't have a single component that represents a large collapsed prim hierarchy
 	if (UsdPrim.IsA<pxr::UsdGeomCamera>() || UsdPrim.HasAPI<pxr::UsdLuxLightAPI>())
@@ -3656,36 +3656,47 @@ bool UnrealToUsd::CreateComponentPropertyBaker(
 
 			Attr.Clear();
 
+			// Keep track of our original attach parent so that we can export baked transforms in case we have attach tracks.
+			// We're generating these before the Sequencer is evaluated with our Sequence, so this OriginalAttachParent should be
+			// the same parent that will be exported to the USD layers
+			USceneComponent* OriginalAttachParent = Component.GetAttachParent();
+
+			const bool bStageIsZUp = UsdUtils::GetUsdStageUpAxis(UsdStage) == pxr::UsdGeomTokens->z;
+
 			// Compensate different orientation for light or camera components
-			FTransform AdditionalRotation = FTransform::Identity;
+			FTransform CameraCompensation = FTransform::Identity;
 			if (UsdPrim.IsA<pxr::UsdGeomCamera>() || UsdPrim.HasAPI<pxr::UsdLuxLightAPI>())
 			{
-				AdditionalRotation = FTransform(FRotator(0.0f, 90.0f, 0.0f));
+				CameraCompensation = FTransform(FRotator(0.0f, 90.0f, 0.0f));
 
-				if (StageInfo.UpAxis == EUsdUpAxis::ZAxis)
+				if (bStageIsZUp)
 				{
-					AdditionalRotation *= FTransform(FRotator(90.0f, 0.0f, 0.0f));
+					CameraCompensation *= FTransform(FRotator(90.0f, 0.0f, 0.0f));
 				}
 			}
 
-			// Invert compensation applied to parent if it's a light or camera component
-			if (const USceneComponent* AttachParent = Component.GetAttachParent())
+			// Note how we only need the ParentCameraCompensation for our actual OriginalAttachParent. When we're "attached"
+			// to a camera or light via our attach track the world transform for that camera or light won't contain a CameraCompensation
+			// itself, as that is something that we add ourselves only when exporting the relative transforms to USD
+			FTransform InverseParentCameraCompensation = FTransform::Identity;
+			if (OriginalAttachParent)
 			{
-				if (AttachParent->IsA(UCineCameraComponent::StaticClass()) || AttachParent->IsA(ULightComponent::StaticClass()))
+				if (OriginalAttachParent->IsA(UCineCameraComponent::StaticClass()) || OriginalAttachParent->IsA(ULightComponent::StaticClass()))
 				{
-					FTransform InverseCompensation = FTransform(FRotator(0.0f, 90.f, 0.0f));
+					InverseParentCameraCompensation = FTransform(FRotator(0.0f, 90.f, 0.0f));
 
-					if (UsdUtils::GetUsdStageUpAxis(UsdStage) == pxr::UsdGeomTokens->z)
+					if (bStageIsZUp)
 					{
-						InverseCompensation *= FTransform(FRotator(90.0f, 0.f, 0.0f));
+						InverseParentCameraCompensation *= FTransform(FRotator(90.0f, 0.f, 0.0f));
 					}
 
-					AdditionalRotation = AdditionalRotation * InverseCompensation.Inverse();
+					InverseParentCameraCompensation = InverseParentCameraCompensation.Inverse();
 				}
 			}
 
 			BakerType = EBakingType::Transform;
-			BakerFunction = [&Component, AdditionalRotation, StageInfo, Attr](double UsdTimeCode)
+			BakerFunction =
+				[&Component, CameraCompensation, InverseParentCameraCompensation, StageInfo, Attr, OriginalAttachParent](double UsdTimeCode)
 			{
 				FScopedUsdAllocs Allocs;
 
@@ -3697,19 +3708,20 @@ bool UnrealToUsd::CreateComponentPropertyBaker(
 				// It may seem wasteful to do this inside the baker function, but you can place "Attach tracks" on the
 				// Sequencer that may make the attach socket change every frame, so we do need this
 				FTransform RelativeTransform;
-				if (USceneComponent* Parent = Component.GetAttachParent())
+				if (OriginalAttachParent)
 				{
-					Parent->ConditionalUpdateComponentToWorld();
-					Parent->UpdateChildTransforms();
-					RelativeTransform = Component.GetComponentTransform().GetRelativeTransform(Parent->GetComponentTransform());
+					OriginalAttachParent->ConditionalUpdateComponentToWorld();
+					OriginalAttachParent->UpdateChildTransforms();
+					RelativeTransform = Component.GetComponentTransform().GetRelativeTransform(OriginalAttachParent->GetComponentTransform());
 				}
 				else
 				{
 					RelativeTransform = Component.GetRelativeTransform();
 				}
 
-				FTransform FinalUETransform = AdditionalRotation * RelativeTransform;
-				pxr::GfMatrix4d UsdTransform = UnrealToUsd::ConvertTransform(StageInfo, FinalUETransform);
+				RelativeTransform = CameraCompensation * RelativeTransform * InverseParentCameraCompensation;
+
+				pxr::GfMatrix4d UsdTransform = UnrealToUsd::ConvertTransform(StageInfo, RelativeTransform);
 				Attr.Set<pxr::GfMatrix4d>(UsdTransform, UsdTimeCode);
 			};
 		}
@@ -3772,8 +3784,7 @@ bool UnrealToUsd::CreateComponentPropertyBaker(
 			TEXT("FocusSettings.ManualFocusDistance"),
 			TEXT("CurrentAperture"),
 			TEXT("Filmback.SensorWidth"),
-			TEXT("Filmback.SensorHeight")
-		};
+			TEXT("Filmback.SensorHeight")};
 
 		if (RelevantProperties.Contains(PropertyPath))
 		{
@@ -3852,8 +3863,7 @@ bool UnrealToUsd::CreateComponentPropertyBaker(
 	{
 		static TSet<FString> RelevantProperties = {
 			GET_MEMBER_NAME_CHECKED(UUsdDrawModeComponent, BoundsMin).ToString(),
-			GET_MEMBER_NAME_CHECKED(UUsdDrawModeComponent, BoundsMax).ToString()
-		};
+			GET_MEMBER_NAME_CHECKED(UUsdDrawModeComponent, BoundsMax).ToString()};
 
 		if (RelevantProperties.Contains(PropertyPath))
 		{
@@ -4126,37 +4136,41 @@ UnrealToUsd::FPropertyTrackWriter UnrealToUsd::CreatePropertyTrackWriter(
 					{
 						Attr = TransformOp.GetAttr();
 
+						const bool bStageIsZUp = StageInfo.UpAxis == EUsdUpAxis::ZAxis;
+
 						// Compensate different orientation for light or camera components
-						FTransform AdditionalRotation = FTransform::Identity;
+						FTransform Compensation = FTransform::Identity;
 						if (UsdPrim.IsA<pxr::UsdGeomCamera>() || UsdPrim.HasAPI<pxr::UsdLuxLightAPI>())
 						{
-							AdditionalRotation = FTransform(FRotator(0.0f, 90.0f, 0.0f));
+							Compensation = FTransform(FRotator(0.0f, 90.0f, 0.0f));
 
-							if (StageInfo.UpAxis == EUsdUpAxis::ZAxis)
+							if (bStageIsZUp)
 							{
-								AdditionalRotation *= FTransform(FRotator(90.0f, 0.0f, 0.0f));
+								Compensation *= FTransform(FRotator(90.0f, 0.0f, 0.0f));
 							}
 						}
 
 						// Invert compensation applied to parent if it's a light or camera component
+						FTransform InverseParentCompensation = FTransform::Identity;
 						if (const USceneComponent* AttachParent = Component.GetAttachParent())
 						{
 							if (AttachParent->IsA(UCineCameraComponent::StaticClass()) || AttachParent->IsA(ULightComponent::StaticClass()))
 							{
-								FTransform InverseCompensation = FTransform(FRotator(0.0f, 90.f, 0.0f));
+								InverseParentCompensation = FTransform(FRotator(0.0f, 90.f, 0.0f));
 
-								if (UsdUtils::GetUsdStageUpAxis(UsdStage) == pxr::UsdGeomTokens->z)
+								if (bStageIsZUp)
 								{
-									InverseCompensation *= FTransform(FRotator(90.0f, 0.f, 0.0f));
+									InverseParentCompensation *= FTransform(FRotator(90.0f, 0.f, 0.0f));
 								}
 
-								AdditionalRotation = AdditionalRotation * InverseCompensation.Inverse();
+								InverseParentCompensation = InverseParentCompensation.Inverse();
 							}
 						}
 
-						Result.TransformWriter = [&Component, AdditionalRotation, StageInfo, Attr](const FTransform& UEValue, double UsdTimeCode)
+						Result.TransformWriter =
+							[Compensation, InverseParentCompensation, StageInfo, Attr](const FTransform& UEValue, double UsdTimeCode)
 						{
-							FTransform FinalUETransform = AdditionalRotation * UEValue;
+							FTransform FinalUETransform = Compensation * UEValue * InverseParentCompensation;
 							pxr::GfMatrix4d UsdTransform = UnrealToUsd::ConvertTransform(StageInfo, FinalUETransform);
 							Attr.Set<pxr::GfMatrix4d>(UsdTransform, UsdTimeCode);
 						};
@@ -4594,13 +4608,11 @@ UnrealToUsd::FPropertyTrackWriter UnrealToUsd::CreatePropertyTrackWriter(
 					pxr::GfVec3f UsdMin{
 						FMath::Min(UEBoundsMinUsdSpace[0], UEBoundsMaxUsdSpace[0]),
 						FMath::Min(UEBoundsMinUsdSpace[1], UEBoundsMaxUsdSpace[1]),
-						FMath::Min(UEBoundsMinUsdSpace[2], UEBoundsMaxUsdSpace[2])
-					};
+						FMath::Min(UEBoundsMinUsdSpace[2], UEBoundsMaxUsdSpace[2])};
 					pxr::GfVec3f UsdMax{
 						FMath::Max(UEBoundsMinUsdSpace[0], UEBoundsMaxUsdSpace[0]),
 						FMath::Max(UEBoundsMinUsdSpace[1], UEBoundsMaxUsdSpace[1]),
-						FMath::Max(UEBoundsMinUsdSpace[2], UEBoundsMaxUsdSpace[2])
-					};
+						FMath::Max(UEBoundsMinUsdSpace[2], UEBoundsMaxUsdSpace[2])};
 					pxr::VtArray<pxr::GfVec3f> Extents{UsdMin, UsdMax};
 
 					Attr.Set(Extents, UsdTimeCode);
@@ -4925,23 +4937,20 @@ TArray<UE::FUsdAttribute> UnrealToUsd::GetAttributesForProperty(const UE::FUsdPr
 				UE::FUsdAttribute{UsdPrim.GetAttribute(pxr::UsdLuxTokens->inputsIntensity)},
 				UE::FUsdAttribute{UsdPrim.GetAttribute(pxr::UsdLuxTokens->inputsExposure)},
 				UE::FUsdAttribute{UsdPrim.GetAttribute(pxr::UsdLuxTokens->inputsWidth)},
-				UE::FUsdAttribute{UsdPrim.GetAttribute(pxr::UsdLuxTokens->inputsHeight)}
-			};
+				UE::FUsdAttribute{UsdPrim.GetAttribute(pxr::UsdLuxTokens->inputsHeight)}};
 		}
 		else if (UsdPrim.IsA<pxr::UsdLuxDiskLight>())
 		{
 			return {
 				UE::FUsdAttribute{UsdPrim.GetAttribute(pxr::UsdLuxTokens->inputsIntensity)},
 				UE::FUsdAttribute{UsdPrim.GetAttribute(pxr::UsdLuxTokens->inputsExposure)},
-				UE::FUsdAttribute{UsdPrim.GetAttribute(pxr::UsdLuxTokens->inputsRadius)}
-			};
+				UE::FUsdAttribute{UsdPrim.GetAttribute(pxr::UsdLuxTokens->inputsRadius)}};
 		}
 		else if (UsdPrim.IsA<pxr::UsdLuxDistantLight>())
 		{
 			return {
 				UE::FUsdAttribute{UsdPrim.GetAttribute(pxr::UsdLuxTokens->inputsIntensity)},
-				UE::FUsdAttribute{UsdPrim.GetAttribute(pxr::UsdLuxTokens->inputsExposure)}
-			};
+				UE::FUsdAttribute{UsdPrim.GetAttribute(pxr::UsdLuxTokens->inputsExposure)}};
 		}
 		else if (UsdPrim.IsA<pxr::UsdLuxSphereLight>())
 		{
@@ -4952,16 +4961,14 @@ TArray<UE::FUsdAttribute> UnrealToUsd::GetAttributesForProperty(const UE::FUsdPr
 					UE::FUsdAttribute{UsdPrim.GetAttribute(pxr::UsdLuxTokens->inputsExposure)},
 					UE::FUsdAttribute{UsdPrim.GetAttribute(pxr::UsdLuxTokens->inputsRadius)},
 					UE::FUsdAttribute{UsdPrim.GetAttribute(pxr::UsdLuxTokens->inputsShapingConeAngle)},
-					UE::FUsdAttribute{UsdPrim.GetAttribute(pxr::UsdLuxTokens->inputsShapingConeSoftness)}
-				};
+					UE::FUsdAttribute{UsdPrim.GetAttribute(pxr::UsdLuxTokens->inputsShapingConeSoftness)}};
 			}
 			else
 			{
 				return {
 					UE::FUsdAttribute{UsdPrim.GetAttribute(pxr::UsdLuxTokens->inputsIntensity)},
 					UE::FUsdAttribute{UsdPrim.GetAttribute(pxr::UsdLuxTokens->inputsExposure)},
-					UE::FUsdAttribute{UsdPrim.GetAttribute(pxr::UsdLuxTokens->inputsRadius)}
-				};
+					UE::FUsdAttribute{UsdPrim.GetAttribute(pxr::UsdLuxTokens->inputsRadius)}};
 			}
 		}
 	}
@@ -5011,8 +5018,7 @@ TArray<UE::FUsdAttribute> UnrealToUsd::GetAttributesForProperty(const UE::FUsdPr
 	{
 		return {
 			UE::FUsdAttribute{UsdPrim.GetAttribute(pxr::UsdLuxTokens->inputsShapingConeAngle)},
-			UE::FUsdAttribute{UsdPrim.GetAttribute(pxr::UsdLuxTokens->inputsShapingConeSoftness)}
-		};
+			UE::FUsdAttribute{UsdPrim.GetAttribute(pxr::UsdLuxTokens->inputsShapingConeSoftness)}};
 	}
 	else if (PropertyPath == LightSourceAnglePropertyName)
 	{
@@ -5092,13 +5098,11 @@ bool UnrealToUsd::ConvertBoundsComponent(const UUsdDrawModeComponent& BoundsComp
 		pxr::GfVec3f UsdMin{
 			FMath::Min(UEBoundsMinUsdSpace[0], UEBoundsMaxUsdSpace[0]),
 			FMath::Min(UEBoundsMinUsdSpace[1], UEBoundsMaxUsdSpace[1]),
-			FMath::Min(UEBoundsMinUsdSpace[2], UEBoundsMaxUsdSpace[2])
-		};
+			FMath::Min(UEBoundsMinUsdSpace[2], UEBoundsMaxUsdSpace[2])};
 		pxr::GfVec3f UsdMax{
 			FMath::Max(UEBoundsMinUsdSpace[0], UEBoundsMaxUsdSpace[0]),
 			FMath::Max(UEBoundsMinUsdSpace[1], UEBoundsMaxUsdSpace[1]),
-			FMath::Max(UEBoundsMinUsdSpace[2], UEBoundsMaxUsdSpace[2])
-		};
+			FMath::Max(UEBoundsMinUsdSpace[2], UEBoundsMaxUsdSpace[2])};
 
 		pxr::VtArray<pxr::GfVec3f> Extents{UsdMin, UsdMax};
 
@@ -5298,8 +5302,7 @@ namespace UE::USDPrimConversionImpl::Private
 			// definition itself and we'll likely run into trouble if we try writing anything that differs from it
 			const static TSet<FString> FieldsToSkip = {
 				UsdToUnreal::ConvertToken(pxr::SdfFieldKeys->Specifier),
-				UsdToUnreal::ConvertToken(pxr::SdfFieldKeys->TypeName)
-			};
+				UsdToUnreal::ConvertToken(pxr::SdfFieldKeys->TypeName)};
 			if (FieldsToSkip.Contains(FullKeyPath))
 			{
 				continue;
