USD: Fix how changing the stage actor's root layer property directly wasn't actually loading the stage anymore.

diff --git a/Engine/Plugins/Importers/USDImporter/Source/USDStage/Private/USDStageActor.cpp b/Engine/Plugins/Importers/USDImporter/Source/USDStage/Private/USDStageActor.cpp
--- a/Engine/Plugins/Importers/USDImporter/Source/USDStage/Private/USDStageActor.cpp
+++ b/Engine/Plugins/Importers/USDImporter/Source/USDStage/Private/USDStageActor.cpp
@@ -4655,10 +4655,10 @@ void AUsdStageActor::HandlePropertyChangedEvent(FPropertyChangedEvent& PropertyC
 
 	// If we're changing a property inside a struct, like "bCollectMetadata" inside our MetadataOptions, then
 	// "MemberProperty" will point to "MetadataOptions", and "Property" is the thing that will point to "bCollectMetadata"
-	FProperty* PropertyThatChanged = PropertyChangedEvent.Property;
-	const FName PropertyName = PropertyThatChanged ? PropertyThatChanged->GetFName() : NAME_None;
+	const FName PropertyName = PropertyChangedEvent.Property ? PropertyChangedEvent.Property->GetFName() : NAME_None;
+	const FName MemberPropertyName = PropertyChangedEvent.MemberProperty ? PropertyChangedEvent.MemberProperty->GetFName() : NAME_None;
 
-	if (PropertyName == GET_MEMBER_NAME_CHECKED(AUsdStageActor, RootLayer))
+	if (MemberPropertyName == GET_MEMBER_NAME_CHECKED(AUsdStageActor, RootLayer))
 	{
 #if WITH_EDITOR
 		static IConsoleVariable* CVar = IConsoleManager::Get().FindConsoleVariable(TEXT("USD.DiscardUndoBufferOnStageOpenClose"));
