USD: Fix crash when importing stages with SkelBindingAPI.

diff --git a/Engine/Plugins/Importers/USDImporter/Source/USDStageImporter/Private/USDStageImporter.cpp b/Engine/Plugins/Importers/USDImporter/Source/USDStageImporter/Private/USDStageImporter.cpp
--- a/Engine/Plugins/Importers/USDImporter/Source/USDStageImporter/Private/USDStageImporter.cpp
+++ b/Engine/Plugins/Importers/USDImporter/Source/USDStageImporter/Private/USDStageImporter.cpp
@@ -32,6 +32,7 @@
 #include "Components/SkeletalMeshComponent.h"
 #include "Components/StaticMeshComponent.h"
 #include "Editor.h"
+#include "Editor/Transactor.h"
 #include "Engine/SkeletalMesh.h"
 #include "Engine/SkinnedAssetCommon.h"
 #include "Engine/StaticMesh.h"
@@ -1896,9 +1897,14 @@ void UUsdStageImporter::ImportFromFile(FUsdStageImportContext& ImportContext)
 	// downstream from accidentally creating their own transactions. This happens for example on USkeleton::AccumulateCurveMetaData,
 	// and can lead to thousands of transactions showing up on the editor and a huge performance cost due to serialization spam.
 	// See also UEditorEngine::CanTransact.
+	TStrongObjectPtr<UTransactor> TransactorPin;
 	TOptional<TGuardValue<TObjectPtr<class UTransactor>>> TransactionSuppressor;
 	if (!GIsTransacting)
 	{
+		// Some stuff like the DuplicateObject or FBlueprintCompilationManager::CompileSynchronously inside SkelSkeletonTranslator
+		// will trigger GC, which we cannot prevent from here. If that happened when our transactor wasn't being referenced via
+		// GEditor->Trans it would actually get collected, so here we prevent that from happening
+		TransactorPin.Reset(GEditor->Trans.Get());
 		TransactionSuppressor.Emplace(GEditor->Trans, nullptr);
 	}
 
@@ -2085,6 +2091,21 @@ bool UUsdStageImporter::ReimportSingleAsset(
 #if USE_USD_SDK
 	double StartTime = FPlatformTime::Cycles64();
 
+	// If we don't have any transaction yet, let's temporarly disable the transaction buffer, to prevent code
+	// downstream from accidentally creating their own transactions. This happens for example on USkeleton::AccumulateCurveMetaData,
+	// and can lead to thousands of transactions showing up on the editor and a huge performance cost due to serialization spam.
+	// See also UEditorEngine::CanTransact.
+	TStrongObjectPtr<UTransactor> TransactorPin;
+	TOptional<TGuardValue<TObjectPtr<class UTransactor>>> TransactionSuppressor;
+	if (!GIsTransacting)
+	{
+		// Some stuff like the DuplicateObject or FBlueprintCompilationManager::CompileSynchronously inside SkelSkeletonTranslator
+		// will trigger GC, which we cannot prevent from here. If that happened when our transactor wasn't being referenced via
+		// GEditor->Trans it would actually get collected, so here we prevent that from happening
+		TransactorPin.Reset(GEditor->Trans.Get());
+		TransactionSuppressor.Emplace(GEditor->Trans, nullptr);
+	}
+
 	if (!ImportContext.ImportOptions)
 	{
 		ImportContext.ImportOptions = NewObject<UUsdStageImportOptions>();
