USD: Fix ensures triggering on invalid data.

diff --git a/Engine/Plugins/Importers/USDImporter/Source/USDUtilities/Private/USDGeomMeshConversion.cpp b/Engine/Plugins/Importers/USDImporter/Source/USDUtilities/Private/USDGeomMeshConversion.cpp
--- a/Engine/Plugins/Importers/USDImporter/Source/USDUtilities/Private/USDGeomMeshConversion.cpp
+++ b/Engine/Plugins/Importers/USDImporter/Source/USDUtilities/Private/USDGeomMeshConversion.cpp
@@ -2207,7 +2207,6 @@ namespace UE::UsdGeomMeshConversion::Private
 			struct FUVSet
 			{
 				int32 UVSetIndexUE;	   // The user may only have 'uv4' and 'uv5', so we can't just use array indices to find the target UV channel
-				TOptional<pxr::VtIntArray> UVIndices;	 // UVs might be indexed or they might be flat (one per vertex)
 				pxr::VtVec2fArray UVs;
 
 				pxr::TfToken InterpType = pxr::UsdGeomTokens->faceVarying;
@@ -2224,7 +2223,6 @@ namespace UE::UsdGeomMeshConversion::Private
 
 				if (InMeshData.UVSetIndices[UVChannelIndex].size() > 0)
 				{
-					UVSet.UVIndices = InMeshData.UVSetIndices[UVChannelIndex];
 					UVSet.UVs = InMeshData.UVSets[UVChannelIndex];
 
 					if (UVSet.UVs.size() > 0)
@@ -2350,29 +2348,30 @@ namespace UE::UsdGeomMeshConversion::Private
 						}
 					}
 
-					for (const FUVSet& UVSet : UVSets)
+					for (int32 UVSetIndex = 0; UVSetIndex < UVSets.Num(); ++UVSetIndex)
 					{
+						const FUVSet& UVSet = UVSets[UVSetIndex];
+
 						const size_t ValueIndex = GetPrimValueIndex(UVSet.InterpType, ControlPointIndex, CurrentVertexInstanceIndex, PolygonIndex);
 
 						pxr::GfVec2f UV(0.f, 0.f);
 
-						if (UVSet.UVIndices.IsSet())
+						if (UVSet.UVs.size() > ValueIndex)
 						{
-							const pxr::VtIntArray& UVIndices = UVSet.UVIndices.GetValue();
-
-							if (ensure(ValueIndex < UVIndices.size()))
-							{
-								size_t UVIndex = UVIndices[ValueIndex];
-
-								if (ensure(UVIndex < UVSet.UVs.size()))
-								{
-									UV = UVSet.UVs[UVIndex];
-								}
-							}
+							UV = UVSet.UVs[ValueIndex];
 						}
-						else if (ensure(UVSet.UVs.size() > ValueIndex))
+						else
 						{
-							UV = UVSet.UVs[ValueIndex];
+							UE_LOG(
+								LogUsd,
+								Warning,
+								TEXT("Trying to read UV at index %u from prim '%s' but the UV set %d only has %u values! Using zeros "
+									 "instead."),
+								ValueIndex,
+								*InMeshData.SourcePrimPath,
+								UVSetIndex,
+								UVSet.UVs.size()
+							);
 						}
 
 						// Flip V for Unreal uv's which match directx
@@ -2391,9 +2390,24 @@ namespace UE::UsdGeomMeshConversion::Private
 
 						pxr::GfVec3f UsdColor(1.f, 1.f, 1.f);
 
-						if (!InMeshData.DisplayColors.empty() && ensure(InMeshData.DisplayColors.size() > ValueIndex))
+						if (!InMeshData.DisplayColors.empty())
 						{
-							UsdColor = InMeshData.DisplayColors.cdata()[ValueIndex];
+							if (InMeshData.DisplayColors.size() > ValueIndex)
+							{
+								UsdColor = InMeshData.DisplayColors.cdata()[ValueIndex];
+							}
+							else
+							{
+								UE_LOG(
+									LogUsd,
+									Warning,
+									TEXT("Trying to read displayColor at index %u from prim '%s' but the prim only has %u values! Using "
+										 "zeros instead."),
+									ValueIndex,
+									*InMeshData.SourcePrimPath,
+									InMeshData.DisplayColors.size()
+								);
+							}
 						}
 
 						MeshDescriptionColors[AddedVertexInstanceId] = UsdToUnreal::ConvertColor(UsdColor);
@@ -2408,9 +2422,24 @@ namespace UE::UsdGeomMeshConversion::Private
 							PolygonIndex
 						);
 
-						if (!InMeshData.DisplayOpacities.empty() && ensure(InMeshData.DisplayOpacities.size() > ValueIndex))
+						if (!InMeshData.DisplayOpacities.empty())
 						{
-							MeshDescriptionColors[AddedVertexInstanceId][3] = InMeshData.DisplayOpacities.cdata()[ValueIndex];
+							if (InMeshData.DisplayOpacities.size() > ValueIndex)
+							{
+								MeshDescriptionColors[AddedVertexInstanceId][3] = InMeshData.DisplayOpacities.cdata()[ValueIndex];
+							}
+							else
+							{
+								UE_LOG(
+									LogUsd,
+									Warning,
+									TEXT("Trying to read displayOpacity at index %u from prim '%s' but the prim only has %u values! Using "
+										 "zeros instead."),
+									ValueIndex,
+									*InMeshData.SourcePrimPath,
+									InMeshData.DisplayColors.size()
+								);
+							}
 						}
 					}
 				}
