USD: Fix material assignments not being factored in for the asset hash when exporting static and skeletal meshes.

diff --git a/Engine/Plugins/Importers/USDImporter/Source/USDExporter/Private/SkeletalMeshExporterUSD.cpp b/Engine/Plugins/Importers/USDImporter/Source/USDExporter/Private/SkeletalMeshExporterUSD.cpp
--- a/Engine/Plugins/Importers/USDImporter/Source/USDExporter/Private/SkeletalMeshExporterUSD.cpp
+++ b/Engine/Plugins/Importers/USDImporter/Source/USDExporter/Private/SkeletalMeshExporterUSD.cpp
@@ -62,6 +62,28 @@ namespace UE::SkeletalMeshExporterUSD::Private
 			Extension
 		);
 	}
+
+	void HashSkeletalMesh(const USkeletalMesh* SkeletalMesh, FSHA1& InOutHashToUpdate)
+	{
+		if (const FSkeletalMeshRenderData* RenderData = SkeletalMesh->GetResourceForRendering())
+		{
+			const FString& DDCKey = RenderData->DerivedDataKey;
+			InOutHashToUpdate.UpdateWithString(*DDCKey, DDCKey.Len());
+		}
+
+		for (const FSkeletalMaterial& SkeletalMaterial : SkeletalMesh->GetMaterials())
+		{
+			FString MaterialPath = TEXT("None");
+			if (SkeletalMaterial.MaterialInterface)
+			{
+				MaterialPath = SkeletalMaterial.MaterialInterface->GetPathName();
+			}
+			InOutHashToUpdate.UpdateWithString(*MaterialPath, MaterialPath.Len());
+
+			// Note that we could hash the material slot name here too, but we don't because we always
+			// just write out the slots with UsdGeomSubsets named "Section0", "Section1", ..., "SectionN" anyway
+		}
+	}
 }
 
 USkeletalMeshExporterUsd::USkeletalMeshExporterUsd()
@@ -157,18 +179,13 @@ bool USkeletalMeshExporterUsd::ExportBinary(
 	}
 
 	// Get a simple GUID hash/identifier of our mesh
-	FString DDCKeyHash;
-	if (FSkeletalMeshRenderData* RenderData = SkeletalMesh->GetResourceForRendering())
-	{
-		const FString& DDCKey = RenderData->DerivedDataKey;
-		FSHA1 SHA1;
-		SHA1.UpdateWithString(*DDCKey, DDCKey.Len());
-		UsdUtils::HashForSkeletalMeshExport(*Options, SHA1);
-		SHA1.Final();
-		FSHAHash Hash;
-		SHA1.GetHash(&Hash.Hash[0]);
-		DDCKeyHash = Hash.ToString();
-	}
+	FSHA1 SHA1;
+	UE::SkeletalMeshExporterUSD::Private::HashSkeletalMesh(SkeletalMesh, SHA1);
+	UsdUtils::HashForSkeletalMeshExport(*Options, SHA1);
+	SHA1.Final();
+	FSHAHash Hash;
+	SHA1.GetHash(&Hash.Hash[0]);
+	FString CurrentHashString = Hash.ToString();
 
 	// Check if we already have exported what we plan on exporting anyway
 	if (FPaths::FileExists(UExporter::CurrentFilename) && FPaths::FileExists(PayloadFilename))
@@ -199,7 +216,7 @@ bool USkeletalMeshExporterUsd::ExportBinary(
 				{
 					FUsdUnrealAssetInfo Info = UsdUtils::GetPrimAssetInfo(DefaultPrim);
 
-					const bool bVersionMatches = !Info.Version.IsEmpty() && Info.Version == DDCKeyHash;
+					const bool bVersionMatches = !Info.Version.IsEmpty() && Info.Version == CurrentHashString;
 
 					const bool bAssetTypeMatches = !Info.UnrealAssetType.IsEmpty() && Info.UnrealAssetType == SkeletalMesh->GetClass()->GetName();
 
@@ -315,7 +332,7 @@ bool USkeletalMeshExporterUsd::ExportBinary(
 			FUsdUnrealAssetInfo Info;
 			Info.Name = SkeletalMesh->GetName();
 			Info.Identifier = UExporter::CurrentFilename;
-			Info.Version = DDCKeyHash;
+			Info.Version = CurrentHashString;
 			Info.UnrealContentPath = SkeletalMesh->GetPathName();
 			Info.UnrealAssetType = SkeletalMesh->GetClass()->GetName();
 			Info.UnrealExportTime = FDateTime::Now().ToString();
diff --git a/Engine/Plugins/Importers/USDImporter/Source/USDExporter/Private/StaticMeshExporterUSD.cpp b/Engine/Plugins/Importers/USDImporter/Source/USDExporter/Private/StaticMeshExporterUSD.cpp
--- a/Engine/Plugins/Importers/USDImporter/Source/USDExporter/Private/StaticMeshExporterUSD.cpp
+++ b/Engine/Plugins/Importers/USDImporter/Source/USDExporter/Private/StaticMeshExporterUSD.cpp
@@ -60,6 +60,28 @@ namespace UE::StaticMeshExporterUSD::Private
 			Extension
 		);
 	}
+
+	void HashStaticMesh(const UStaticMesh* StaticMesh, FSHA1& InOutHashToUpdate)
+	{
+		if (const FStaticMeshRenderData* RenderData = StaticMesh->GetRenderData())
+		{
+			const FString& DDCKey = RenderData->DerivedDataKey;
+			InOutHashToUpdate.UpdateWithString(*DDCKey, DDCKey.Len());
+		}
+
+		for (const FStaticMaterial& StaticMaterial : StaticMesh->GetStaticMaterials())
+		{
+			FString MaterialPath = TEXT("None");
+			if (StaticMaterial.MaterialInterface)
+			{
+				MaterialPath = StaticMaterial.MaterialInterface->GetPathName();
+			}
+			InOutHashToUpdate.UpdateWithString(*MaterialPath, MaterialPath.Len());
+
+			// Note that we could hash the material slot name here too, but we don't because we always
+			// just write out the slots with UsdGeomSubsets named "Section0", "Section1", ..., "SectionN" anyway
+		}
+	}
 }
 
 bool UStaticMeshExporterUsd::IsUsdAvailable()
@@ -161,19 +183,14 @@ bool UStaticMeshExporterUsd::ExportBinary(UObject* Object, const TCHAR* Type, FA
 		return false;
 	}
 
-	// Get a simple GUID hash/identifier of our mesh
-	FString DDCKeyHash;
-	if (FStaticMeshRenderData* RenderData = StaticMesh->GetRenderData())
-	{
-		const FString& DDCKey = RenderData->DerivedDataKey;
-		FSHA1 SHA1;
-		SHA1.UpdateWithString(*DDCKey, DDCKey.Len());
-		UsdUtils::HashForStaticMeshExport(*Options, SHA1);
-		SHA1.Final();
-		FSHAHash Hash;
-		SHA1.GetHash(&Hash.Hash[0]);
-		DDCKeyHash = Hash.ToString();
-	}
+	// Get a simple GUID hash/identifier of our mesh and options
+	FSHA1 SHA1;
+	UE::StaticMeshExporterUSD::Private::HashStaticMesh(StaticMesh, SHA1);
+	UsdUtils::HashForStaticMeshExport(*Options, SHA1);
+	SHA1.Final();
+	FSHAHash Hash;
+	SHA1.GetHash(&Hash.Hash[0]);
+	FString CurrentHashString = Hash.ToString();
 
 	// Check if we already have exported what we plan on exporting anyway
 	if (FPaths::FileExists(UExporter::CurrentFilename) && FPaths::FileExists(PayloadFilename))
@@ -204,7 +221,7 @@ bool UStaticMeshExporterUsd::ExportBinary(UObject* Object, const TCHAR* Type, FA
 				{
 					FUsdUnrealAssetInfo Info = UsdUtils::GetPrimAssetInfo(DefaultPrim);
 
-					const bool bVersionMatches = !Info.Version.IsEmpty() && Info.Version == DDCKeyHash;
+					const bool bVersionMatches = !Info.Version.IsEmpty() && Info.Version == CurrentHashString;
 
 					const bool bAssetTypeMatches = !Info.UnrealAssetType.IsEmpty() && Info.UnrealAssetType == StaticMesh->GetClass()->GetName();
 
@@ -318,7 +335,7 @@ bool UStaticMeshExporterUsd::ExportBinary(UObject* Object, const TCHAR* Type, FA
 			FUsdUnrealAssetInfo Info;
 			Info.Name = StaticMesh->GetName();
 			Info.Identifier = UExporter::CurrentFilename;
-			Info.Version = DDCKeyHash;
+			Info.Version = CurrentHashString;
 			Info.UnrealContentPath = StaticMesh->GetPathName();
 			Info.UnrealAssetType = StaticMesh->GetClass()->GetName();
 			Info.UnrealExportTime = FDateTime::Now().ToString();
