USD: Fix ensure when enabling selection sync with StageState "Opened" and selecting prims.

diff --git a/Engine/Plugins/Importers/USDImporter/Source/USDStage/Private/USDStageActor.cpp b/Engine/Plugins/Importers/USDImporter/Source/USDStage/Private/USDStageActor.cpp
--- a/Engine/Plugins/Importers/USDImporter/Source/USDStage/Private/USDStageActor.cpp
+++ b/Engine/Plugins/Importers/USDImporter/Source/USDStage/Private/USDStageActor.cpp
@@ -2755,7 +2755,7 @@ ULevelSequence* AUsdStageActor::GetLevelSequence()
 USceneComponent* AUsdStageActor::GetGeneratedComponent(const FString& PrimPath)
 {
 	const UE::FUsdStage& CurrentStage = static_cast<const AUsdStageActor*>(this)->GetUsdStage();
-	if (!CurrentStage)
+	if (!CurrentStage || StageState != EUsdStageState::OpenedAndLoaded)
 	{
 		return nullptr;
 	}
@@ -2784,7 +2784,7 @@ USceneComponent* AUsdStageActor::GetGeneratedComponent(const FString& PrimPath)
 TArray<UObject*> AUsdStageActor::GetGeneratedAssets(const FString& PrimPath)
 {
 	const UE::FUsdStage& CurrentStage = static_cast<const AUsdStageActor*>(this)->GetUsdStage();
-	if (!CurrentStage)
+	if (!CurrentStage || StageState != EUsdStageState::OpenedAndLoaded)
 	{
 		return {};
 	}
