USD: Remove ensure that forced all track changes to come from Movable components.

diff --git a/Engine/Plugins/Importers/USDImporter/Source/USDStage/Private/USDLevelSequenceHelper.cpp b/Engine/Plugins/Importers/USDImporter/Source/USDStage/Private/USDLevelSequenceHelper.cpp
--- a/Engine/Plugins/Importers/USDImporter/Source/USDStage/Private/USDLevelSequenceHelper.cpp
+++ b/Engine/Plugins/Importers/USDImporter/Source/USDStage/Private/USDLevelSequenceHelper.cpp
@@ -3777,8 +3777,6 @@ void FUsdLevelSequenceHelperImpl::HandleTrackChange(const UMovieSceneTrack& Trac
 
 	if (PrimTwin)
 	{
-		ensure(BoundSceneComponent->Mobility != EComponentMobility::Static);
-
 		FScopedBlockNoticeListening BlockNotices(StageActor.Get());
 		UE::FUsdPrim UsdPrim = UsdStage.GetPrimAtPath(UE::FSdfPath(*PrimTwin->PrimPath));
 
