USD: Fix "kind" prim property being shown twice.

diff --git a/Engine/Plugins/Importers/USDImporter/Source/USDStageEditorViewModels/Private/USDObjectFieldViewModel.cpp b/Engine/Plugins/Importers/USDImporter/Source/USDStageEditorViewModels/Private/USDObjectFieldViewModel.cpp
--- a/Engine/Plugins/Importers/USDImporter/Source/USDStageEditorViewModels/Private/USDObjectFieldViewModel.cpp
+++ b/Engine/Plugins/Importers/USDImporter/Source/USDStageEditorViewModels/Private/USDObjectFieldViewModel.cpp
@@ -404,12 +404,6 @@ void FUsdObjectFieldsViewModel::Refresh(const UE::FUsdStageWeak& InUsdStage, con
 				bPrimReadOnly
 			);
 			CreateField(EObjectFieldType::Metadata, TEXT("path"), ObjectPath, UsdUtils::EUsdBasicDataTypes::String, Role, bPrimReadOnly);
-			CreateField(
-				EObjectFieldType::Metadata,
-				TEXT("kind"),
-				UsdToUnreal::ConvertString(IUsdPrim::GetKind(UsdPrim).GetString()),
-				UsdUtils::EUsdBasicDataTypes::Token
-			);
 
 			FScopedUsdAllocs UsdAllocs;
 
