USD: Remove remaining cases of exporting the 'unrealMaterial' attribute;

Fixup some minor missing formatting issues missed on the reformat somehow;
Update automated tests to match new conventions and to reuse some code for checking material assignments.

diff --git a/Engine/Plugins/Importers/USDImporter/Source/USDUtilities/Private/USDGeomMeshConversion.cpp b/Engine/Plugins/Importers/USDImporter/Source/USDUtilities/Private/USDGeomMeshConversion.cpp
--- a/Engine/Plugins/Importers/USDImporter/Source/USDUtilities/Private/USDGeomMeshConversion.cpp
+++ b/Engine/Plugins/Importers/USDImporter/Source/USDUtilities/Private/USDGeomMeshConversion.cpp
@@ -352,9 +352,9 @@ namespace UE::UsdGeomMeshConversion::Private
 		int32 LODIndex,
 		const FStaticMeshLODResources& LODRenderMesh,
 		pxr::UsdGeomMesh& UsdMesh,
-		const pxr::VtArray<std::string>& MaterialAssignments,
+		const TArray<FString>& MaterialAssignments,
 		const pxr::UsdTimeCode TimeCode,
-		pxr::UsdPrim MaterialPrim
+		pxr::UsdPrim PrimToReceiveMaterialAssignments
 	)
 	{
 		pxr::UsdPrim MeshPrim = UsdMesh.GetPrim();
@@ -524,47 +524,36 @@ namespace UE::UsdGeomMeshConversion::Private
 		{
 			bool bHasUEMaterialAssignements = false;
 
-			pxr::VtArray<std::string> UnrealMaterialsForLOD;
+			TArray<FString> UnrealMaterialsForLOD;
 			for (const FStaticMeshSection& Section : LODRenderMesh.Sections)
 			{
-				if (Section.MaterialIndex >= 0 && Section.MaterialIndex < MaterialAssignments.size())
+				if (MaterialAssignments.IsValidIndex(Section.MaterialIndex))
 				{
-					UnrealMaterialsForLOD.push_back(MaterialAssignments[Section.MaterialIndex]);
+					UnrealMaterialsForLOD.Add(MaterialAssignments[Section.MaterialIndex]);
 					bHasUEMaterialAssignements = true;
 				}
 				else
 				{
 					// Keep unrealMaterials with the same number of elements as our MaterialIndices expect
-					UnrealMaterialsForLOD.push_back("");
+					UnrealMaterialsForLOD.Add(TEXT(""));
 				}
 			}
 
-			// This LOD has a single material assignment, just add an unrealMaterials attribute to the mesh prim
-			if (bHasUEMaterialAssignements && UnrealMaterialsForLOD.size() == 1)
+			// This LOD has a single material assignment, just create/bind an UnrealMaterial child prim directly
+			if (bHasUEMaterialAssignements && UnrealMaterialsForLOD.Num() == 1)
 			{
-				if (pxr::UsdAttribute UEMaterialsAttribute = MaterialPrim.CreateAttribute(
-						UnrealIdentifiers::MaterialAssignment,
-						pxr::SdfValueTypeNames->String
-					))
-				{
-					UEMaterialsAttribute.Set(UnrealMaterialsForLOD[0]);
-				}
+				UsdUtils::AuthorUnrealMaterialBinding(PrimToReceiveMaterialAssignments, UnrealMaterialsForLOD[0]);
 			}
 			// Multiple material assignments to the same LOD (and so the same mesh prim). Need to create a GeomSubset for each UE mesh section
-			else if (UnrealMaterialsForLOD.size() > 1)
+			else if (UnrealMaterialsForLOD.Num() > 1)
 			{
 				// Need to fetch all triangles of a section, and add their indices
 				for (int32 SectionIndex = 0; SectionIndex < LODRenderMesh.Sections.Num(); ++SectionIndex)
 				{
 					const FStaticMeshSection& Section = LODRenderMesh.Sections[SectionIndex];
 
-					// Note that we will continue on even if we have no material assignment, so as to satisfy the "partition" family condition (below)
-					std::string SectionMaterial;
-					if (Section.MaterialIndex >= 0 && Section.MaterialIndex < MaterialAssignments.size())
-					{
-						SectionMaterial = MaterialAssignments[Section.MaterialIndex];
-					}
-
+					// Note that we will continue authoring the GeomSubsets on even if we later find out we have no material assignment (just
+					// "") for this section, so as to satisfy the "partition" family condition (below)
 					pxr::UsdPrim GeomSubsetPrim = Stage->DefinePrim(
 						MeshPrim.GetPath().AppendPath(pxr::SdfPath("Section" + std::to_string(SectionIndex))),
 						UnrealToUsd::ConvertToken(TEXT("GeomSubset")).Get()
@@ -572,10 +561,10 @@ namespace UE::UsdGeomMeshConversion::Private
 
 					// MaterialPrim may be in another stage, so we may need another GeomSubset there
 					pxr::UsdPrim MaterialGeomSubsetPrim = GeomSubsetPrim;
-					if (MaterialPrim.GetStage() != MeshPrim.GetStage())
+					if (PrimToReceiveMaterialAssignments.GetStage() != MeshPrim.GetStage())
 					{
-						MaterialGeomSubsetPrim = MaterialPrim.GetStage()->OverridePrim(
-							MaterialPrim.GetPath().AppendPath(pxr::SdfPath("Section" + std::to_string(SectionIndex)))
+						MaterialGeomSubsetPrim = PrimToReceiveMaterialAssignments.GetStage()->OverridePrim(
+							PrimToReceiveMaterialAssignments.GetPath().AppendPath(pxr::SdfPath("Section" + std::to_string(SectionIndex)))
 						);
 					}
 
@@ -588,7 +577,6 @@ namespace UE::UsdGeomMeshConversion::Private
 					// Indices attribute
 					const uint32 TriangleCount = Section.NumTriangles;
 					const uint32 FirstTriangleIndex = Section.FirstIndex / 3;	 // FirstIndex is the first *vertex* instance index
-					FIndexArrayView VertexInstances = LODRenderMesh.IndexBuffer.GetArrayView();
 					pxr::VtArray<int> IndicesAttrValue;
 					for (uint32 TriangleIndex = FirstTriangleIndex; TriangleIndex - FirstTriangleIndex < TriangleCount; ++TriangleIndex)
 					{
@@ -607,14 +595,8 @@ namespace UE::UsdGeomMeshConversion::Private
 					// Family type
 					pxr::UsdGeomSubset::SetFamilyType(UsdMesh, pxr::UsdShadeTokens->materialBind, pxr::UsdGeomTokens->partition);
 
-					// unrealMaterial attribute
-					if (pxr::UsdAttribute UEMaterialsAttribute = MaterialGeomSubsetPrim.CreateAttribute(
-							UnrealIdentifiers::MaterialAssignment,
-							pxr::SdfValueTypeNames->String
-						))
-					{
-						UEMaterialsAttribute.Set(UnrealMaterialsForLOD[SectionIndex]);
-					}
+					// material:binding relationship
+					UsdUtils::AuthorUnrealMaterialBinding(MaterialGeomSubsetPrim, UnrealMaterialsForLOD[SectionIndex]);
 				}
 			}
 		}
@@ -1633,8 +1615,7 @@ namespace UE::UsdGeomMeshConversion::Private
 			InOutMeshData.CreaseLengths,
 			InOutMeshData.CreaseSharpnesses,
 			InOutMeshData.CornerIndices,
-			InOutMeshData.CornerSharpnesses
-		};
+			InOutMeshData.CornerSharpnesses};
 
 		pxr::PxOsdMeshTopology Topology{
 			InOutMeshData.SubdivScheme,
@@ -1642,8 +1623,7 @@ namespace UE::UsdGeomMeshConversion::Private
 			InOutMeshData.FaceVertexCounts,
 			InOutMeshData.FaceVertexIndices,
 			InOutMeshData.HoleIndices,
-			SubdivTags
-		};
+			SubdivTags};
 
 		std::vector<pxr::VtArray<int>> FaceVaryingTopologies;
 		FaceVaryingTopologies.resize(FaceVaryingChannelCounter);
@@ -3694,7 +3674,7 @@ bool UnrealToUsd::ConvertStaticMesh(
 
 	// Collect all material assignments, referenced by the sections' material indices
 	bool bHasMaterialAssignments = false;
-	pxr::VtArray<std::string> MaterialAssignments;
+	TArray<FString> MaterialAssignments;
 	for (const FStaticMaterial& StaticMaterial : StaticMesh->GetStaticMaterials())
 	{
 		FString AssignedMaterialPathName;
@@ -3707,12 +3687,12 @@ bool UnrealToUsd::ConvertStaticMesh(
 			}
 		}
 
-		MaterialAssignments.push_back(UnrealToUsd::ConvertString(*AssignedMaterialPathName).Get());
+		MaterialAssignments.Add(AssignedMaterialPathName);
 	}
 	if (!bHasMaterialAssignments)
 	{
-		// Prevent creation of the unrealMaterials attribute in case we don't have any assignments at all
-		MaterialAssignments.clear();
+		// Prevent creation of the UnrealMaterials prims in case we don't have any assignments at all
+		MaterialAssignments.Reset();
 	}
 
 	// Do this outside the variant edit context or else it's going to be a weaker opinion than the stuff outside
@@ -3981,9 +3961,9 @@ namespace UE::UsdGeometryCacheConversion::Private
 	void ConvertGeometryCacheMeshData(
 		const FGeometryCacheMeshData& MeshData,
 		pxr::UsdGeomMesh& UsdMesh,
-		const pxr::VtArray<std::string>& MaterialAssignments,
+		const TArray<FString>& MaterialAssignments,
 		const pxr::UsdTimeCode TimeCode,
-		pxr::UsdPrim MaterialPrim,
+		pxr::UsdPrim PrimToReceiveMaterialAssignments,
 		FGeometryCacheExportContext& ExportContext
 	)
 	{
@@ -4172,19 +4152,13 @@ namespace UE::UsdGeometryCacheConversion::Private
 
 		// Material assignments
 		{
-			// This mesh has a single material assignment, just add an unrealMaterials attribute to the mesh prim
-			if (MaterialAssignments.size() == 1)
+			// This LOD has a single material assignment, just create/bind an UnrealMaterial child prim directly
+			if (MaterialAssignments.Num() == 1)
 			{
-				if (pxr::UsdAttribute UEMaterialsAttribute = MaterialPrim.CreateAttribute(
-						UnrealIdentifiers::MaterialAssignment,
-						pxr::SdfValueTypeNames->String
-					))
-				{
-					UEMaterialsAttribute.Set(MaterialAssignments[0]);
-				}
+				UsdUtils::AuthorUnrealMaterialBinding(PrimToReceiveMaterialAssignments, MaterialAssignments[0]);
 			}
 			// Multiple material assignments to the same mesh. Need to create a GeomSubset for each UE mesh section
-			else if (MaterialAssignments.size() > 1)
+			else if (MaterialAssignments.Num() > 1)
 			{
 				TSet<FString> UsedSectionNames;
 				// Need to fetch all triangles of a section, and add their indices
@@ -4192,13 +4166,8 @@ namespace UE::UsdGeometryCacheConversion::Private
 				{
 					const FGeometryCacheMeshBatchInfo& Section = MeshData.BatchesInfo[SectionIndex];
 
-					// Note that we will continue on even if we have no material assignment, so as to satisfy the "partition" family condition (below)
-					std::string SectionMaterial;
-					if (Section.MaterialIndex >= 0 && Section.MaterialIndex < MaterialAssignments.size())
-					{
-						SectionMaterial = MaterialAssignments[Section.MaterialIndex];
-					}
-
+					// Note that we will continue authoring the GeomSubsets on even if we later find out we have no material assignment (just
+					// "") for this section, so as to satisfy the "partition" family condition (below)
 					FString SectionName;
 					if (ExportContext.SlotNames.IsValidIndex(SectionIndex))
 					{
@@ -4219,9 +4188,11 @@ namespace UE::UsdGeometryCacheConversion::Private
 
 					// MaterialPrim may be in another stage, so we may need another GeomSubset there
 					pxr::UsdPrim MaterialGeomSubsetPrim = GeomSubsetPrim;
-					if (MaterialPrim.GetStage() != MeshPrim.GetStage())
+					if (PrimToReceiveMaterialAssignments.GetStage() != MeshPrim.GetStage())
 					{
-						MaterialGeomSubsetPrim = MaterialPrim.GetStage()->OverridePrim(MaterialPrim.GetPath().AppendPath(PrimPath));
+						MaterialGeomSubsetPrim = PrimToReceiveMaterialAssignments.GetStage()->OverridePrim(
+							PrimToReceiveMaterialAssignments.GetPath().AppendPath(PrimPath)
+						);
 					}
 
 					pxr::UsdGeomSubset GeomSubsetSchema{GeomSubsetPrim};
@@ -4257,13 +4228,10 @@ namespace UE::UsdGeometryCacheConversion::Private
 						// Family type
 						pxr::UsdGeomSubset::SetFamilyType(UsdMesh, pxr::UsdShadeTokens->materialBind, pxr::UsdGeomTokens->partition);
 
-						// unrealMaterial attribute
-						if (pxr::UsdAttribute UEMaterialsAttribute = MaterialGeomSubsetPrim.CreateAttribute(
-								UnrealIdentifiers::MaterialAssignment,
-								pxr::SdfValueTypeNames->String
-							))
+						// material:binding relationship
+						if (MaterialAssignments.IsValidIndex(Section.MaterialIndex))
 						{
-							UEMaterialsAttribute.Set(SectionMaterial);
+							UsdUtils::AuthorUnrealMaterialBinding(MaterialGeomSubsetPrim, MaterialAssignments[Section.MaterialIndex]);
 						}
 					}
 				}
@@ -4290,7 +4258,7 @@ bool UnrealToUsd::ConvertGeometryCache(const UGeometryCache* GeometryCache, pxr:
 
 	// Collect all material assignments, referenced by the sections' material indices
 	bool bHasMaterialAssignments = false;
-	pxr::VtArray<std::string> MaterialAssignments;
+	TArray<FString> MaterialAssignments;
 	for (const UMaterialInterface* Material : GeometryCache->Materials)
 	{
 		FString AssignedMaterialPathName;
@@ -4303,12 +4271,12 @@ bool UnrealToUsd::ConvertGeometryCache(const UGeometryCache* GeometryCache, pxr:
 			}
 		}
 
-		MaterialAssignments.push_back(UnrealToUsd::ConvertString(*AssignedMaterialPathName).Get());
+		MaterialAssignments.Add(AssignedMaterialPathName);
 	}
 	if (!bHasMaterialAssignments)
 	{
 		// Prevent creation of the unrealMaterials attribute in case we don't have any assignments at all
-		MaterialAssignments.clear();
+		MaterialAssignments.Reset();
 	}
 
 	// Author material bindings on the dedicated stage if we have one
@@ -4878,6 +4846,14 @@ void UsdUtils::ReplaceUnrealMaterialsWithBaked(
 
 					if (!bAlreadyHasReference)
 					{
+						// Without this, if we tried exporting material overrides for LOD meshes they would
+						// end up outside of the variant set
+						TOptional<pxr::UsdEditContext> VarContext;
+						if (bAuthorInsideVariants)
+						{
+							VarContext.Emplace(OuterVariantSet.GetValue().GetVariantEditContext());
+						}
+
 						UE::FUsdPrim UEMatPrim{MatPrim};
 						UsdUtils::AddReference(UEMatPrim, *BakedFilename);
 					}
diff --git a/Engine/Plugins/Importers/USDImporter/Source/USDUtilities/Private/USDPrimConversion.cpp b/Engine/Plugins/Importers/USDImporter/Source/USDUtilities/Private/USDPrimConversion.cpp
--- a/Engine/Plugins/Importers/USDImporter/Source/USDUtilities/Private/USDPrimConversion.cpp
+++ b/Engine/Plugins/Importers/USDImporter/Source/USDUtilities/Private/USDPrimConversion.cpp
@@ -123,264 +123,169 @@ static FAutoConsoleVariableRef CVarConsiderAllPrimsHaveAnimatedBounds(
 		 "extremely expensive!")
 );
 
-namespace UE
+namespace UE::USDPrimConversion::Private
 {
-	namespace USDPrimConversionImpl
+	// On the current edit target, will set the Xformable's op order to a single "xformOp:transform",
+	// create the corresponding attribute, and return the op
+	pxr::UsdGeomXformOp ForceMatrixXform(pxr::UsdGeomXformable& Xformable)
 	{
-		namespace Private
-		{
-			// Writes UEMaterialAssetPath as a material binding for MeshPrim, either by reusing a Material binding if
-			// it already has an 'unreal' render context output and the expected structure, or by creating a new Material prim
-			// that fulfills those requirements.
-			// Doesn't write to 'unrealMaterial' at all, as we intend on deprecating it in the future.
-			void AuthorMaterialOverride(pxr::UsdPrim& MeshPrim, const FString& UEMaterialAssetPath)
-			{
-				if (!MeshPrim)
-				{
-					return;
-				}
-
-				FScopedUsdAllocs UsdAllocs;
+		FScopedUsdAllocs Allocs;
 
-				pxr::UsdShadeMaterialBindingAPI BindingAPI = pxr::UsdShadeMaterialBindingAPI::Apply(MeshPrim);
+		// Note: We don't use Xformable.MakeMatrixXform() here because while it can clear the
+		// xform op order on the current edit target just fine, it will later try to AddTransformOp(),
+		// which calls AddXformOp. Internally, it will read the *composed* prim and if it finds that it already
+		// has an op of that type it will early out and not author anything. This means that if our stage
+		// has a strong opinion for an e.g. "xformOp:transform" already on the layer stack, it's not possible
+		// to author that same op on a weaker layer. We want to do this here, to ensure this prim's transform
+		// works as expected even if this weaker layer is used standalone, so we must do the analogous ourselves
 
-				// If this mesh prim already has a binding to a *child* material with the 'unreal' render context,
-				// just write our material there and early out
-				if (pxr::UsdShadeMaterial ShadeMaterial = BindingAPI.ComputeBoundMaterial())
-				{
-					// Only consider this material reusable if its within the same layer as the edit target, otherwise we'll
-					// prefer to author something else that can be fully defined on the edit target. This is handy when we're
-					// exporting a level, as we'll ensure we're making these prims on the MaterialOverrides layer
-					pxr::SdfLayerRefPtr EditTarget = MeshPrim.GetStage()->GetEditTarget().GetLayer();
-					pxr::SdfLayerRefPtr MaterialPrimLayer = static_cast<pxr::SdfLayerRefPtr>(UsdUtils::FindLayerForPrim(ShadeMaterial.GetPrim()));
-					if (EditTarget == MaterialPrimLayer)
-					{
-						// We need to try reusing these materials or else we'd write a new material prim every time we change
-						// the override in UE, but we also run the risk of modifying a material that is used by multiple prims
-						// (and here we just want to set the override for this Mesh prim). The compromise is to only reuse the
-						// material if it is a child of MeshPrim already, and always to author our material prims as children
-						std::string MaterialPath = ShadeMaterial.GetPrim().GetPath().GetString();
-						std::string MeshPrimPath = MeshPrim.GetPath().GetString();
-						if (MaterialPath.rfind(MeshPrimPath, 0) == 0)
-						{
-							if (pxr::UsdPrim MaterialPrim = ShadeMaterial.GetPrim())
-							{
-								UsdUtils::SetUnrealSurfaceOutput(MaterialPrim, UEMaterialAssetPath);
-								return;
-							}
-						}
-					}
-				}
+		// References: private constructor for UsdGeomXformOp that can receive a UsdPrim and UsdGeomXformable::AddXformOp
 
-				// Find a unique name for our child material prim
-				// Note how we'll always author these materials as children of the meshes themselves instead of emitting a common
-				// Material prim to use for multiple overrides: This because in the future we'll want to have a separate material
-				// bake for each mesh (to make sure we get vertex color effects, etc.), and so we'd have multiple baked .usda material
-				// asset layers for each UE material, and we'd want each mesh/section/LOD to refer to its own anyway
-				FString ChildMaterialName = TEXT("UnrealMaterial");
-				if (pxr::UsdPrim ExistingPrim = MeshPrim.GetChild(UnrealToUsd::ConvertToken(*ChildMaterialName).Get()))
-				{
-					// Get a unique name for a new prim. Don't even try checking if this prim is usable as the material binding,
-					// because if it was the material binding for this mesh we would have already used it above, when fetching the ExistingShader.
-					// If we're here, we don't know what this prim is about
-					TSet<FString> UsedNames;
-					for (pxr::UsdPrim Child : MeshPrim.GetFilteredChildren(pxr::UsdTraverseInstanceProxies(pxr::UsdPrimAllPrimsPredicate)))
-					{
-						UsedNames.Add(UsdToUnreal::ConvertToken(Child.GetName()));
-					}
+		// Clear the existing xform op order for this prim on this layer
+		Xformable.ClearXformOpOrder();
 
-					ChildMaterialName = UsdUtils::GetUniqueName(ChildMaterialName, UsedNames);
-				}
+		// Find details about the transform attribute related to the default transform type xform op
+		pxr::TfToken TransformAttrName = pxr::UsdGeomXformOp::GetOpName(pxr::UsdGeomXformOp::TypeTransform);
+		const pxr::SdfValueTypeName& TransformAttrTypeName = pxr::UsdGeomXformOp::GetValueTypeName(
+			pxr::UsdGeomXformOp::TypeTransform,
+			pxr::UsdGeomXformOp::PrecisionDouble
+		);
+		if (TransformAttrName.IsEmpty() || !TransformAttrTypeName)
+		{
+			return {};
+		}
 
-				pxr::UsdStageRefPtr Stage = MeshPrim.GetStage();
-				pxr::SdfPath MeshPath = MeshPrim.GetPath();
-				pxr::SdfPath MaterialPath = MeshPath.AppendChild(UnrealToUsd::ConvertToken(*ChildMaterialName).Get());
+		// Create the transform attribute that would match the default transform type xform op
+		const bool bCustom = false;
+		pxr::UsdPrim UsdPrim = Xformable.GetPrim();
+		pxr::UsdAttribute TransformAttr = UsdPrim.CreateAttribute(TransformAttrName, TransformAttrTypeName, bCustom);
+		if (!TransformAttr)
+		{
+			return {};
+		}
 
-				pxr::UsdShadeMaterial ChildMaterial = pxr::UsdShadeMaterial::Define(Stage, MaterialPath);
-				if (!ChildMaterial)
-				{
-					UE_LOG(
-						LogUsd,
-						Warning,
-						TEXT("Failed to author material prim '%s' when trying to write '%s's material override '%s' to USD"),
-						*UsdToUnreal::ConvertPath(MaterialPath),
-						*UsdToUnreal::ConvertPath(MeshPrim.GetPath()),
-						*UEMaterialAssetPath
-					);
-					return;
-				}
+		// Now that the attribute is created, use it to create the corresponding pxr::UsdGeomXformOp
+		const bool bIsInverseOp = false;
+		pxr::UsdGeomXformOp NewOp{TransformAttr, bIsInverseOp};
+		if (!NewOp)
+		{
+			return {};
+		}
 
-				if (pxr::UsdPrim MaterialPrim = ChildMaterial.GetPrim())
-				{
-					UsdUtils::SetUnrealSurfaceOutput(MaterialPrim, UEMaterialAssetPath);
+		// Store the Op name on an array that will be our new op order value
+		pxr::VtTokenArray NewOps;
+		NewOps.push_back(NewOp.GetOpName());
+		Xformable.CreateXformOpOrderAttr().Set(NewOps);
 
-					BindingAPI.Bind(ChildMaterial);
-				}
-			}
+		return NewOp;
+	}
 
-			// On the current edit target, will set the Xformable's op order to a single "xformOp:transform",
-			// create the corresponding attribute, and return the op
-			pxr::UsdGeomXformOp ForceMatrixXform(pxr::UsdGeomXformable& Xformable)
-			{
-				FScopedUsdAllocs Allocs;
+	// Turns OutTransform into the UE-space relative (local to parent) transform for Xformable, paying attention to if it
+	// or any of its ancestors has the '!resetXformStack!' xformOp.
+	void GetPrimConvertedRelativeTransform(
+		pxr::UsdGeomXformable Xformable,
+		double UsdTimeCode,
+		FTransform& OutTransform,
+		bool bIgnoreLocalTransform = false
+	)
+	{
+		if (!Xformable)
+		{
+			return;
+		}
 
-				// Note: We don't use Xformable.MakeMatrixXform() here because while it can clear the
-				// xform op order on the current edit target just fine, it will later try to AddTransformOp(),
-				// which calls AddXformOp. Internally, it will read the *composed* prim and if it finds that it already
-				// has an op of that type it will early out and not author anything. This means that if our stage
-				// has a strong opinion for an e.g. "xformOp:transform" already on the layer stack, it's not possible
-				// to author that same op on a weaker layer. We want to do this here, to ensure this prim's transform
-				// works as expected even if this weaker layer is used standalone, so we must do the analogous ourselves
+		FScopedUsdAllocs Allocs;
 
-				// References: private constructor for UsdGeomXformOp that can receive a UsdPrim and UsdGeomXformable::AddXformOp
+		pxr::UsdPrim UsdPrim = Xformable.GetPrim();
+		pxr::UsdStageRefPtr UsdStage = UsdPrim.GetStage();
 
-				// Clear the existing xform op order for this prim on this layer
-				Xformable.ClearXformOpOrder();
+		bool bResetTransformStack = false;
+		if (bIgnoreLocalTransform)
+		{
+			FTransform Dummy;
+			UsdToUnreal::ConvertXformable(UsdStage, Xformable, Dummy, UsdTimeCode, &bResetTransformStack);
 
-				// Find details about the transform attribute related to the default transform type xform op
-				pxr::TfToken TransformAttrName = pxr::UsdGeomXformOp::GetOpName(pxr::UsdGeomXformOp::TypeTransform);
-				const pxr::SdfValueTypeName& TransformAttrTypeName = pxr::UsdGeomXformOp::GetValueTypeName(
-					pxr::UsdGeomXformOp::TypeTransform,
-					pxr::UsdGeomXformOp::PrecisionDouble
+			OutTransform = FTransform::Identity;
+		}
+		else
+		{
+			UsdToUnreal::ConvertXformable(UsdStage, Xformable, OutTransform, UsdTimeCode, &bResetTransformStack);
+		}
+
+		// If we have the resetXformStack op on this prim's xformOpOrder we have to essentially use its transform
+		// as the world transform (i.e. we have to discard the parent transforms). We won't do this here, and will instead
+		// keep relative transforms everywhere for consistency, which means we must manually invert the ParentToWorld transform
+		// and compute our relative transform ourselves.
+		//
+		// Ideally we'd query the components for this for performance reasons, but not only we don't have access to them here,
+		// but neither the stage actor's PrimsToAnimate nor the sequencer guarantee a particular evaluation order anyway,
+		// which means that if our parent is also animated, we could end up computing our relative transforms using the outdated
+		// parent's transform instead. This means we must compute our relative transform using the actual prim hierarchy.
+		//
+		// Additionally, our parent prims may be animated, so we must query all of our ancestors for a new world matrix every frame.
+		//
+		// We could use UsdGeomXformCache for this, but given that we won't actually cache anything (since we'll have to resample
+		// all ancestors every frame anyway) and that we would have to manually handle the camera/light compensation at least for
+		// our immediate parent, it's simpler to just recursively call our own UsdToUnreal::ConvertXformable and concatenate the
+		// results. Its not as fast, but we'll only do this on the initial read for prims with `resetXformStack`, so it should
+		// be very rare. We don't ever write out the resetXformStack either, so after that initial read this op should just disappear.
+		//
+		// Note that, alternatively, we could also handle this whole situation by having the scene components specify their transforms
+		// as absolute, and the Sequencer would work with that as well. However that would spread out the handling of
+		// resetXformStack through all USD workflows, and mean we'd have to *write out* resetXformStack when writing/exporting
+		// absolute transform components, and also convert between them when the user toggles between relative/absolute manually,
+		// which is probably worse than just baking it as relative transforms on first read and forgetting about it.
+		if (bResetTransformStack)
+		{
+			FTransform ParentToWorld = FTransform::Identity;
+
+			pxr::UsdPrim AncestorPrim = UsdPrim.GetParent();
+			while (AncestorPrim && !AncestorPrim.IsPseudoRoot())
+			{
+				FTransform AncestorTransform = FTransform::Identity;
+				bool bAncestorResetTransformStack = false;
+				UsdToUnreal::ConvertXformable(
+					UsdStage,
+					pxr::UsdGeomXformable{AncestorPrim},
+					AncestorTransform,
+					UsdTimeCode,
+					&bAncestorResetTransformStack
 				);
-				if (TransformAttrName.IsEmpty() || !TransformAttrTypeName)
-				{
-					return {};
-				}
 
-				// Create the transform attribute that would match the default transform type xform op
-				const bool bCustom = false;
-				pxr::UsdPrim UsdPrim = Xformable.GetPrim();
-				pxr::UsdAttribute TransformAttr = UsdPrim.CreateAttribute(TransformAttrName, TransformAttrTypeName, bCustom);
-				if (!TransformAttr)
-				{
-					return {};
-				}
+				ParentToWorld = ParentToWorld * AncestorTransform;
 
-				// Now that the attribute is created, use it to create the corresponding pxr::UsdGeomXformOp
-				const bool bIsInverseOp = false;
-				pxr::UsdGeomXformOp NewOp{TransformAttr, bIsInverseOp};
-				if (!NewOp)
+				// If we find a parent that also has the resetXformStack, then we're in luck: That transform value will be its world
+				// transform already, so we can stop concatenating stuff. Yes, on the component-side of things we'd have done the same
+				// thing of making a fake relative transform for it, but the end result would have been the same final world transform
+				if (bAncestorResetTransformStack)
 				{
-					return {};
+					break;
 				}
 
-				// Store the Op name on an array that will be our new op order value
-				pxr::VtTokenArray NewOps;
-				NewOps.push_back(NewOp.GetOpName());
-				Xformable.CreateXformOpOrderAttr().Set(NewOps);
-
-				return NewOp;
+				AncestorPrim = AncestorPrim.GetParent();
 			}
 
-			// Turns OutTransform into the UE-space relative (local to parent) transform for Xformable, paying attention to if it
-			// or any of its ancestors has the '!resetXformStack!' xformOp.
-			void GetPrimConvertedRelativeTransform(
-				pxr::UsdGeomXformable Xformable,
-				double UsdTimeCode,
-				FTransform& OutTransform,
-				bool bIgnoreLocalTransform = false
-			)
+			const FVector& Scale = ParentToWorld.GetScale3D();
+			if (!FMath::IsNearlyEqual(Scale.X, Scale.Y) || !FMath::IsNearlyEqual(Scale.X, Scale.Z))
 			{
-				if (!Xformable)
-				{
-					return;
-				}
-
-				FScopedUsdAllocs Allocs;
-
-				pxr::UsdPrim UsdPrim = Xformable.GetPrim();
-				pxr::UsdStageRefPtr UsdStage = UsdPrim.GetStage();
-
-				bool bResetTransformStack = false;
-				if (bIgnoreLocalTransform)
-				{
-					FTransform Dummy;
-					UsdToUnreal::ConvertXformable(UsdStage, Xformable, Dummy, UsdTimeCode, &bResetTransformStack);
-
-					OutTransform = FTransform::Identity;
-				}
-				else
-				{
-					UsdToUnreal::ConvertXformable(UsdStage, Xformable, OutTransform, UsdTimeCode, &bResetTransformStack);
-				}
-
-				// If we have the resetXformStack op on this prim's xformOpOrder we have to essentially use its transform
-				// as the world transform (i.e. we have to discard the parent transforms). We won't do this here, and will instead
-				// keep relative transforms everywhere for consistency, which means we must manually invert the ParentToWorld transform
-				// and compute our relative transform ourselves.
-				//
-				// Ideally we'd query the components for this for performance reasons, but not only we don't have access to them here,
-				// but neither the stage actor's PrimsToAnimate nor the sequencer guarantee a particular evaluation order anyway,
-				// which means that if our parent is also animated, we could end up computing our relative transforms using the outdated
-				// parent's transform instead. This means we must compute our relative transform using the actual prim hierarchy.
-				//
-				// Additionally, our parent prims may be animated, so we must query all of our ancestors for a new world matrix every frame.
-				//
-				// We could use UsdGeomXformCache for this, but given that we won't actually cache anything (since we'll have to resample
-				// all ancestors every frame anyway) and that we would have to manually handle the camera/light compensation at least for
-				// our immediate parent, it's simpler to just recursively call our own UsdToUnreal::ConvertXformable and concatenate the
-				// results. Its not as fast, but we'll only do this on the initial read for prims with `resetXformStack`, so it should
-				// be very rare. We don't ever write out the resetXformStack either, so after that initial read this op should just disappear.
-				//
-				// Note that, alternatively, we could also handle this whole situation by having the scene components specify their transforms
-				// as absolute, and the Sequencer would work with that as well. However that would spread out the handling of
-				// resetXformStack through all USD workflows, and mean we'd have to *write out* resetXformStack when writing/exporting
-				// absolute transform components, and also convert between them when the user toggles between relative/absolute manually,
-				// which is probably worse than just baking it as relative transforms on first read and forgetting about it.
-				if (bResetTransformStack)
-				{
-					FTransform ParentToWorld = FTransform::Identity;
-
-					pxr::UsdPrim AncestorPrim = UsdPrim.GetParent();
-					while (AncestorPrim && !AncestorPrim.IsPseudoRoot())
-					{
-						FTransform AncestorTransform = FTransform::Identity;
-						bool bAncestorResetTransformStack = false;
-						UsdToUnreal::ConvertXformable(
-							UsdStage,
-							pxr::UsdGeomXformable{AncestorPrim},
-							AncestorTransform,
-							UsdTimeCode,
-							&bAncestorResetTransformStack
-						);
-
-						ParentToWorld = ParentToWorld * AncestorTransform;
-
-						// If we find a parent that also has the resetXformStack, then we're in luck: That transform value will be its world
-						// transform already, so we can stop concatenating stuff. Yes, on the component-side of things we'd have done the same
-						// thing of making a fake relative transform for it, but the end result would have been the same final world transform
-						if (bAncestorResetTransformStack)
-						{
-							break;
-						}
-
-						AncestorPrim = AncestorPrim.GetParent();
-					}
-
-					const FVector& Scale = ParentToWorld.GetScale3D();
-					if (!FMath::IsNearlyEqual(Scale.X, Scale.Y) || !FMath::IsNearlyEqual(Scale.X, Scale.Z))
-					{
-						UE_LOG(
-							LogUsd,
-							Warning,
-							TEXT("Inverting transform with non-uniform scaling '%s' when computing relative transform for prim '%s'! Result will "
-								 "likely be incorrect, since FTransforms can't invert non-uniform scalings. You can work around this by baking your "
-								 "non-uniform scaling transform into the vertices, or by not using the !resetXformStack! Xform op."),
-							*Scale.ToString(),
-							*UsdToUnreal::ConvertPath(UsdPrim.GetPrimPath())
-						);
-					}
-
-					// Multiplying with matrices here helps mitigate the issues encountered with non-uniform scaling, however it will stil
-					// never be perfect, as it is not possible to generate an FTransform that can properly invert a complex transform with non-uniform
-					// scaling when just multiplying them (which is what downstream code within USceneComponent will do).
-					OutTransform = FTransform{OutTransform.ToMatrixWithScale() * ParentToWorld.ToInverseMatrixWithScale()};
-				}
+				UE_LOG(
+					LogUsd,
+					Warning,
+					TEXT("Inverting transform with non-uniform scaling '%s' when computing relative transform for prim '%s'! Result will "
+						 "likely be incorrect, since FTransforms can't invert non-uniform scalings. You can work around this by baking your "
+						 "non-uniform scaling transform into the vertices, or by not using the !resetXformStack! Xform op."),
+					*Scale.ToString(),
+					*UsdToUnreal::ConvertPath(UsdPrim.GetPrimPath())
+				);
 			}
-		}	 // namespace Private
-	}		 // namespace USDPrimConversionImpl
-}	 // namespace UE
+
+			// Multiplying with matrices here helps mitigate the issues encountered with non-uniform scaling, however it will stil
+			// never be perfect, as it is not possible to generate an FTransform that can properly invert a complex transform with non-uniform
+			// scaling when just multiplying them (which is what downstream code within USceneComponent will do).
+			OutTransform = FTransform{OutTransform.ToMatrixWithScale() * ParentToWorld.ToInverseMatrixWithScale()};
+		}
+	}
+}	 // namespace UE::USDPrimConversion::Private
 
 bool UsdToUnreal::ConvertXformable(
 	const pxr::UsdStageRefPtr& Stage,
@@ -516,7 +421,7 @@ bool UsdToUnreal::ConvertXformable(
 
 	// Transform
 	FTransform Transform;
-	UE::USDPrimConversionImpl::Private::GetPrimConvertedRelativeTransform(Xformable, EvalTime, Transform, !bUsePrimTransform);
+	UE::USDPrimConversion::Private::GetPrimConvertedRelativeTransform(Xformable, EvalTime, Transform, !bUsePrimTransform);
 	SceneComponent.SetRelativeTransform(Transform);
 
 	SceneComponent.Modify();
@@ -1202,7 +1107,7 @@ UsdToUnreal::FPropertyTrackReader UsdToUnreal::CreatePropertyTrackReader(
 		if (PropertyPath == UnrealIdentifiers::TransformPropertyName)
 		{
 			FTransform Default = FTransform::Identity;
-			UE::USDPrimConversionImpl::Private::GetPrimConvertedRelativeTransform(
+			UE::USDPrimConversion::Private::GetPrimConvertedRelativeTransform(
 				Xformable,
 				UsdUtils::GetDefaultTimeCode(),
 				Default,
@@ -1212,7 +1117,7 @@ UsdToUnreal::FPropertyTrackReader UsdToUnreal::CreatePropertyTrackReader(
 			Reader.TransformReader = [UsdStage, Xformable, Default, bIgnorePrimLocalTransform](double UsdTimeCode)
 			{
 				FTransform Result = Default;
-				UE::USDPrimConversionImpl::Private::GetPrimConvertedRelativeTransform(Xformable, UsdTimeCode, Result, bIgnorePrimLocalTransform);
+				UE::USDPrimConversion::Private::GetPrimConvertedRelativeTransform(Xformable, UsdTimeCode, Result, bIgnorePrimLocalTransform);
 				return Result;
 			};
 			return Reader;
@@ -1856,7 +1761,7 @@ bool UsdToUnreal::ConvertBounds(const pxr::UsdPrim& Prim, UUsdDrawModeComponent*
 	return true;
 }
 
-namespace UE::USDPrimConversionImpl::Private
+namespace UE::USDPrimConversion::Private
 {
 	const static std::string UsdNamespaceDelimiter = UnrealToUsd::ConvertString(*UnrealIdentifiers::UsdNamespaceDelimiter).Get();
 
@@ -2031,7 +1936,7 @@ namespace UE::USDPrimConversionImpl::Private
 			}
 		}
 	}
-};	  // namespace UE::USDPrimConversionImpl::Private
+};	  // namespace UE::USDPrimConversion::Private
 
 bool UsdToUnreal::ConvertMetadata(
 	const pxr::UsdPrim& Prim,
@@ -2046,14 +1951,14 @@ bool UsdToUnreal::ConvertMetadata(
 		return false;
 	}
 
-	UE::USDPrimConversionImpl::Private::CollectMetadataForPrim(Prim, CombinedMetadata, BlockedPrefixFilters, bInvertFilters);
+	UE::USDPrimConversion::Private::CollectMetadataForPrim(Prim, CombinedMetadata, BlockedPrefixFilters, bInvertFilters);
 
 	if (bCollectFromEntireSubtrees)
 	{
 		pxr::UsdPrimRange PrimRange{Prim, pxr::UsdTraverseInstanceProxies()};
 		for (pxr::UsdPrimRange::iterator It = ++PrimRange.begin(); It != PrimRange.end(); ++It)
 		{
-			UE::USDPrimConversionImpl::Private::CollectMetadataForPrim(*It, CombinedMetadata, BlockedPrefixFilters, bInvertFilters);
+			UE::USDPrimConversion::Private::CollectMetadataForPrim(*It, CombinedMetadata, BlockedPrefixFilters, bInvertFilters);
 		}
 	}
 
@@ -3067,7 +2972,7 @@ bool UnrealToUsd::ConvertMaterialOverrides(
 						{
 							pxr::SdfPath OverridePrimPath = UnrealToUsd::ConvertPath(*PrimPath).Get();
 							pxr::UsdPrim MeshPrim = Stage->OverridePrim(OverridePrimPath);
-							UE::USDPrimConversionImpl::Private::AuthorMaterialOverride(MeshPrim, Override->GetPathName());
+							UsdUtils::AuthorUnrealMaterialBinding(MeshPrim, Override->GetPathName());
 						}
 					}
 				}
@@ -3079,7 +2984,7 @@ bool UnrealToUsd::ConvertMaterialOverrides(
 			{
 				pxr::SdfPath OverridePrimPath = UsdPrim.GetPath();
 				pxr::UsdPrim MeshPrim = Stage->OverridePrim(OverridePrimPath);
-				UE::USDPrimConversionImpl::Private::AuthorMaterialOverride(MeshPrim, Override->GetPathName());
+				UsdUtils::AuthorUnrealMaterialBinding(MeshPrim, Override->GetPathName());
 			}
 		}
 	}
@@ -3141,7 +3046,7 @@ bool UnrealToUsd::ConvertMaterialOverrides(
 								{
 									pxr::SdfPath OverridePrimPath = UnrealToUsd::ConvertPath(*PrimPath).Get();
 									pxr::UsdPrim MeshPrim = Stage->OverridePrim(OverridePrimPath);
-									UE::USDPrimConversionImpl::Private::AuthorMaterialOverride(MeshPrim, Override->GetPathName());
+									UsdUtils::AuthorUnrealMaterialBinding(MeshPrim, Override->GetPathName());
 								}
 							}
 						}
@@ -3178,7 +3083,7 @@ bool UnrealToUsd::ConvertMaterialOverrides(
 						}
 
 						pxr::UsdPrim MeshPrim = Stage->OverridePrim(OverridePrimPath);
-						UE::USDPrimConversionImpl::Private::AuthorMaterialOverride(MeshPrim, Override->GetPathName());
+						UsdUtils::AuthorUnrealMaterialBinding(MeshPrim, Override->GetPathName());
 					}
 				}
 			}
@@ -3313,7 +3218,7 @@ bool UnrealToUsd::ConvertMaterialOverrides(
 									{
 										pxr::SdfPath OverridePrimPath = UnrealToUsd::ConvertPath(*SourcePrimPath).Get();
 										pxr::UsdPrim MeshPrim = Stage->OverridePrim(OverridePrimPath);
-										UE::USDPrimConversionImpl::Private::AuthorMaterialOverride(MeshPrim, Override->GetPathName());
+										UsdUtils::AuthorUnrealMaterialBinding(MeshPrim, Override->GetPathName());
 										break;
 									}
 								}
@@ -3372,7 +3277,7 @@ bool UnrealToUsd::ConvertMaterialOverrides(
 						}
 
 						pxr::UsdPrim MeshPrim = Stage->OverridePrim(OverridePrimPath);
-						UE::USDPrimConversionImpl::Private::AuthorMaterialOverride(MeshPrim, Override->GetPathName());
+						UsdUtils::AuthorUnrealMaterialBinding(MeshPrim, Override->GetPathName());
 					}
 				}
 			}
@@ -3414,7 +3319,7 @@ bool UnrealToUsd::ConvertXformable(const FTransform& RelativeTransform, pxr::Usd
 
 	const pxr::UsdTimeCode UsdTimeCode(TimeCode);
 
-	if (pxr::UsdGeomXformOp MatrixXform = UE::USDPrimConversionImpl::Private::ForceMatrixXform(XForm))
+	if (pxr::UsdGeomXformOp MatrixXform = UE::USDPrimConversion::Private::ForceMatrixXform(XForm))
 	{
 		MatrixXform.Set(UsdTransform, UsdTimeCode);
 
@@ -3426,61 +3331,55 @@ bool UnrealToUsd::ConvertXformable(const FTransform& RelativeTransform, pxr::Usd
 }
 
 #if WITH_EDITOR
-namespace UE
+namespace UE::USDPrimConversion::Private
 {
-	namespace USDPrimConversionImpl
-	{
-		namespace Private
-		{
-			void ConvertFoliageInstances(
-				const FFoliageInfo& Info,
-				const TSet<int32>& UEInstances,
-				const FTransform& UEWorldToFoliageActor,
-				const FUsdStageInfo& StageInfo,
-				int PrototypeIndex,
-				pxr::VtArray<int>& ProtoIndices,
-				pxr::VtArray<pxr::GfVec3f>& Positions,
-				pxr::VtArray<pxr::GfQuath>& Orientations,
-				pxr::VtArray<pxr::GfVec3f>& Scales
-			)
-			{
-				FScopedUsdAllocs Allocs;
-
-				const int32 NumInstances = UEInstances.Num();
+	void ConvertFoliageInstances(
+		const FFoliageInfo& Info,
+		const TSet<int32>& UEInstances,
+		const FTransform& UEWorldToFoliageActor,
+		const FUsdStageInfo& StageInfo,
+		int PrototypeIndex,
+		pxr::VtArray<int>& ProtoIndices,
+		pxr::VtArray<pxr::GfVec3f>& Positions,
+		pxr::VtArray<pxr::GfQuath>& Orientations,
+		pxr::VtArray<pxr::GfVec3f>& Scales
+	)
+	{
+		FScopedUsdAllocs Allocs;
 
-				ProtoIndices.reserve(ProtoIndices.size() + NumInstances);
-				Positions.reserve(Positions.size() + NumInstances);
-				Orientations.reserve(Orientations.size() + NumInstances);
-				Scales.reserve(Scales.size() + NumInstances);
+		const int32 NumInstances = UEInstances.Num();
 
-				for (int32 InstanceIndex : UEInstances)
-				{
-					const FFoliageInstancePlacementInfo* Instance = &Info.Instances[InstanceIndex];
+		ProtoIndices.reserve(ProtoIndices.size() + NumInstances);
+		Positions.reserve(Positions.size() + NumInstances);
+		Orientations.reserve(Orientations.size() + NumInstances);
+		Scales.reserve(Scales.size() + NumInstances);
 
-					// Convert axes
-					FTransform UEWorldTransform{Instance->Rotation, (FVector)Instance->Location, (FVector)Instance->DrawScale3D};
-					FTransform USDTransform = UsdUtils::ConvertAxes(StageInfo.UpAxis == EUsdUpAxis::ZAxis, UEWorldTransform * UEWorldToFoliageActor);
+		for (int32 InstanceIndex : UEInstances)
+		{
+			const FFoliageInstancePlacementInfo* Instance = &Info.Instances[InstanceIndex];
 
-					FVector Translation = USDTransform.GetTranslation();
-					FQuat Rotation = USDTransform.GetRotation();
-					FVector Scale = USDTransform.GetScale3D();
+			// Convert axes
+			FTransform UEWorldTransform{Instance->Rotation, (FVector)Instance->Location, (FVector)Instance->DrawScale3D};
+			FTransform USDTransform = UsdUtils::ConvertAxes(StageInfo.UpAxis == EUsdUpAxis::ZAxis, UEWorldTransform * UEWorldToFoliageActor);
 
-					// Compensate metersPerUnit
-					const float UEMetersPerUnit = 0.01f;
-					if (!FMath::IsNearlyEqual(UEMetersPerUnit, StageInfo.MetersPerUnit))
-					{
-						Translation *= (UEMetersPerUnit / StageInfo.MetersPerUnit);
-					}
+			FVector Translation = USDTransform.GetTranslation();
+			FQuat Rotation = USDTransform.GetRotation();
+			FVector Scale = USDTransform.GetScale3D();
 
-					ProtoIndices.push_back(PrototypeIndex);
-					Positions.push_back(pxr::GfVec3f(Translation.X, Translation.Y, Translation.Z));
-					Orientations.push_back(pxr::GfQuath(Rotation.W, Rotation.X, Rotation.Y, Rotation.Z));
-					Scales.push_back(pxr::GfVec3f(Scale.X, Scale.Y, Scale.Z));
-				}
+			// Compensate metersPerUnit
+			const float UEMetersPerUnit = 0.01f;
+			if (!FMath::IsNearlyEqual(UEMetersPerUnit, StageInfo.MetersPerUnit))
+			{
+				Translation *= (UEMetersPerUnit / StageInfo.MetersPerUnit);
 			}
-		}	 // namespace Private
-	}		 // namespace USDPrimConversionImpl
-}	 // namespace UE
+
+			ProtoIndices.push_back(PrototypeIndex);
+			Positions.push_back(pxr::GfVec3f(Translation.X, Translation.Y, Translation.Z));
+			Orientations.push_back(pxr::GfQuath(Rotation.W, Rotation.X, Rotation.Y, Rotation.Z));
+			Scales.push_back(pxr::GfVec3f(Scale.X, Scale.Y, Scale.Z));
+		}
+	}
+}	 // namespace UE::USDPrimConversion::Private
 #endif	  // WITH_EDITOR
 
 bool UnrealToUsd::ConvertInstancedFoliageActor(const AInstancedFoliageActor& Actor, pxr::UsdPrim& UsdPrim, double TimeCode, ULevel* InstancesLevel)
@@ -3535,7 +3434,7 @@ bool UnrealToUsd::ConvertInstancedFoliageActor(const AInstancedFoliageActor& Act
 
 			if (const TSet<int32>* InstanceSet = Info.ComponentHash.Find(ComponentId))
 			{
-				UE::USDPrimConversionImpl::Private::ConvertFoliageInstances(
+				UE::USDPrimConversion::Private::ConvertFoliageInstances(
 					Info,
 					*InstanceSet,
 					UEWorldToFoliageActor,
@@ -3564,7 +3463,7 @@ bool UnrealToUsd::ConvertInstancedFoliageActor(const AInstancedFoliageActor& Act
 				}
 
 				const TSet<int32>& InstanceSet = Pair.Value;
-				UE::USDPrimConversionImpl::Private::ConvertFoliageInstances(
+				UE::USDPrimConversion::Private::ConvertFoliageInstances(
 					Info,
 					InstanceSet,
 					UEWorldToFoliageActor,
@@ -3642,7 +3541,7 @@ bool UnrealToUsd::CreateComponentPropertyBaker(
 			Xformable.CreateXformOpOrderAttr();
 
 			// Clear existing transform data and leave just one Transform op there
-			pxr::UsdGeomXformOp TransformOp = UE::USDPrimConversionImpl::Private::ForceMatrixXform(Xformable);
+			pxr::UsdGeomXformOp TransformOp = UE::USDPrimConversion::Private::ForceMatrixXform(Xformable);
 			if (!TransformOp)
 			{
 				return false;
@@ -4132,7 +4031,7 @@ UnrealToUsd::FPropertyTrackWriter UnrealToUsd::CreatePropertyTrackWriter(
 				{
 					Xformable.CreateXformOpOrderAttr();
 
-					if (pxr::UsdGeomXformOp TransformOp = UE::USDPrimConversionImpl::Private::ForceMatrixXform(Xformable))
+					if (pxr::UsdGeomXformOp TransformOp = UE::USDPrimConversion::Private::ForceMatrixXform(Xformable))
 					{
 						Attr = TransformOp.GetAttr();
 
@@ -4791,7 +4690,7 @@ bool UnrealToUsd::ConvertXformable(
 
 	if (bIsDataOutOfSync)
 	{
-		if (pxr::UsdGeomXformOp TransformOp = UE::USDPrimConversionImpl::Private::ForceMatrixXform(Xformable))
+		if (pxr::UsdGeomXformOp TransformOp = UE::USDPrimConversion::Private::ForceMatrixXform(Xformable))
 		{
 			TransformOp.GetAttr().Clear();	  // Clear existing transform data
 		}
@@ -5252,7 +5151,7 @@ bool UnrealToUsd::ConvertBoundsComponent(const UUsdDrawModeComponent& BoundsComp
 	return true;
 }
 
-namespace UE::USDPrimConversionImpl::Private
+namespace UE::USDPrimConversion::Private
 {
 	FString PrimPathToNamespace(FString PrimPath)
 	{
@@ -5278,7 +5177,7 @@ namespace UE::USDPrimConversionImpl::Private
 		const FString& NamespacePrefix = {}
 	)
 	{
-		using namespace UE::USDPrimConversionImpl::Private;
+		using namespace UE::USDPrimConversion::Private;
 
 		if (!Prim || PrimMetadata.Metadata.Num() == 0 || (bInvertFilters && BlockedPrefixFilters.Num() == 0))
 		{
@@ -5399,7 +5298,7 @@ namespace UE::USDPrimConversionImpl::Private
 		}
 		return bSuccess;
 	}
-}	 // namespace UE::USDPrimConversionImpl::Private
+}	 // namespace UE::USDPrimConversion::Private
 
 bool UnrealToUsd::ConvertMetadata(
 	const FUsdCombinedPrimMetadata& CombinedPrimMetadata,
@@ -5408,7 +5307,7 @@ bool UnrealToUsd::ConvertMetadata(
 	bool bInvertFilters
 )
 {
-	using namespace UE::USDPrimConversionImpl::Private;
+	using namespace UE::USDPrimConversion::Private;
 
 	if (!Prim || CombinedPrimMetadata.PrimPathToMetadata.Num() == 0 || (bInvertFilters && BlockedPrefixFilters.Num() == 0))
 	{
@@ -5498,7 +5397,7 @@ bool UnrealToUsd::ConvertMetadata(
 	bool bInvertFilters
 )
 {
-	return UE::USDPrimConversionImpl::Private::ConvertMetadataInternal(PrimMetadata, Prim, BlockedPrefixFilters, bInvertFilters);
+	return UE::USDPrimConversion::Private::ConvertMetadataInternal(PrimMetadata, Prim, BlockedPrefixFilters, bInvertFilters);
 }
 
 bool UnrealToUsd::ConvertMetadata(
diff --git a/Engine/Plugins/Importers/USDImporter/Source/USDUtilities/Private/USDShadeConversion.cpp b/Engine/Plugins/Importers/USDImporter/Source/USDUtilities/Private/USDShadeConversion.cpp
--- a/Engine/Plugins/Importers/USDImporter/Source/USDUtilities/Private/USDShadeConversion.cpp
+++ b/Engine/Plugins/Importers/USDImporter/Source/USDUtilities/Private/USDShadeConversion.cpp
@@ -57,6 +57,7 @@
 #include "pxr/usd/ar/resolverScopedCache.h"
 #include "pxr/usd/sdf/layerUtils.h"
 #include "pxr/usd/usdShade/material.h"
+#include "pxr/usd/usdShade/materialBindingAPI.h"
 #include "pxr/usd/usdShade/nodeDefAPI.h"
 #include "USDIncludesEnd.h"
 
@@ -1226,8 +1227,7 @@ namespace UE
 						TextureValue.UVScale.GetVector()[0],
 						TextureValue.UVScale.GetVector()[1],
 						TextureValue.UVTranslation[0],
-						TextureValue.UVTranslation[1]
-					};
+						TextureValue.UVTranslation[1]};
 					UsdUtils::SetVectorParameterValue(Material, *FString::Printf(TEXT("%sScaleTranslation"), ParameterName), ScaleAndTranslation);
 
 					UsdUtils::SetScalarParameterValue(Material, *FString::Printf(TEXT("%sRotation"), ParameterName), TextureValue.UVRotation);
@@ -3294,6 +3294,83 @@ void UsdUtils::SetBoolParameterValue(UMaterialInstance& Material, const TCHAR* P
 	}
 }
 
+void UsdUtils::AuthorUnrealMaterialBinding(pxr::UsdPrim& MeshOrGeomSubsetPrim, const FString& UnrealMaterialPathName)
+{
+	if (!MeshOrGeomSubsetPrim || UnrealMaterialPathName.IsEmpty())
+	{
+		return;
+	}
+
+	FScopedUsdAllocs UsdAllocs;
+
+	pxr::UsdShadeMaterialBindingAPI BindingAPI = pxr::UsdShadeMaterialBindingAPI::Apply(MeshOrGeomSubsetPrim);
+
+	// If this mesh prim already has a binding to a *child* material with the 'unreal' render context,
+	// just write our material there and early out
+	if (pxr::UsdShadeMaterial ShadeMaterial = BindingAPI.ComputeBoundMaterial())
+	{
+		// We need to try reusing these materials or else we'd write a new material prim every time we change
+		// the override in UE, but we also run the risk of modifying a material that is used by multiple prims
+		// (and here we just want to set the override for this Mesh prim). The compromise is to only reuse the
+		// material if it is a child of MeshPrim already, and always to author our material prims as children
+		std::string MaterialPath = ShadeMaterial.GetPrim().GetPath().GetString();
+		std::string MeshPrimPath = MeshOrGeomSubsetPrim.GetPath().GetString();
+		if (MaterialPath.rfind(MeshPrimPath, 0) == 0)
+		{
+			if (pxr::UsdPrim MaterialPrim = ShadeMaterial.GetPrim())
+			{
+				UsdUtils::SetUnrealSurfaceOutput(MaterialPrim, UnrealMaterialPathName);
+				return;
+			}
+		}
+	}
+
+	// Find a unique name for our child material prim
+	// Note how we'll always author these materials as children of the meshes themselves instead of emitting a common
+	// Material prim to use for multiple overrides: This because in the future we'll want to have a separate material
+	// bake for each mesh (to make sure we get vertex color effects, etc.), and so we'd have multiple baked .usda material
+	// asset layers for each UE material, and we'd want each mesh/section/LOD to refer to its own anyway
+	FString ChildMaterialName = TEXT("UnrealMaterial");
+	if (pxr::UsdPrim ExistingPrim = MeshOrGeomSubsetPrim.GetChild(UnrealToUsd::ConvertToken(*ChildMaterialName).Get()))
+	{
+		// Get a unique name for a new prim. Don't even try checking if this prim is usable as the material binding,
+		// because if it was the material binding for this mesh we would have already used it above, when fetching the ExistingShader.
+		// If we're here, we don't know what this prim is about
+		TSet<FString> UsedNames;
+		for (pxr::UsdPrim Child : MeshOrGeomSubsetPrim.GetFilteredChildren(pxr::UsdTraverseInstanceProxies(pxr::UsdPrimAllPrimsPredicate)))
+		{
+			UsedNames.Add(UsdToUnreal::ConvertToken(Child.GetName()));
+		}
+
+		ChildMaterialName = UsdUtils::GetUniqueName(ChildMaterialName, UsedNames);
+	}
+
+	pxr::UsdStageRefPtr Stage = MeshOrGeomSubsetPrim.GetStage();
+	pxr::SdfPath MeshPath = MeshOrGeomSubsetPrim.GetPath();
+	pxr::SdfPath MaterialPath = MeshPath.AppendChild(UnrealToUsd::ConvertToken(*ChildMaterialName).Get());
+
+	pxr::UsdShadeMaterial ChildMaterial = pxr::UsdShadeMaterial::Define(Stage, MaterialPath);
+	if (!ChildMaterial)
+	{
+		UE_LOG(
+			LogUsd,
+			Warning,
+			TEXT("Failed to author material prim '%s' when trying to write '%s's material assignment '%s' to USD"),
+			*UsdToUnreal::ConvertPath(MaterialPath),
+			*UsdToUnreal::ConvertPath(MeshOrGeomSubsetPrim.GetPath()),
+			*UnrealMaterialPathName
+		);
+		return;
+	}
+
+	if (pxr::UsdPrim MaterialPrim = ChildMaterial.GetPrim())
+	{
+		UsdUtils::SetUnrealSurfaceOutput(MaterialPrim, UnrealMaterialPathName);
+
+		BindingAPI.Bind(ChildMaterial);
+	}
+}
+
 TOptional<FString> UsdUtils::GetUnrealSurfaceOutput(const pxr::UsdPrim& MaterialPrim)
 {
 	if (!MaterialPrim)
diff --git a/Engine/Plugins/Importers/USDImporter/Source/USDUtilities/Private/USDSkeletalDataConversion.cpp b/Engine/Plugins/Importers/USDImporter/Source/USDUtilities/Private/USDSkeletalDataConversion.cpp
--- a/Engine/Plugins/Importers/USDImporter/Source/USDUtilities/Private/USDSkeletalDataConversion.cpp
+++ b/Engine/Plugins/Importers/USDImporter/Source/USDUtilities/Private/USDSkeletalDataConversion.cpp
@@ -14,6 +14,7 @@
 #include "USDLog.h"
 #include "USDMemory.h"
 #include "USDPrimConversion.h"
+#include "USDShadeConversion.h"
 #include "USDTypesConversion.h"
 
 #include "UsdWrappers/SdfLayer.h"
@@ -784,10 +785,10 @@ namespace UnrealToUsdImpl
 		const FSkeletalMeshLODModel& LODModel,
 		pxr::UsdGeomMesh& UsdLODPrimGeomMesh,
 		bool bHasVertexColors,
-		pxr::VtArray<std::string>& MaterialAssignments,
+		const TArray<FString>& MaterialAssignments,
 		const TArray<int32>& LODMaterialMap,
 		const pxr::UsdTimeCode TimeCode,
-		pxr::UsdPrim MaterialPrim
+		pxr::UsdPrim PrimToReceiveMaterialAssignments
 	)
 	{
 		FScopedUsdAllocs UsdAllocs;
@@ -1061,7 +1062,7 @@ namespace UnrealToUsdImpl
 		{
 			bool bHasUEMaterialAssignements = false;
 
-			pxr::VtArray<std::string> UnrealMaterialsForLOD;
+			TArray<FString> UnrealMaterialsForLOD;
 			for (const FSkelMeshSection& Section : LODModel.Sections)
 			{
 				int32 SkeletalMaterialIndex = INDEX_NONE;
@@ -1075,31 +1076,25 @@ namespace UnrealToUsdImpl
 					SkeletalMaterialIndex = Section.MaterialIndex;
 				}
 
-				if (SkeletalMaterialIndex >= 0 && SkeletalMaterialIndex < MaterialAssignments.size())
+				if (MaterialAssignments.IsValidIndex(SkeletalMaterialIndex))
 				{
-					UnrealMaterialsForLOD.push_back(MaterialAssignments[SkeletalMaterialIndex]);
+					UnrealMaterialsForLOD.Add(MaterialAssignments[SkeletalMaterialIndex]);
 					bHasUEMaterialAssignements = true;
 				}
 				else
 				{
 					// Keep unrealMaterials with the same number of elements as our MaterialIndices expect
-					UnrealMaterialsForLOD.push_back("");
+					UnrealMaterialsForLOD.Add("");
 				}
 			}
 
 			// This LOD has a single material assignment, just add an unrealMaterials attribute to the mesh prim
-			if (bHasUEMaterialAssignements && UnrealMaterialsForLOD.size() == 1)
+			if (bHasUEMaterialAssignements && UnrealMaterialsForLOD.Num() == 1)
 			{
-				if (pxr::UsdAttribute UEMaterialsAttribute = MaterialPrim.CreateAttribute(
-						UnrealIdentifiers::MaterialAssignment,
-						pxr::SdfValueTypeNames->String
-					))
-				{
-					UEMaterialsAttribute.Set(UnrealMaterialsForLOD[0]);
-				}
+				UsdUtils::AuthorUnrealMaterialBinding(PrimToReceiveMaterialAssignments, UnrealMaterialsForLOD[0]);
 			}
 			// Multiple material assignments to the same LOD (and so the same mesh prim). Need to create a GeomSubset for each UE mesh section
-			else if (UnrealMaterialsForLOD.size() > 1)
+			else if (UnrealMaterialsForLOD.Num() > 1)
 			{
 				// Need to fetch all triangles of a section, and add their indices to the GeomSubset
 				for (int32 SectionIndex = 0; SectionIndex < LODModel.Sections.Num(); ++SectionIndex)
@@ -1110,19 +1105,18 @@ namespace UnrealToUsdImpl
 						continue;
 					}
 
-					// Note that we will continue on even if we have no material assignment, so as to satisfy the "partition" family condition
-					std::string SectionMaterial = UnrealMaterialsForLOD[SectionIndex];
-
+					// Note that we will continue authoring the GeomSubsets on even if we later find out we have no material assignment (just
+					// "") for this section, so as to satisfy the "partition" family condition (below)
 					pxr::UsdPrim GeomSubsetPrim = Stage->DefinePrim(
 						UsdLODPrimGeomMesh.GetPath().AppendPath(pxr::SdfPath("Section" + std::to_string(SectionIndex))),
 						UnrealToUsd::ConvertToken(TEXT("GeomSubset")).Get()
 					);
 
 					pxr::UsdPrim MaterialGeomSubsetPrim = GeomSubsetPrim;
-					if (MaterialPrim.GetStage() != MeshPrim.GetStage())
+					if (PrimToReceiveMaterialAssignments.GetStage() != MeshPrim.GetStage())
 					{
-						MaterialGeomSubsetPrim = MaterialPrim.GetStage()->OverridePrim(
-							MaterialPrim.GetPath().AppendPath(pxr::SdfPath("Section" + std::to_string(SectionIndex)))
+						MaterialGeomSubsetPrim = PrimToReceiveMaterialAssignments.GetStage()->OverridePrim(
+							PrimToReceiveMaterialAssignments.GetPath().AppendPath(pxr::SdfPath("Section" + std::to_string(SectionIndex)))
 						);
 					}
 
@@ -1169,14 +1163,8 @@ namespace UnrealToUsdImpl
 					// Family type
 					pxr::UsdGeomSubset::SetFamilyType(UsdLODPrimGeomMesh, pxr::UsdShadeTokens->materialBind, pxr::UsdGeomTokens->partition);
 
-					// unrealMaterials attribute
-					if (pxr::UsdAttribute UEMaterialsAttribute = MaterialGeomSubsetPrim.CreateAttribute(
-							UnrealIdentifiers::MaterialAssignment,
-							pxr::SdfValueTypeNames->String
-						))
-					{
-						UEMaterialsAttribute.Set(UnrealMaterialsForLOD[SectionIndex]);
-					}
+					// material:binding relationship
+					UsdUtils::AuthorUnrealMaterialBinding(MaterialGeomSubsetPrim, UnrealMaterialsForLOD[SectionIndex]);
 				}
 			}
 		}
@@ -3316,7 +3304,7 @@ bool UnrealToUsd::ConvertSkeletalMesh(
 
 	// Collect all material assignments, referenced by the sections' material indices
 	bool bHasMaterialAssignments = false;
-	pxr::VtArray<std::string> MaterialAssignments;
+	TArray<FString> MaterialAssignments;
 	for (const FSkeletalMaterial& SkeletalMaterial : SkeletalMesh->GetMaterials())
 	{
 		FString AssignedMaterialPathName;
@@ -3329,12 +3317,12 @@ bool UnrealToUsd::ConvertSkeletalMesh(
 			}
 		}
 
-		MaterialAssignments.push_back(UnrealToUsd::ConvertString(*AssignedMaterialPathName).Get());
+		MaterialAssignments.Add(AssignedMaterialPathName);
 	}
 	if (!bHasMaterialAssignments)
 	{
 		// Prevent creation of the unrealMaterials attribute in case we don't have any assignments at all
-		MaterialAssignments.clear();
+		MaterialAssignments.Reset();
 	}
 
 	// Create and fill skeleton
diff --git a/Engine/Plugins/Importers/USDImporter/Source/USDUtilities/Public/USDShadeConversion.h b/Engine/Plugins/Importers/USDImporter/Source/USDUtilities/Public/USDShadeConversion.h
--- a/Engine/Plugins/Importers/USDImporter/Source/USDUtilities/Public/USDShadeConversion.h
+++ b/Engine/Plugins/Importers/USDImporter/Source/USDUtilities/Public/USDShadeConversion.h
@@ -182,6 +182,12 @@ namespace UnrealToUsd
 
 namespace UsdUtils
 {
+	// Writes UnrealMaterialPathName as a material binding for MeshOrGeomSubsetPrim, either by reusing an existing UsdShadeMaterial
+	// binding if it already has an 'unreal' render context output and the expected structure, or by creating a new Material prim
+	// that fulfills those requirements.
+	// Doesn't write to the 'unrealMaterial' attribute at all, as we intend on deprecating it in the future.
+	USDUTILITIES_API void AuthorUnrealMaterialBinding(pxr::UsdPrim& MeshOrGeomSubsetPrim, const FString& UnrealMaterialPathName);
+
 	/** Returns a path to an UE asset (e.g. "/Game/Assets/Red.Red") if MaterialPrim has an 'unreal' render context surface output that points at one
 	 */
 	USDUTILITIES_API TOptional<FString> GetUnrealSurfaceOutput(const pxr::UsdPrim& MaterialPrim);
