USD: Rename window to "USD Stage Editor" and move it outside the Virtual Production section.

diff --git a/Engine/Plugins/Importers/USDImporter/Source/USDStageEditor/Private/USDStageEditorModule.cpp b/Engine/Plugins/Importers/USDImporter/Source/USDStageEditor/Private/USDStageEditorModule.cpp
--- a/Engine/Plugins/Importers/USDImporter/Source/USDStageEditor/Private/USDStageEditorModule.cpp
+++ b/Engine/Plugins/Importers/USDImporter/Source/USDStageEditor/Private/USDStageEditorModule.cpp
@@ -253,7 +253,7 @@ public:
 		// clang-format off
 		return SNew(SDockTab)
 			.TabRole(ETabRole::NomadTab)
-			.Label(LOCTEXT("USDStage", "USD Stage"))
+			.Label(LOCTEXT("USDStageEditorTab", "USD Stage Editor"))
 			[
 				SNew(SBorder)
 				.Padding(0)
@@ -285,9 +285,9 @@ public:
 					UE::UsdStageEditorModule::Private::UsdStageEditorTabID.TabType,
 					FOnSpawnTab::CreateStatic(&FUsdStageEditorModule::SpawnUsdStageTab)
 				)
-				.SetDisplayName(LOCTEXT("USDStage", "USD Stage"))
-				.SetTooltipText(LOCTEXT("USDStageTab", "Open USD Stage tab"))
-				.SetGroup(WorkspaceMenu::GetMenuStructure().GetLevelEditorVirtualProductionCategory())
+				.SetDisplayName(LOCTEXT("USDStageEditorMenuItem", "USD Stage Editor"))
+				.SetTooltipText(LOCTEXT("USDStageEditorTooltip", "Open the USD Stage Editor tab. Use this to open and manage USD Stages without importing."))
+				.SetGroup(WorkspaceMenu::GetMenuStructure().GetLevelEditorCategory())
 				.SetIcon(LayersIcon);
 				// clang-format on
 			}
