USD: Refresh prim info panel when tree view needs full update.

diff --git a/Engine/Plugins/Importers/USDImporter/Source/USDStageEditor/Private/SUSDStage.cpp b/Engine/Plugins/Importers/USDImporter/Source/USDStageEditor/Private/SUSDStage.cpp
--- a/Engine/Plugins/Importers/USDImporter/Source/USDStageEditor/Private/SUSDStage.cpp
+++ b/Engine/Plugins/Importers/USDImporter/Source/USDStageEditor/Private/SUSDStage.cpp
@@ -2697,7 +2697,7 @@ void SUsdStage::OnSlateTick(float Time)
 		}
 	}
 
-	if (UsdPrimInfoWidget && !PrimPathToSelect.IsEmpty())
+	if (UsdPrimInfoWidget && (bStageTreeViewNeedsFullUpdate || !PrimPathToSelect.IsEmpty()))
 	{
 		UsdPrimInfoWidget->SetPrimPath(GetCurrentStage(), *PrimPathToSelect);
 		PrimPathToSelect.Empty();
