USD: Card textures being wiped from draw mode components when undoing;

Fix write out of DrawModeComponents replacing relative with absolute paths to the same file for no reason;
Rename instances of "BoundsComponent" to "DrawModeComponent".

diff --git a/Engine/Plugins/Importers/USDImporter/Source/USDClasses/Private/USDDrawModeComponent.cpp b/Engine/Plugins/Importers/USDImporter/Source/USDClasses/Private/USDDrawModeComponent.cpp
--- a/Engine/Plugins/Importers/USDImporter/Source/USDClasses/Private/USDDrawModeComponent.cpp
+++ b/Engine/Plugins/Importers/USDImporter/Source/USDClasses/Private/USDDrawModeComponent.cpp
@@ -83,11 +83,8 @@ namespace UE::UsdDrawModeComponentImpl::Private
 				// Reference: PrimitiveDrawingUtils.cpp, DrawBox function
 
 				// Calculate verts for a face pointing down Z
-				FVector3f Positions[4] = {
-					FVector3f(-0.5, -0.5, +0.5),
-					FVector3f(-0.5, +0.5, +0.5),
-					FVector3f(+0.5, +0.5, +0.5),
-					FVector3f(+0.5, -0.5, +0.5)};
+				FVector3f Positions[4] =
+					{FVector3f(-0.5, -0.5, +0.5), FVector3f(-0.5, +0.5, +0.5), FVector3f(+0.5, +0.5, +0.5), FVector3f(+0.5, -0.5, +0.5)};
 
 				FVector2f UVs[4] = {
 					FVector2f(0, 0),
@@ -730,46 +727,52 @@ EUsdModelCardFace UsdUtils::GetOppositeFaceOnSameAxis(EUsdModelCardFace Face)
 
 void UUsdDrawModeComponent::SetBoundsMin(const FVector& NewMin)
 {
+	Modify();
 	BoundsMin = NewMin;
 	MarkRenderStateDirty();
 }
 
 void UUsdDrawModeComponent::SetBoundsMax(const FVector& NewMax)
 {
+	Modify();
 	BoundsMax = NewMax;
 	MarkRenderStateDirty();
 }
 
 void UUsdDrawModeComponent::SetDrawMode(EUsdDrawMode NewDrawMode)
 {
+	Modify();
 	DrawMode = NewDrawMode;
 	MarkRenderStateDirty();
 }
 
 void UUsdDrawModeComponent::SetBoundsColor(FLinearColor NewColor)
 {
+	Modify();
 	BoundsColor = NewColor;
 	MarkRenderStateDirty();
 }
 
 void UUsdDrawModeComponent::SetCardGeometry(EUsdModelCardGeometry NewGeometry)
 {
+	Modify();
 	CardGeometry = NewGeometry;
 	MarkRenderStateDirty();
 }
 
 void UUsdDrawModeComponent::SetCardTextureXPos(UTexture2D* NewTexture)
 {
+	Modify();
 	// We have this logic so that clearing the texture on the details panel means
 	// removing the opinion of that face texture in the USD stage
 	// (when we convert the component back to USD we'll clear faces that aren't in AuthoredFaces).
 	if (NewTexture)
 	{
-		AuthoredFaces |= EUsdModelCardFace::XPos;
+		AuthoredFaces |= (int32)EUsdModelCardFace::XPos;
 	}
 	else
 	{
-		AuthoredFaces &= ~EUsdModelCardFace::XPos;
+		AuthoredFaces &= (int32)~EUsdModelCardFace::XPos;
 	}
 
 	CardTextureXPos = NewTexture;
@@ -778,13 +781,14 @@ void UUsdDrawModeComponent::SetCardTextureXPos(UTexture2D* NewTexture)
 
 void UUsdDrawModeComponent::SetCardTextureYPos(UTexture2D* NewTexture)
 {
+	Modify();
 	if (NewTexture)
 	{
-		AuthoredFaces |= EUsdModelCardFace::YPos;
+		AuthoredFaces |= (int32)EUsdModelCardFace::YPos;
 	}
 	else
 	{
-		AuthoredFaces &= ~EUsdModelCardFace::YPos;
+		AuthoredFaces &= (int32)~EUsdModelCardFace::YPos;
 	}
 
 	CardTextureYPos = NewTexture;
@@ -793,13 +797,14 @@ void UUsdDrawModeComponent::SetCardTextureYPos(UTexture2D* NewTexture)
 
 void UUsdDrawModeComponent::SetCardTextureZPos(UTexture2D* NewTexture)
 {
+	Modify();
 	if (NewTexture)
 	{
-		AuthoredFaces |= EUsdModelCardFace::ZPos;
+		AuthoredFaces |= (int32)EUsdModelCardFace::ZPos;
 	}
 	else
 	{
-		AuthoredFaces &= ~EUsdModelCardFace::ZPos;
+		AuthoredFaces &= (int32)~EUsdModelCardFace::ZPos;
 	}
 
 	CardTextureZPos = NewTexture;
@@ -808,13 +813,14 @@ void UUsdDrawModeComponent::SetCardTextureZPos(UTexture2D* NewTexture)
 
 void UUsdDrawModeComponent::SetCardTextureXNeg(UTexture2D* NewTexture)
 {
+	Modify();
 	if (NewTexture)
 	{
-		AuthoredFaces |= EUsdModelCardFace::XNeg;
+		AuthoredFaces |= (int32)EUsdModelCardFace::XNeg;
 	}
 	else
 	{
-		AuthoredFaces &= ~EUsdModelCardFace::XNeg;
+		AuthoredFaces &= (int32)~EUsdModelCardFace::XNeg;
 	}
 
 	CardTextureXNeg = NewTexture;
@@ -823,13 +829,14 @@ void UUsdDrawModeComponent::SetCardTextureXNeg(UTexture2D* NewTexture)
 
 void UUsdDrawModeComponent::SetCardTextureYNeg(UTexture2D* NewTexture)
 {
+	Modify();
 	if (NewTexture)
 	{
-		AuthoredFaces |= EUsdModelCardFace::YNeg;
+		AuthoredFaces |= (int32)EUsdModelCardFace::YNeg;
 	}
 	else
 	{
-		AuthoredFaces &= ~EUsdModelCardFace::YNeg;
+		AuthoredFaces &= (int32)~EUsdModelCardFace::YNeg;
 	}
 
 	CardTextureYNeg = NewTexture;
@@ -838,13 +845,14 @@ void UUsdDrawModeComponent::SetCardTextureYNeg(UTexture2D* NewTexture)
 
 void UUsdDrawModeComponent::SetCardTextureZNeg(UTexture2D* NewTexture)
 {
+	Modify();
 	if (NewTexture)
 	{
-		AuthoredFaces |= EUsdModelCardFace::ZNeg;
+		AuthoredFaces |= (int32)EUsdModelCardFace::ZNeg;
 	}
 	else
 	{
-		AuthoredFaces &= ~EUsdModelCardFace::ZNeg;
+		AuthoredFaces &= (int32)~EUsdModelCardFace::ZNeg;
 	}
 
 	CardTextureZNeg = NewTexture;
@@ -948,7 +956,7 @@ UUsdDrawModeComponent::UUsdDrawModeComponent()
 	, CardTextureYNeg{nullptr}
 	, CardTextureZNeg{nullptr}
 	, MaterialInstances({nullptr, nullptr, nullptr, nullptr, nullptr, nullptr})
-	, AuthoredFaces{EUsdModelCardFace::None}
+	, AuthoredFaces{(int32)EUsdModelCardFace::None}
 {
 	RefreshMaterialInstances();
 }
@@ -1083,12 +1091,13 @@ HHitProxy* UUsdDrawModeComponent::CreateMeshHitProxy(int32 SectionIndex, int32 M
 
 EUsdModelCardFace UUsdDrawModeComponent::GetAuthoredFaces() const
 {
-	return AuthoredFaces;
+	return (EUsdModelCardFace)AuthoredFaces;
 }
 
 void UUsdDrawModeComponent::SetAuthoredFaces(EUsdModelCardFace NewAuthoredFaces)
 {
-	AuthoredFaces = NewAuthoredFaces;
+	Modify();
+	AuthoredFaces = (int32)NewAuthoredFaces;
 	MarkRenderStateDirty();
 }
 
diff --git a/Engine/Plugins/Importers/USDImporter/Source/USDClasses/Public/USDDrawModeComponent.h b/Engine/Plugins/Importers/USDImporter/Source/USDClasses/Public/USDDrawModeComponent.h
--- a/Engine/Plugins/Importers/USDImporter/Source/USDClasses/Public/USDDrawModeComponent.h
+++ b/Engine/Plugins/Importers/USDImporter/Source/USDClasses/Public/USDDrawModeComponent.h
@@ -100,13 +100,16 @@ private:
 	UPROPERTY()
 	TArray<TObjectPtr<UMaterialInstance>> MaterialInstances;
 
-	// Flags to enable the card faces to draw.
+	// EUsdModelCardFace flags to enable the card faces to draw.
 	// We need these because the required USD behavior specifies a difference between not having a texture because
 	// one wasn't authored (at which point we flip the opposite face), and not having a texture because it failed
 	// to resolve (at which point we draw vertex color). In both of those cases we'll end up with nullptr on the
-	// CardTexture property, but this will help us tell which case we're talking about
+	// CardTexture property, but this will help us tell which case we're talking about.
+	// Note that this must be an int32 instead of the actual EUsdModelCardFace type because otherwise during any
+	// transaction the UPROPERTY Enum code will freak out if we have more than one flag set, presumably consider
+	// the combination of flags an invalid value for the enum, and then just set it to zero instead (see UE-200646).
 	UPROPERTY()
-	EUsdModelCardFace AuthoredFaces;
+	int32 AuthoredFaces;
 
 public:
 	UFUNCTION(BlueprintSetter, Category = "USD")
diff --git a/Engine/Plugins/Importers/USDImporter/Source/USDExporter/Private/USDConversionBlueprintContext.cpp b/Engine/Plugins/Importers/USDImporter/Source/USDExporter/Private/USDConversionBlueprintContext.cpp
--- a/Engine/Plugins/Importers/USDImporter/Source/USDExporter/Private/USDConversionBlueprintContext.cpp
+++ b/Engine/Plugins/Importers/USDImporter/Source/USDExporter/Private/USDConversionBlueprintContext.cpp
@@ -233,7 +233,7 @@ bool UUsdConversionBlueprintContext::ConvertSpotLightComponent(const USpotLightC
 #endif	  // USE_USD_SDK
 }
 
-bool UUsdConversionBlueprintContext::ConvertBoundsComponent(const UUsdDrawModeComponent* Component, const FString& PrimPath, float TimeCode)
+bool UUsdConversionBlueprintContext::ConvertDrawModeComponent(const UUsdDrawModeComponent* Component, const FString& PrimPath, float TimeCode)
 {
 #if USE_USD_SDK
 	UE::FUsdPrim Prim = UnrealToUsdImpl::GetPrim(Stage, PrimPath);
@@ -243,7 +243,7 @@ bool UUsdConversionBlueprintContext::ConvertBoundsComponent(const UUsdDrawModeCo
 	}
 
 	const bool bWriteExtents = true;
-	return UnrealToUsd::ConvertBoundsComponent(*Component, Prim, bWriteExtents, TimeCode == FLT_MAX ? UsdUtils::GetDefaultTimeCode() : TimeCode);
+	return UnrealToUsd::ConvertDrawModeComponent(*Component, Prim, bWriteExtents, TimeCode == FLT_MAX ? UsdUtils::GetDefaultTimeCode() : TimeCode);
 #else
 	return false;
 #endif	  // USE_USD_SDK
diff --git a/Engine/Plugins/Importers/USDImporter/Source/USDExporter/Public/USDConversionBlueprintContext.h b/Engine/Plugins/Importers/USDImporter/Source/USDExporter/Public/USDConversionBlueprintContext.h
--- a/Engine/Plugins/Importers/USDImporter/Source/USDExporter/Public/USDConversionBlueprintContext.h
+++ b/Engine/Plugins/Importers/USDImporter/Source/USDExporter/Public/USDConversionBlueprintContext.h
@@ -121,7 +121,7 @@ public:
 	bool ConvertSpotLightComponent(const USpotLightComponent* Component, const FString& PrimPath, float TimeCode = 3.402823466e+38F);
 
 	UFUNCTION(BlueprintCallable, Category = "Component conversion")
-	bool ConvertBoundsComponent(const UUsdDrawModeComponent* Component, const FString& PrimPath, float TimeCode = 3.402823466e+38F);
+	bool ConvertDrawModeComponent(const UUsdDrawModeComponent* Component, const FString& PrimPath, float TimeCode = 3.402823466e+38F);
 
 	UFUNCTION(BlueprintCallable, Category = "Component conversion")
 	bool ConvertSceneComponent(const USceneComponent* Component, const FString& PrimPath);
diff --git a/Engine/Plugins/Importers/USDImporter/Source/USDSchemas/Private/USDGeomXformableTranslator.cpp b/Engine/Plugins/Importers/USDImporter/Source/USDSchemas/Private/USDGeomXformableTranslator.cpp
--- a/Engine/Plugins/Importers/USDImporter/Source/USDSchemas/Private/USDGeomXformableTranslator.cpp
+++ b/Engine/Plugins/Importers/USDImporter/Source/USDSchemas/Private/USDGeomXformableTranslator.cpp
@@ -687,7 +687,7 @@ void FUsdGeomXformableTranslator::UpdateComponents(USceneComponent* SceneCompone
 				StaticMeshComponent->BodyInstance.SetCollisionProfileName(PrimStaticMesh->GetBodySetup()->DefaultInstance.GetCollisionProfileName());
 			}
 		}
-		else if (UUsdDrawModeComponent* BoundsComponent = Cast<UUsdDrawModeComponent>(SceneComponent))
+		else if (UUsdDrawModeComponent* DrawModeComponent = Cast<UUsdDrawModeComponent>(SceneComponent))
 		{
 			TOptional<FWriteScopeLock> BBoxLock;
 			pxr::UsdGeomBBoxCache* PxrBBoxCache = nullptr;
@@ -696,7 +696,7 @@ void FUsdGeomXformableTranslator::UpdateComponents(USceneComponent* SceneCompone
 				BBoxLock.Emplace(UEBBoxCache->Lock);
 				PxrBBoxCache = &static_cast<pxr::UsdGeomBBoxCache&>(*UEBBoxCache);
 			}
-			UsdToUnreal::ConvertBounds(Prim, BoundsComponent, Context->Time, PxrBBoxCache);
+			UsdToUnreal::ConvertDrawMode(Prim, DrawModeComponent, Context->Time, PxrBBoxCache);
 		}
 
 		// Handle LiveLink, but only if we're not a skeletal case: The SkelSkeletonTranslator will deal with the
@@ -759,12 +759,7 @@ void FUsdGeomXformableTranslator::UpdateComponents(USceneComponent* SceneCompone
 
 namespace UE::UsdXformableTranslatorImpl::Private
 {
-	void AssignBoundsComponentTextures(
-		UE::FUsdPrim Prim,
-		UUsdDrawModeComponent* BoundsComponent,
-		UUsdAssetCache2& AssetCache,
-		FUsdInfoCache& InfoCache
-	)
+	void AssignDrawModeComponentTextures(UE::FUsdPrim Prim, UUsdDrawModeComponent* DrawModeComponent, UUsdAssetCache2& AssetCache, FUsdInfoCache& InfoCache)
 	{
 		if (!Prim)
 		{
@@ -833,7 +828,7 @@ namespace UE::UsdXformableTranslatorImpl::Private
 
 		TFunction<void(const pxr::UsdAttribute&, TextureSetterFunc, EUsdModelCardFace)> HandleCardFace =
 			[&AuthoredFaces,
-			 BoundsComponent,
+			 DrawModeComponent,
 			 &AttrPathToTextures](const pxr::UsdAttribute& Attr, TextureSetterFunc TextureSetter, EUsdModelCardFace Face)
 		{
 			if (Attr && Attr.HasAuthoredValue())
@@ -847,7 +842,7 @@ namespace UE::UsdXformableTranslatorImpl::Private
 					{
 						if (UTexture2D* Texture = iter->second)
 						{
-							(BoundsComponent->*TextureSetter)(Texture);
+							(DrawModeComponent->*TextureSetter)(Texture);
 						}
 					}
 				}
@@ -863,7 +858,7 @@ namespace UE::UsdXformableTranslatorImpl::Private
 		// Override the AuthoredFaces with the correct value. The texture setter functions will all set the authored faces
 		// when we set any texture in the component, but we also want to set as authored the faces where there *was* some
 		// texture authored in USD but we failed to resolve it, so that we can show that face with vertex color like USD specifies
-		BoundsComponent->SetAuthoredFaces(AuthoredFaces);
+		DrawModeComponent->SetAuthoredFaces(AuthoredFaces);
 	}
 }	 // namespace UE::UsdXformableTranslatorImpl::Private
 
@@ -887,7 +882,7 @@ USceneComponent* FUsdGeomXformableTranslator::CreateAlternativeDrawModeComponent
 			{
 				// For now we only assign textures when creating components, not when updating. Maybe in the future we can
 				// add support for "texture animations"
-				UE::UsdXformableTranslatorImpl::Private::AssignBoundsComponentTextures(
+				UE::UsdXformableTranslatorImpl::Private::AssignDrawModeComponentTextures(
 					GetPrim(),
 					Component,
 					*Context->AssetCache,
diff --git a/Engine/Plugins/Importers/USDImporter/Source/USDStage/Private/USDStageActor.cpp b/Engine/Plugins/Importers/USDImporter/Source/USDStage/Private/USDStageActor.cpp
--- a/Engine/Plugins/Importers/USDImporter/Source/USDStage/Private/USDStageActor.cpp
+++ b/Engine/Plugins/Importers/USDImporter/Source/USDStage/Private/USDStageActor.cpp
@@ -4541,7 +4541,7 @@ void AUsdStageActor::OnObjectPropertyChanged(UObject* ObjectBeingModified, FProp
 				{
 					UnrealToUsd::ConvertMeshComponent(CurrentStage, MeshComponent, UsdPrim);
 				}
-				else if (UUsdDrawModeComponent* BoundsComponent = Cast<UUsdDrawModeComponent>(PrimSceneComponent))
+				else if (UUsdDrawModeComponent* DrawModeComponent = Cast<UUsdDrawModeComponent>(PrimSceneComponent))
 				{
 					const static TSet<FName> BoundsProperties = {
 						GET_MEMBER_NAME_CHECKED(UUsdDrawModeComponent, BoundsMin),
@@ -4551,7 +4551,7 @@ void AUsdStageActor::OnObjectPropertyChanged(UObject* ObjectBeingModified, FProp
 					// If we just manually tweaked the extents, also author those back out to USD as extents opinions
 					const bool bWriteExtents = BoundsProperties.Contains(PropertyChangedEvent.GetMemberPropertyName());
 					const double UsdTimeCode = UsdUtils::GetDefaultTimeCode();
-					UnrealToUsd::ConvertBoundsComponent(*BoundsComponent, UsdPrim, bWriteExtents, UsdTimeCode);
+					UnrealToUsd::ConvertDrawModeComponent(*DrawModeComponent, UsdPrim, bWriteExtents, UsdTimeCode);
 				}
 				else if (UsdPrim && UsdPrim.IsA(TEXT("Camera")))
 				{
diff --git a/Engine/Plugins/Importers/USDImporter/Source/USDUtilities/Private/USDPrimConversion.cpp b/Engine/Plugins/Importers/USDImporter/Source/USDUtilities/Private/USDPrimConversion.cpp
--- a/Engine/Plugins/Importers/USDImporter/Source/USDUtilities/Private/USDPrimConversion.cpp
+++ b/Engine/Plugins/Importers/USDImporter/Source/USDUtilities/Private/USDPrimConversion.cpp
@@ -1662,12 +1662,12 @@ UsdToUnreal::FPropertyTrackReader UsdToUnreal::CreatePropertyTrackReader(
 	return Reader;
 }
 
-bool UsdToUnreal::ConvertBounds(const pxr::UsdPrim& Prim, UUsdDrawModeComponent* BoundsComponent, double EvalTime, pxr::UsdGeomBBoxCache* BBoxCache)
+bool UsdToUnreal::ConvertDrawMode(const pxr::UsdPrim& Prim, UUsdDrawModeComponent* DrawModeComponent, double EvalTime, pxr::UsdGeomBBoxCache* BBoxCache)
 {
 	// We're not going to check if Prim actually has the "bounds" draw mode or if it has "applyDrawMode" set
 	// to true, as that can be expensive and this can get called from UpdateComponents, which can get called
 	// every frame of Time animation. If we have a UUsdDrawModeComponent at all we'll assume we're OK here
-	if (!BoundsComponent)
+	if (!DrawModeComponent)
 	{
 		return false;
 	}
@@ -1719,8 +1719,8 @@ bool UsdToUnreal::ConvertBounds(const pxr::UsdPrim& Prim, UUsdDrawModeComponent*
 			TArray<FVector>{UESpaceUSDMin, UESpaceUSDMax}
 		 };
 
-		BoundsComponent->SetBoundsMin(UEBox.Min);
-		BoundsComponent->SetBoundsMax(UEBox.Max);
+		DrawModeComponent->SetBoundsMin(UEBox.Min);
+		DrawModeComponent->SetBoundsMax(UEBox.Max);
 	}
 
 	if (pxr::UsdGeomModelAPI GeomModelAPI{Prim})
@@ -1729,7 +1729,7 @@ bool UsdToUnreal::ConvertBounds(const pxr::UsdPrim& Prim, UUsdDrawModeComponent*
 		pxr::UsdAttribute ColorAttr = GeomModelAPI.GetModelDrawModeColorAttr();
 		if (ColorAttr && ColorAttr.Get(&Color))
 		{
-			BoundsComponent->SetBoundsColor(UsdToUnreal::ConvertColor(Color));
+			DrawModeComponent->SetBoundsColor(UsdToUnreal::ConvertColor(Color));
 		}
 
 		// It's not super efficient to call this whole function every time but not only this lets us
@@ -1739,7 +1739,7 @@ bool UsdToUnreal::ConvertBounds(const pxr::UsdPrim& Prim, UUsdDrawModeComponent*
 		// the prim is not a model, or doesn't have applyDrawMode enabled, etc.) so we can't just check
 		// the drawMode attribute directly here
 		EUsdDrawMode DrawMode = UsdUtils::GetAppliedDrawMode(Prim);
-		BoundsComponent->SetDrawMode(DrawMode);
+		DrawModeComponent->SetDrawMode(DrawMode);
 
 		pxr::TfToken CardGeometryToken;
 		pxr::UsdAttribute GeometryAttr = GeomModelAPI.GetModelCardGeometryAttr();
@@ -1754,7 +1754,7 @@ bool UsdToUnreal::ConvertBounds(const pxr::UsdPrim& Prim, UUsdDrawModeComponent*
 			{
 				CardGeometry = EUsdModelCardGeometry::FromTexture;
 			}
-			BoundsComponent->SetCardGeometry(CardGeometry);
+			DrawModeComponent->SetCardGeometry(CardGeometry);
 		}
 	}
 
@@ -3758,7 +3758,7 @@ bool UnrealToUsd::CreateComponentPropertyBaker(
 			}
 		}
 	}
-	else if (const UUsdDrawModeComponent* BoundsComponent = Cast<UUsdDrawModeComponent>(&Component))
+	else if (const UUsdDrawModeComponent* DrawModeComponent = Cast<UUsdDrawModeComponent>(&Component))
 	{
 		static TSet<FString> RelevantProperties = {
 			GET_MEMBER_NAME_CHECKED(UUsdDrawModeComponent, BoundsMin).ToString(),
@@ -3767,10 +3767,10 @@ bool UnrealToUsd::CreateComponentPropertyBaker(
 		if (RelevantProperties.Contains(PropertyPath))
 		{
 			BakerType = EBakingType::Bounds;
-			BakerFunction = [UsdStage, BoundsComponent, UsdPrim](double UsdTimeCode) mutable
+			BakerFunction = [UsdStage, DrawModeComponent, UsdPrim](double UsdTimeCode) mutable
 			{
 				const bool bWriteExtents = true;
-				UnrealToUsd::ConvertBoundsComponent(*BoundsComponent, UsdPrim, bWriteExtents, UsdTimeCode);
+				UnrealToUsd::ConvertDrawModeComponent(*DrawModeComponent, UsdPrim, bWriteExtents, UsdTimeCode);
 			};
 		}
 	}
@@ -4487,7 +4487,7 @@ UnrealToUsd::FPropertyTrackWriter UnrealToUsd::CreatePropertyTrackWriter(
 			}
 
 			// If we still don't have an attr try applying the schema.
-			// Not entirely sure how this can possibly happen at this point, but this is more for "parity" with ConvertBoundsComponent
+			// Not entirely sure how this can possibly happen at this point, but this is more for "parity" with ConvertDrawModeComponent
 			if (!Attr)
 			{
 				if (pxr::UsdGeomModelAPI GeomModelAPI = pxr::UsdGeomModelAPI::Apply(UsdPrim))
@@ -4944,7 +4944,7 @@ TArray<UE::FUsdAttribute> UnrealToUsd::GetAttributesForProperty(const UE::FUsdPr
 	return {};
 }
 
-bool UnrealToUsd::ConvertBoundsComponent(const UUsdDrawModeComponent& BoundsComponent, pxr::UsdPrim& UsdPrim, bool bWriteExtents, double UsdTimeCode)
+bool UnrealToUsd::ConvertDrawModeComponent(const UUsdDrawModeComponent& DrawModeComponent, pxr::UsdPrim& UsdPrim, bool bWriteExtents, double UsdTimeCode)
 {
 	if (!UsdPrim)
 	{
@@ -4990,8 +4990,8 @@ bool UnrealToUsd::ConvertBoundsComponent(const UUsdDrawModeComponent& BoundsComp
 	{
 		FUsdStageInfo StageInfo{UsdPrim.GetStage()};
 
-		const FVector& UEBoundsMin = BoundsComponent.BoundsMin;
-		const FVector& UEBoundsMax = BoundsComponent.BoundsMax;
+		const FVector& UEBoundsMin = DrawModeComponent.BoundsMin;
+		const FVector& UEBoundsMax = DrawModeComponent.BoundsMax;
 		pxr::GfVec3f UEBoundsMinUsdSpace = UnrealToUsd::ConvertVectorFloat(StageInfo, UEBoundsMin);
 		pxr::GfVec3f UEBoundsMaxUsdSpace = UnrealToUsd::ConvertVectorFloat(StageInfo, UEBoundsMax);
 		pxr::GfVec3f UsdMin{
@@ -5026,10 +5026,10 @@ bool UnrealToUsd::ConvertBoundsComponent(const UUsdDrawModeComponent& BoundsComp
 	if (pxr::UsdAttribute Attr = GeomModelAPI.CreateModelDrawModeAttr())
 	{
 		Attr.Set(
-			BoundsComponent.DrawMode == EUsdDrawMode::Origin	  ? pxr::UsdGeomTokens->origin
-			: BoundsComponent.DrawMode == EUsdDrawMode::Bounds	  ? pxr::UsdGeomTokens->bounds
-			: BoundsComponent.DrawMode == EUsdDrawMode::Cards	  ? pxr::UsdGeomTokens->cards
-			: BoundsComponent.DrawMode == EUsdDrawMode::Inherited ? pxr::UsdGeomTokens->inherited
+			DrawModeComponent.DrawMode == EUsdDrawMode::Origin	  ? pxr::UsdGeomTokens->origin
+			: DrawModeComponent.DrawMode == EUsdDrawMode::Bounds	  ? pxr::UsdGeomTokens->bounds
+			: DrawModeComponent.DrawMode == EUsdDrawMode::Cards	  ? pxr::UsdGeomTokens->cards
+			: DrawModeComponent.DrawMode == EUsdDrawMode::Inherited ? pxr::UsdGeomTokens->inherited
 																  : pxr::UsdGeomTokens->default_,
 			UsdTimeCode
 		);
@@ -5038,8 +5038,8 @@ bool UnrealToUsd::ConvertBoundsComponent(const UUsdDrawModeComponent& BoundsComp
 	if (pxr::UsdAttribute Attr = GeomModelAPI.CreateModelCardGeometryAttr())
 	{
 		Attr.Set(
-			BoundsComponent.CardGeometry == EUsdModelCardGeometry::Cross ? pxr::UsdGeomTokens->cross
-			: BoundsComponent.CardGeometry == EUsdModelCardGeometry::Box ? pxr::UsdGeomTokens->box
+			DrawModeComponent.CardGeometry == EUsdModelCardGeometry::Cross ? pxr::UsdGeomTokens->cross
+			: DrawModeComponent.CardGeometry == EUsdModelCardGeometry::Box ? pxr::UsdGeomTokens->box
 																		 : pxr::UsdGeomTokens->fromTexture,
 			UsdTimeCode
 		);
@@ -5056,7 +5056,7 @@ bool UnrealToUsd::ConvertBoundsComponent(const UUsdDrawModeComponent& BoundsComp
 	{
 		// This color is just a vec3f and our color is already linear anyway, so just convert it directly instead of going through
 		// UnrealToUsd::ConvertColor
-		pxr::GfVec3f UsdColor{BoundsComponent.BoundsColor.R, BoundsComponent.BoundsColor.G, BoundsComponent.BoundsColor.B};
+		pxr::GfVec3f UsdColor{DrawModeComponent.BoundsColor.R, DrawModeComponent.BoundsColor.G, DrawModeComponent.BoundsColor.B};
 		Attr.Set(UsdColor, UsdTimeCode);
 	}
 
@@ -5067,15 +5067,15 @@ bool UnrealToUsd::ConvertBoundsComponent(const UUsdDrawModeComponent& BoundsComp
 	// when setting/clearing textures via blueprint/details panels we will also tweak AuthoredFaces. The combined effect is that when a user sets a
 	// new texture in a property we will assume that means it became "authored", and when they clear a texture property we will assume that ceases to
 	// be "authored"
-	EUsdModelCardFace AuthoredFaces = BoundsComponent.GetAuthoredFaces();
+	EUsdModelCardFace AuthoredFaces = DrawModeComponent.GetAuthoredFaces();
 
 	using CreateAttrFuncType = decltype(&pxr::UsdGeomModelAPI::CreateModelCardTextureXPosAttr);
 	using GetAttrFuncType = decltype(&pxr::UsdGeomModelAPI::GetModelCardTextureXPosAttr);
 
 	TFunction<void(EUsdModelCardFace, GetAttrFuncType, CreateAttrFuncType)> ExportCardFace =
-		[&BoundsComponent, AuthoredFaces, UsdTimeCode, &GeomModelAPI](EUsdModelCardFace Face, GetAttrFuncType GetAttr, CreateAttrFuncType CreateAttr)
+		[&DrawModeComponent, AuthoredFaces, UsdTimeCode, &GeomModelAPI](EUsdModelCardFace Face, GetAttrFuncType GetAttr, CreateAttrFuncType CreateAttr)
 	{
-		UTexture2D* FaceTexture = BoundsComponent.GetTextureForFace(Face);
+		UTexture2D* FaceTexture = DrawModeComponent.GetTextureForFace(Face);
 
 		if (EnumHasAllFlags(AuthoredFaces, Face))
 		{
@@ -5089,12 +5089,46 @@ bool UnrealToUsd::ConvertBoundsComponent(const UUsdDrawModeComponent& BoundsComp
 #if WITH_EDITOR
 					if (FaceTexture->AssetImportData)
 					{
-						// If this texture doesn't have an asset import data, or the filename stored there is not valid then this reference
-						// won't resolve when it read it back from USD...
-						// TOD: Should we export the texture now? To where?
 						FString TextureSourcePath = FaceTexture->AssetImportData->GetFirstFilename();
-						pxr::SdfAssetPath AssetPath = pxr::SdfAssetPath{UnrealToUsd::ConvertString(*TextureSourcePath).Get()};
-						Attr.Set(AssetPath, UsdTimeCode);
+						FString ResolvedPath = UsdUtils::GetResolvedTexturePath(Attr);
+
+						// Avoid authoring anything unless they point at different files because in the general case the
+						// asset import data will have an absolute path, while the path on the attribute may be currently relative.
+						// We still want to implement some kind of larger feature to let the user pick whether we should be authoring
+						// relative or absolute paths all over, but for now the least we can do is try not changing relative paths
+						// to absolute unnecessarily (or vice versa)
+						if (!FPaths::IsSamePath(ResolvedPath, TextureSourcePath))
+						{
+							if (!FPaths::FileExists(TextureSourcePath))
+							{
+								UE_LOG(
+									LogUsd,
+									Warning,
+									TEXT("Authoring card texture path '%s' for texture '%s' onto attribute '%s', but the source image file doesn't "
+										 "exist on disk! It may not be possible to display this card texture if the stage is reloaded or reopened. "
+										 "UTexture assets can't be automatically exported in this manner just yet. If you want to assign a new "
+										 "texture, make sure the UTexture asset's source image file exists on disk."),
+									*TextureSourcePath,
+									*FaceTexture->GetPathName(),
+									*UsdToUnreal::ConvertPath(Attr.GetPath())
+								);
+							}
+
+							pxr::SdfAssetPath AssetPath = pxr::SdfAssetPath{UnrealToUsd::ConvertString(*TextureSourcePath).Get()};
+							Attr.Set(AssetPath, UsdTimeCode);
+						}
+					}
+					else
+					{
+						UE_LOG(
+							LogUsd,
+							Warning,
+							TEXT("Not authoring card texture for attribute '%s' because the assigned texture '%s' has no AssetImportData! UTexture "
+								 "assets can't be automatically exported in this manner just yet. If you want to assign a new texture, make sure the "
+								 "UTexture asset's source image file exists on disk."),
+							*UsdToUnreal::ConvertPath(Attr.GetPath()),
+							*FaceTexture->GetPathName()
+						);
 					}
 #endif	  // WITH_EDITOR
 				}
diff --git a/Engine/Plugins/Importers/USDImporter/Source/USDUtilities/Public/USDPrimConversion.h b/Engine/Plugins/Importers/USDImporter/Source/USDUtilities/Public/USDPrimConversion.h
--- a/Engine/Plugins/Importers/USDImporter/Source/USDUtilities/Public/USDPrimConversion.h
+++ b/Engine/Plugins/Importers/USDImporter/Source/USDUtilities/Public/USDPrimConversion.h
@@ -194,9 +194,9 @@ namespace UsdToUnreal
 	 * UUsdDrawModeComponent at EvalTime.
 	 * Must use a BBoxCache: If one is not provided it will be created on-demand for this call alone.
 	 */
-	USDUTILITIES_API bool ConvertBounds(
+	USDUTILITIES_API bool ConvertDrawMode(
 		const pxr::UsdPrim& Prim,
-		UUsdDrawModeComponent* BoundsComponent,
+		UUsdDrawModeComponent* DrawModeComponent,
 		double EvalTime = UsdUtils::GetDefaultTimeCode(),
 		pxr::UsdGeomBBoxCache* BBoxCache = nullptr
 	);
@@ -421,8 +421,8 @@ namespace UnrealToUsd
 	 * Note that this will even apply the schema itself to UsdPrim if it doesn't already have it, and potentially author
 	 * the 'kind' metadata on it and all of its ancestors (as UsdGeomModelAPI requires the prim to be a "model" to have function)
 	 */
-	USDUTILITIES_API bool ConvertBoundsComponent(
-		const UUsdDrawModeComponent& BoundsComponent,
+	USDUTILITIES_API bool ConvertDrawModeComponent(
+		const UUsdDrawModeComponent& DrawModeComponent,
 		pxr::UsdPrim& UsdPrim,
 		bool bWriteExtents = false,
 		double UsdTimeCode = UsdUtils::GetDefaultTimeCode()
