USD: Fix handling of asset user data for reused assets;

Clear previously collected metadata whenever the options prevent collection.

diff --git a/Engine/Plugins/Importers/USDImporter/Source/USDSchemas/Private/MeshTranslationImpl.cpp b/Engine/Plugins/Importers/USDImporter/Source/USDSchemas/Private/MeshTranslationImpl.cpp
--- a/Engine/Plugins/Importers/USDImporter/Source/USDSchemas/Private/MeshTranslationImpl.cpp
+++ b/Engine/Plugins/Importers/USDImporter/Source/USDSchemas/Private/MeshTranslationImpl.cpp
@@ -52,16 +52,6 @@ namespace UE::MeshTranslationImplInternal::Private
 			return nullptr;
 		}
 
-		UUsdMaterialAssetUserData* OneSidedUserData = OneSidedMat->GetAssetUserData<UUsdMaterialAssetUserData>();
-		if (!ensureMsgf(
-				OneSidedUserData,
-				TEXT("Expected one-sided material '%s' to have an UUsdMaterialAssetUserData at this point!"),
-				*OneSidedMat->GetPathName()
-			))
-		{
-			return nullptr;
-		}
-
 		UMaterialInterface* TwoSidedMat = nullptr;
 
 		UMaterialInstance* OneSidedMaterialInstance = Cast<UMaterialInstance>(OneSidedMat);
@@ -143,12 +133,6 @@ namespace UE::MeshTranslationImplInternal::Private
 				TwoSidedMat = TwoSidedMID;
 			}
 
-		if (TwoSidedMat)
-		{
-			UUsdMaterialAssetUserData* UserData = DuplicateObject(OneSidedUserData, TwoSidedMat, TEXT("USDAssetUserData"));
-			TwoSidedMat->AddAssetUserData(UserData);
-		}
-
 		return TwoSidedMat;
 	}
 
@@ -326,16 +310,17 @@ namespace UE::MeshTranslationImplInternal::Private
 			CompatibleMaterial = ExistingCompatibleMaterial;
 		}
 
+		TMap<FString, int32> CompatiblePrimvarToUVIndex;
+		CompatiblePrimvarToUVIndex.Reserve(CompatiblePrimvarAndUVIndexPairs.Num());
+		for (const TPair<FString, int32>& Pair : CompatiblePrimvarAndUVIndexPairs)
+		{
+			CompatiblePrimvarToUVIndex.Add(Pair);
+		}
+
 		// We have to create a brand new compatible material instance
+		bool bCreatedNew = false;
 		if (!CompatibleMaterial)
 		{
-			TMap<FString, int32> CompatiblePrimvarToUVIndex;
-			CompatiblePrimvarToUVIndex.Reserve(CompatiblePrimvarAndUVIndexPairs.Num());
-			for (const TPair<FString, int32>& Pair : CompatiblePrimvarAndUVIndexPairs)
-			{
-				CompatiblePrimvarToUVIndex.Add(Pair);
-			}
-
 			const FName NewInstanceName = MakeUniqueObjectName(GetTransientPackage(), UMaterialInstance::StaticClass(), Material.GetFName());
 
 			UE_LOG(
@@ -369,41 +354,51 @@ namespace UE::MeshTranslationImplInternal::Private
 				CompatibleMaterial = CompatibleMID;
 			}
 
-			UUsdMaterialAssetUserData* CompatibleUserData = DuplicateObject(MaterialAssetUserData, CompatibleMaterial, TEXT("USDAssetUserData"));
+			bCreatedNew = true;
+		}
+
+		// Update the AssetUserData whether we created a new material instance or reused one from the asset cache.
+		// The compatible AssetUserData should always match the original except for the different PrimvarToUVIndex
+		UUsdMaterialAssetUserData* CompatibleUserData = nullptr;
+		if (CompatibleMaterial)
+		{
+			CompatibleUserData = DuplicateObject(MaterialAssetUserData, CompatibleMaterial, TEXT("USDAssetUserData"));
 			CompatibleUserData->PrimvarToUVIndex = CompatiblePrimvarToUVIndex;
-			CompatibleMaterial->AddAssetUserData(CompatibleUserData);
 
-			if (UMaterialInstance* CompatibleInstance = Cast<UMaterialInstance>(CompatibleMaterial))
+			UsdUtils::SetAssetUserData(CompatibleMaterial, CompatibleUserData);
+		}
+
+		// Now that the AssetUserData is done, actually set the UV index material parameters with the target indices
+		UMaterialInstance* CompatibleInstance = Cast<UMaterialInstance>(CompatibleMaterial);
+		if (bCreatedNew && CompatibleInstance && CompatibleUserData)
+		{
+			for (const TPair<FString, FString>& ParameterPair : CompatibleUserData->ParameterToPrimvar)
 			{
-				// Actually set the parameters with the target UV indices
-				for (const TPair<FString, FString>& ParameterPair : CompatibleUserData->ParameterToPrimvar)
-				{
-					const FString& Parameter = ParameterPair.Key;
-					const FString& Primvar = ParameterPair.Value;
+				const FString& Parameter = ParameterPair.Key;
+				const FString& Primvar = ParameterPair.Value;
 
-					if (int32* UVIndex = CompatibleUserData->PrimvarToUVIndex.Find(Primvar))
+				if (int32* UVIndex = CompatibleUserData->PrimvarToUVIndex.Find(Primvar))
+				{
+					// Force-disable using the texture at all if the mesh doesn't provide the primvar that should be
+					// used to sample it with
+					if (*UVIndex == UNUSED_UV_INDEX)
 					{
-						// Force-disable using the texture at all if the mesh doesn't provide the primvar that should be
-						// used to sample it with
-						if (*UVIndex == UNUSED_UV_INDEX)
-						{
-							UsdUtils::SetScalarParameterValue(*CompatibleInstance, *FString::Printf(TEXT("Use%sTexture"), *Parameter), 0.0f);
-						}
-						else
-						{
-							UsdUtils::SetScalarParameterValue(
-								*CompatibleInstance,
-								*FString::Printf(TEXT("%sUVIndex"), *Parameter),
-								static_cast<float>(*UVIndex)
-							);
-						}
+						UsdUtils::SetScalarParameterValue(*CompatibleInstance, *FString::Printf(TEXT("Use%sTexture"), *Parameter), 0.0f);
+					}
+					else
+					{
+						UsdUtils::SetScalarParameterValue(
+							*CompatibleInstance,
+							*FString::Printf(TEXT("%sUVIndex"), *Parameter),
+							static_cast<float>(*UVIndex)
+						);
 					}
 				}
+			}
 
 #if WITH_EDITOR
-				CompatibleInstance->PostEditChange();
+			CompatibleInstance->PostEditChange();
 #endif	  // WITH_EDITOR
-			}
 		}
 
 		if (CompatibleMaterial && CompatibleMaterial != &Material)
@@ -579,6 +574,15 @@ TMap<const UsdUtils::FUsdPrimMaterialSlot*, UMaterialInterface*> MeshTranslation
 
 						if (TwoSidedMat)
 						{
+							// Update AssetUserData whether we generated a new material or reused one from the asset cache
+							{
+								UUsdMaterialAssetUserData* OneSidedUserData = OneSidedMat->GetAssetUserData<UUsdMaterialAssetUserData>();
+								ensure(OneSidedUserData);
+
+								UUsdMaterialAssetUserData* UserData = DuplicateObject(OneSidedUserData, TwoSidedMat, TEXT("USDAssetUserData"));
+								UsdUtils::SetAssetUserData(TwoSidedMat, UserData);
+							}
+
 							TwoSidedMat->SetFlags(RF_Transient);
 							Material = TwoSidedMat;
 							PrefixedMaterialHash = PrefixedTwoSidedHash;
diff --git a/Engine/Plugins/Importers/USDImporter/Source/USDSchemas/Private/USDGeomMeshTranslator.cpp b/Engine/Plugins/Importers/USDImporter/Source/USDSchemas/Private/USDGeomMeshTranslator.cpp
--- a/Engine/Plugins/Importers/USDImporter/Source/USDSchemas/Private/USDGeomMeshTranslator.cpp
+++ b/Engine/Plugins/Importers/USDImporter/Source/USDSchemas/Private/USDGeomMeshTranslator.cpp
@@ -1698,6 +1698,11 @@ void FBuildStaticMeshTaskChain::SetupTasks()
 						   );
 					   }
 				   }
+				   else
+				   {
+					   // Strip the metadata from this prim, so that if we uncheck "Collect Metadata" it actually disappears on the AssetUserData
+					   UserData->StageIdentifierToMetadata.Remove(GetPrim().GetStage().GetRootLayer().GetIdentifier());
+				   }
 
 				   MeshTranslationImpl::RecordSourcePrimsForMaterialSlots(LODIndexToMaterialInfo, UserData);
 			   }
diff --git a/Engine/Plugins/Importers/USDImporter/Source/USDSchemas/Private/USDGeomXformableTranslator.cpp b/Engine/Plugins/Importers/USDImporter/Source/USDSchemas/Private/USDGeomXformableTranslator.cpp
--- a/Engine/Plugins/Importers/USDImporter/Source/USDSchemas/Private/USDGeomXformableTranslator.cpp
+++ b/Engine/Plugins/Importers/USDImporter/Source/USDSchemas/Private/USDGeomXformableTranslator.cpp
@@ -573,6 +573,11 @@ USceneComponent* FUsdGeomXformableTranslator::CreateComponentsEx(TOptional<TSubc
 			bCollectMetadataFromSubtree
 		);
 	}
+	else if (UUsdAssetUserData* UserData = UsdUtils::GetAssetUserData(SceneComponent))
+	{
+		// Strip the metadata from this prim, so that if we uncheck "Collect Metadata" it actually disappears on the AssetUserData
+		UserData->StageIdentifierToMetadata.Remove(Prim.GetStage().GetRootLayer().GetIdentifier());
+	}
 
 	if (SceneComponent)
 	{
diff --git a/Engine/Plugins/Importers/USDImporter/Source/USDSchemas/Private/USDGeometryCacheTranslator.cpp b/Engine/Plugins/Importers/USDImporter/Source/USDSchemas/Private/USDGeometryCacheTranslator.cpp
--- a/Engine/Plugins/Importers/USDImporter/Source/USDSchemas/Private/USDGeometryCacheTranslator.cpp
+++ b/Engine/Plugins/Importers/USDImporter/Source/USDSchemas/Private/USDGeometryCacheTranslator.cpp
@@ -902,6 +902,11 @@ void FGeometryCacheCreateAssetsTaskChain::SetupTasks()
 							Context->MetadataOptions.bCollectFromEntireSubtrees
 						);
 					}
+					else
+					{
+						// Strip the metadata from this prim, so that if we uncheck "Collect Metadata" it actually disappears on the AssetUserData
+						UserData->StageIdentifierToMetadata.Remove(GetPrim().GetStage().GetRootLayer().GetIdentifier());
+					}
 
 					MeshTranslationImpl::RecordSourcePrimsForMaterialSlots(LODIndexToMaterialInfo, UserData);
 				}
diff --git a/Engine/Plugins/Importers/USDImporter/Source/USDSchemas/Private/USDGroomTranslator.cpp b/Engine/Plugins/Importers/USDImporter/Source/USDSchemas/Private/USDGroomTranslator.cpp
--- a/Engine/Plugins/Importers/USDImporter/Source/USDSchemas/Private/USDGroomTranslator.cpp
+++ b/Engine/Plugins/Importers/USDImporter/Source/USDSchemas/Private/USDGroomTranslator.cpp
@@ -225,6 +225,12 @@ protected:
 									Context->MetadataOptions.bCollectFromEntireSubtrees
 								);
 							}
+							else
+							{
+								// Strip the metadata from this prim, so that if we uncheck "Collect Metadata" it actually disappears on the
+								// AssetUserData
+								UserData->StageIdentifierToMetadata.Remove(GetPrim().GetStage().GetRootLayer().GetIdentifier());
+							}
 						}
 					}
 				}
@@ -433,6 +439,11 @@ protected:
 								Context->MetadataOptions.bCollectFromEntireSubtrees
 							);
 						}
+						else
+						{
+							// Strip the metadata from this prim, so that if we uncheck "Collect Metadata" it actually disappears on the AssetUserData
+							UserData->StageIdentifierToMetadata.Remove(GetPrim().GetStage().GetRootLayer().GetIdentifier());
+						}
 					}
 
 					Context->AssetCache->CacheAsset(PrefixedGroomCacheHash, GroomCache);
diff --git a/Engine/Plugins/Importers/USDImporter/Source/USDSchemas/Private/USDShadeMaterialTranslator.cpp b/Engine/Plugins/Importers/USDImporter/Source/USDSchemas/Private/USDShadeMaterialTranslator.cpp
--- a/Engine/Plugins/Importers/USDImporter/Source/USDSchemas/Private/USDShadeMaterialTranslator.cpp
+++ b/Engine/Plugins/Importers/USDImporter/Source/USDSchemas/Private/USDShadeMaterialTranslator.cpp
@@ -492,6 +492,11 @@ void FUsdShadeMaterialTranslator::PostImportMaterial(const FString& PrefixedMate
 				Context->MetadataOptions.bCollectFromEntireSubtrees
 			);
 		}
+		else
+		{
+			// Strip the metadata from this prim, so that if we uncheck "Collect Metadata" it actually disappears on the AssetUserData
+			UserData->StageIdentifierToMetadata.Remove(GetPrim().GetStage().GetRootLayer().GetIdentifier());
+		}
 	}
 
 	// Note that this needs to run even if we found this material in the asset cache already, otherwise we won't
diff --git a/Engine/Plugins/Importers/USDImporter/Source/USDSchemas/Private/USDSkelSkeletonTranslator.cpp b/Engine/Plugins/Importers/USDImporter/Source/USDSchemas/Private/USDSkelSkeletonTranslator.cpp
--- a/Engine/Plugins/Importers/USDImporter/Source/USDSchemas/Private/USDSkelSkeletonTranslator.cpp
+++ b/Engine/Plugins/Importers/USDImporter/Source/USDSchemas/Private/USDSkelSkeletonTranslator.cpp
@@ -1562,7 +1562,15 @@ namespace UsdSkelSkeletonTranslatorImpl
 
 						// For the skel task chain we always collect skeletal mesh metadata when first parsing the prims directly, as it
 						// allows us to do it while we're flipping through LOD variants, if any
-						UserData->StageIdentifierToMetadata.Add(GetSkeletonPrim().GetStage().GetRootLayer().GetIdentifier(), LODMetadata);
+						if (Context->MetadataOptions.bCollectMetadata)
+						{
+							UserData->StageIdentifierToMetadata.Add(GetSkeletonPrim().GetStage().GetRootLayer().GetIdentifier(), LODMetadata);
+						}
+						else
+						{
+							// Strip the metadata from this prim, so that if we uncheck "Collect Metadata" it actually disappears on the AssetUserData
+							UserData->StageIdentifierToMetadata.Remove(GetSkeletonPrim().GetStage().GetRootLayer().GetIdentifier());
+						}
 
 						MeshTranslationImpl::RecordSourcePrimsForMaterialSlots(LODIndexToMaterialInfo, UserData);
 					}
@@ -1589,6 +1597,12 @@ namespace UsdSkelSkeletonTranslatorImpl
 									bCollectFromEntireSubtrees
 								);
 							}
+							else
+							{
+								// Strip the metadata from this prim, so that if we uncheck "Collect Metadata" it actually disappears on the
+								// AssetUserData
+								UserData->StageIdentifierToMetadata.Remove(GetSkeletonPrim().GetStage().GetRootLayer().GetIdentifier());
+							}
 						}
 					}
 
@@ -1789,6 +1803,7 @@ namespace UsdSkelSkeletonTranslatorImpl
 				FString PrefixedSkelAnimHash = UsdUtils::GetAssetHashPrefix(SkelAnimationPrim, Context->bReuseIdenticalAssets) + Hash.ToString();
 				UAnimSequence* AnimSequence = Cast<UAnimSequence>(Context->AssetCache->GetCachedAsset(PrefixedSkelAnimHash));
 
+				TOptional<float> LayerStartOffsetSeconds;
 				if (!AnimSequence || AnimSequence->GetSkeleton() != SkeletalMesh->GetSkeleton())
 				{
 					FScopedUnrealAllocs UEAllocs;
@@ -1815,7 +1830,7 @@ namespace UsdSkelSkeletonTranslatorImpl
 					AnimSequence->SetPreviewMesh(SkeletalMesh);
 
 					TUsdStore<pxr::VtArray<pxr::UsdSkelSkinningQuery>> SkinningTargets = SkeletonBinding.Get().GetSkinningTargets();
-					float LayerStartOffsetSeconds = 0.0f;
+					LayerStartOffsetSeconds = 0.0f;
 					const bool bSuccess = UsdToUnreal::ConvertSkelAnim(
 						SkeletonQuery.Get(),
 						&SkinningTargets.Get(),
@@ -1823,39 +1838,56 @@ namespace UsdSkelSkeletonTranslatorImpl
 						Context->bAllowInterpretingLODs,
 						RootMotionPrim,
 						AnimSequence,
-						&LayerStartOffsetSeconds
+						&LayerStartOffsetSeconds.GetValue()
 					);
 
 					if (bSuccess
 						&& (AnimSequence->GetDataModel()->GetNumBoneTracks() != 0 || AnimSequence->GetDataModel()->GetNumberOfFloatCurves() != 0))
 					{
-						if (UUsdAnimSequenceAssetUserData* UserData = UsdUtils::GetOrCreateAssetUserData<UUsdAnimSequenceAssetUserData>(AnimSequence))
-						{
-							UserData->PrimPaths.AddUnique(SkelAnimationPrimPath);
-							UserData->LayerStartOffsetSeconds = LayerStartOffsetSeconds;
-
-							if (Context->MetadataOptions.bCollectMetadata)
-							{
-								// Since we never collapse, we'll spawn assets components for any child prim that happens to be inside
-								// the skeleton itself, and the SkelAnimation type doesn't have any relevant "child prim" type (like for
-								// Mesh prims and UsdGeomSubsets), so we're probably safe in never collecting metadata from the SkelAnimation
-								// prim subtree
-								const bool bCollectFromEntireSubtrees = false;
-								UsdToUnreal::ConvertMetadata(
-									SkelAnimationPrim,
-									UserData,
-									Context->MetadataOptions.BlockedPrefixFilters,
-									Context->MetadataOptions.bInvertFilters,
-									bCollectFromEntireSubtrees
-								);
-							}
-						}
-
 						Context->AssetCache->CacheAsset(PrefixedSkelAnimHash, AnimSequence);
 					}
 					else
 					{
 						AnimSequence->MarkAsGarbage();
+						AnimSequence = nullptr;
+					}
+				}
+
+				if (UUsdAnimSequenceAssetUserData* UserData = UsdUtils::GetOrCreateAssetUserData<UUsdAnimSequenceAssetUserData>(AnimSequence))
+				{
+					UserData->PrimPaths.AddUnique(SkelAnimationPrimPath);
+
+					// It should be fine that we won't fetch/set this again in case we're reusing an AnimSequence from the asset cache
+					// because the LayerStartOffsetSeconds value is influenced by the animation start time codes, which should always
+					// affect the asset hash. This means that if the correct value for this were to change, we'd end up generating a
+					// new AnimSequence and computing it anyway
+					if (LayerStartOffsetSeconds.IsSet())
+					{
+						UserData->LayerStartOffsetSeconds = LayerStartOffsetSeconds.GetValue();
+					}
+
+					if (Context->MetadataOptions.bCollectMetadata)
+					{
+						// Since we never collapse, we'll spawn assets components for any child prim that happens to be inside
+						// the skeleton itself, and the SkelAnimation type doesn't have any relevant "child prim" type (like for
+						// Mesh prims and UsdGeomSubsets), so we're probably safe in never collecting metadata from the SkelAnimation
+						// prim subtree
+						const bool bCollectFromEntireSubtrees = false;
+						UsdToUnreal::ConvertMetadata(
+							SkelAnimationPrim,
+							UserData,
+							Context->MetadataOptions.BlockedPrefixFilters,
+							Context->MetadataOptions.bInvertFilters,
+							bCollectFromEntireSubtrees
+						);
+					}
+					else
+					{
+						// Strip the metadata from this prim, so that if we uncheck "Collect Metadata" it actually disappears on the
+						// AssetUserData
+						UserData->StageIdentifierToMetadata.Remove(
+							UsdToUnreal::ConvertString(SkelAnimationPrim.GetStage()->GetRootLayer()->GetIdentifier())
+						);
 					}
 				}
 
diff --git a/Engine/Plugins/Importers/USDImporter/Source/USDUtilities/Private/USDConversionUtils.cpp b/Engine/Plugins/Importers/USDImporter/Source/USDUtilities/Private/USDConversionUtils.cpp
--- a/Engine/Plugins/Importers/USDImporter/Source/USDUtilities/Private/USDConversionUtils.cpp
+++ b/Engine/Plugins/Importers/USDImporter/Source/USDUtilities/Private/USDConversionUtils.cpp
@@ -1816,7 +1816,7 @@ bool UsdUtils::SetAssetUserData(UObject* Object, UUsdAssetUserData* AssetUserDat
 
 	while (AssetUserDataInterface->HasAssetUserDataOfClass(UUsdAssetUserData::StaticClass()))
 	{
-		UE_LOG(LogUsd, Log, TEXT("Removing old AssetUserData from object '%s' before adding a new one"), *Object->GetPathName());
+		UE_LOG(LogUsd, Verbose, TEXT("Removing old AssetUserData from object '%s' before adding a new one"), *Object->GetPathName());
 		AssetUserDataInterface->RemoveUserDataOfClass(UUsdAssetUserData::StaticClass());
 	}
 
