USD: Implement new widget for adding/removing schemas from prims.

diff --git a/Engine/Plugins/Importers/USDImporter/Source/USDStageEditor/Private/SUSDStageTreeView.cpp b/Engine/Plugins/Importers/USDImporter/Source/USDStageEditor/Private/SUSDStageTreeView.cpp
--- a/Engine/Plugins/Importers/USDImporter/Source/USDStageEditor/Private/SUSDStageTreeView.cpp
+++ b/Engine/Plugins/Importers/USDImporter/Source/USDStageEditor/Private/SUSDStageTreeView.cpp
@@ -27,6 +27,8 @@
 #include "Widgets/Images/SImage.h"
 #include "Widgets/Input/SButton.h"
 #include "Widgets/Input/SCheckBox.h"
+#include "Widgets/Input/SComboBox.h"
+#include "Widgets/Input/SEditableTextBox.h"
 #include "Widgets/SToolTip.h"
 #include "Widgets/Text/SInlineEditableTextBlock.h"
 
@@ -36,6 +38,7 @@
 
 #include "USDIncludesStart.h"
 #include "pxr/usd/usdPhysics/tokens.h"
+#include "pxr/usd/usdSkel/tokens.h"
 #include "USDIncludesEnd.h"
 
 namespace UE::USDStageTreeView::Private
@@ -849,194 +852,21 @@ TSharedPtr<SWidget> SUsdStageTreeView::ConstructPrimContextMenu()
 	}
 	PrimOptions.EndSection();
 
-	PrimOptions.BeginSection("Integrations", LOCTEXT("Integrations", "Integrations"));
+	PrimOptions.BeginSection("Schemas", LOCTEXT("Schemas", "Schemas"));
 	{
-		PrimOptions.AddMenuEntry(
-			LOCTEXT("SetUpLiveLink", "Set up Live Link"),
-			LOCTEXT("SetUpLiveLink_ToolTip", "Sets up the generated component for Live Link and store the connection details to the USD Stage"),
-			FSlateIcon(),
-			FUIAction(
-				FExecuteAction::CreateSP(this, &SUsdStageTreeView::OnApplySchema, FName{*UsdToUnreal::ConvertToken(UnrealIdentifiers::LiveLinkAPI)}),
-				FCanExecuteAction::CreateSP(
-					this,
-					&SUsdStageTreeView::CanApplySchema,
-					FName{*UsdToUnreal::ConvertToken(UnrealIdentifiers::LiveLinkAPI)}
-				)
-			),
-			NAME_None,
-			EUserInterfaceActionType::Button
-		);
-
-		PrimOptions.AddMenuEntry(
-			LOCTEXT("RemoveLiveLink", "Remove Live Link"),
-			LOCTEXT(
-				"RemoveLiveLink_ToolTip",
-				"Reverses the Live Link configuration on the component and removes the connection details from the USD Stage"
-			),
-			FSlateIcon(),
-			FUIAction(
-				FExecuteAction::CreateSP(this, &SUsdStageTreeView::OnRemoveSchema, FName{*UsdToUnreal::ConvertToken(UnrealIdentifiers::LiveLinkAPI)}),
-				FCanExecuteAction::CreateSP(
-					this,
-					&SUsdStageTreeView::CanRemoveSchema,
-					FName{*UsdToUnreal::ConvertToken(UnrealIdentifiers::LiveLinkAPI)}
-				)
-			),
-			NAME_None,
-			EUserInterfaceActionType::Button
-		);
-
-		PrimOptions.AddMenuEntry(
-			LOCTEXT("SetUpControlRig", "Set up Control Rig"),
-			LOCTEXT(
-				"SetUpControlRig_ToolTip",
-				"Sets up the generated component for Control Rig integration and store the connection details to the USD Stage"
-			),
-			FSlateIcon(),
-			FUIAction(
-				FExecuteAction::CreateSP(
-					this,
-					&SUsdStageTreeView::OnApplySchema,
-					FName{*UsdToUnreal::ConvertToken(UnrealIdentifiers::ControlRigAPI)}
-				),
-				FCanExecuteAction::CreateSP(
-					this,
-					&SUsdStageTreeView::CanApplySchema,
-					FName{*UsdToUnreal::ConvertToken(UnrealIdentifiers::ControlRigAPI)}
-				)
-			),
-			NAME_None,
-			EUserInterfaceActionType::Button
-		);
-
-		PrimOptions.AddMenuEntry(
-			LOCTEXT("RemoveControlRig", "Remove Control Rig"),
-			LOCTEXT(
-				"RemoveControlRig_ToolTip",
-				"Reverses the Control Rig configuration on the component and removes the connection details from the USD Stage"
-			),
-			FSlateIcon(),
-			FUIAction(
-				FExecuteAction::CreateSP(
-					this,
-					&SUsdStageTreeView::OnRemoveSchema,
-					FName{*UsdToUnreal::ConvertToken(UnrealIdentifiers::ControlRigAPI)}
-				),
-				FCanExecuteAction::CreateSP(
-					this,
-					&SUsdStageTreeView::CanRemoveSchema,
-					FName{*UsdToUnreal::ConvertToken(UnrealIdentifiers::ControlRigAPI)}
-				)
-			),
-			NAME_None,
-			EUserInterfaceActionType::Button
-		);
-
-		PrimOptions.AddMenuEntry(
-			LOCTEXT("ApplyGroomSchema", "Apply Groom schema"),
-			LOCTEXT("ApplyGroomSchema_ToolTip", "Applies the Groom schema to interpret the prim and its children as a groom"),
-			FSlateIcon(),
-			FUIAction(
-				FExecuteAction::CreateSP(this, &SUsdStageTreeView::OnApplySchema, FName{*UsdToUnreal::ConvertToken(UnrealIdentifiers::GroomAPI)}),
-				FCanExecuteAction::CreateSP(this, &SUsdStageTreeView::CanApplySchema, FName{*UsdToUnreal::ConvertToken(UnrealIdentifiers::GroomAPI)})
-			),
-			NAME_None,
-			EUserInterfaceActionType::Button
-		);
-
-		PrimOptions.AddMenuEntry(
-			LOCTEXT("RemoveGroomSchema", "Remove Groom schema"),
-			LOCTEXT("RemoveGroomSchema_ToolTip", "Removes the Groom schema from the prim to stop interpreting it as a groom"),
-			FSlateIcon(),
-			FUIAction(
-				FExecuteAction::CreateSP(this, &SUsdStageTreeView::OnRemoveSchema, FName{*UsdToUnreal::ConvertToken(UnrealIdentifiers::GroomAPI)}),
-				FCanExecuteAction::CreateSP(this, &SUsdStageTreeView::CanRemoveSchema, FName{*UsdToUnreal::ConvertToken(UnrealIdentifiers::GroomAPI)})
-			),
-			NAME_None,
-			EUserInterfaceActionType::Button
-		);
-
-		PrimOptions.AddMenuEntry(
-			LOCTEXT("ApplyPhysicsCollisionchema", "Apply PhysicsCollision schema"),
-			LOCTEXT("ApplyPhysicsCollisionSchema_ToolTip", "Applies the PhysicsCollision schema to enable the prim as a collider"),
-			FSlateIcon(),
-			FUIAction(
-				FExecuteAction::CreateSP(
-					this,
-					&SUsdStageTreeView::OnApplySchema,
-					FName{*UsdToUnreal::ConvertToken(pxr::UsdPhysicsTokens->PhysicsCollisionAPI)}
-				),
-				FCanExecuteAction::CreateSP(
-					this,
-					&SUsdStageTreeView::CanApplySchema,
-					FName{*UsdToUnreal::ConvertToken(pxr::UsdPhysicsTokens->PhysicsCollisionAPI)}
-				)
-			),
-			NAME_None,
-			EUserInterfaceActionType::Button
-		);
-
-		PrimOptions.AddMenuEntry(
-			LOCTEXT("RemovePhysicsCollisionSchema", "Remove PhysicsCollision schema"),
-			LOCTEXT("RemovePhysicsCollisionSchema_ToolTip", "Removes the PhysicsCollision schema from the prim to disable it as a collider"),
-			FSlateIcon(),
-			FUIAction(
-				FExecuteAction::CreateSP(
-					this,
-					&SUsdStageTreeView::OnRemoveSchema,
-					FName{*UsdToUnreal::ConvertToken(pxr::UsdPhysicsTokens->PhysicsCollisionAPI)}
-				),
-				FCanExecuteAction::CreateSP(
-					this,
-					&SUsdStageTreeView::CanRemoveSchema,
-					FName{*UsdToUnreal::ConvertToken(pxr::UsdPhysicsTokens->PhysicsCollisionAPI)}
-				)
-			),
-			NAME_None,
-			EUserInterfaceActionType::Button
-		);
-
-		PrimOptions.AddMenuEntry(
-			LOCTEXT("ApplyPhysicsMeshCollisionchema", "Apply PhysicsMeshCollision schema"),
-			LOCTEXT(
-				"ApplyPhysicsMeshCollisionSchema_ToolTip",
-				"Applies the PhysicsMeshCollision schema to control how the mesh is made into a collider"
-			),
-			FSlateIcon(),
-			FUIAction(
-				FExecuteAction::CreateSP(
-					this,
-					&SUsdStageTreeView::OnApplySchema,
-					FName{*UsdToUnreal::ConvertToken(pxr::UsdPhysicsTokens->PhysicsMeshCollisionAPI)}
-				),
-				FCanExecuteAction::CreateSP(
-					this,
-					&SUsdStageTreeView::CanApplySchema,
-					FName{*UsdToUnreal::ConvertToken(pxr::UsdPhysicsTokens->PhysicsMeshCollisionAPI)}
-				)
-			),
-			NAME_None,
-			EUserInterfaceActionType::Button
+		const bool bInOpenSubMenuOnClick = false;
+		PrimOptions.AddSubMenu(
+			LOCTEXT("AddSchemaText", "Add schema..."),
+			FText::GetEmpty(),
+			FNewMenuDelegate::CreateSP(this, &SUsdStageTreeView::FillAddSchemaSubmenu),
+			bInOpenSubMenuOnClick
 		);
 
-		PrimOptions.AddMenuEntry(
-			LOCTEXT("RemovePhysicsMeshCollisionSchema", "Remove PhysicsMeshCollision schema"),
-			LOCTEXT("RemovePhysicsMeshCollisionSchema_ToolTip", "Removes the PhysicsMeshCollision schema from the prim"),
-			FSlateIcon(),
-			FUIAction(
-				FExecuteAction::CreateSP(
-					this,
-					&SUsdStageTreeView::OnRemoveSchema,
-					FName{*UsdToUnreal::ConvertToken(pxr::UsdPhysicsTokens->PhysicsMeshCollisionAPI)}
-				),
-				FCanExecuteAction::CreateSP(
-					this,
-					&SUsdStageTreeView::CanRemoveSchema,
-					FName{*UsdToUnreal::ConvertToken(pxr::UsdPhysicsTokens->PhysicsMeshCollisionAPI)}
-				)
-			),
-			NAME_None,
-			EUserInterfaceActionType::Button
+		PrimOptions.AddSubMenu(
+			LOCTEXT("RemoveSchemaText", "Remove schema..."),
+			FText::GetEmpty(),
+			FNewMenuDelegate::CreateSP(this, &SUsdStageTreeView::FillRemoveSchemaSubmenu),
+			bInOpenSubMenuOnClick
 		);
 	}
 	PrimOptions.EndSection();
@@ -1740,6 +1570,160 @@ void SUsdStageTreeView::FillDuplicateSubmenu(FMenuBuilder& MenuBuilder)
 	);
 }
 
+void SUsdStageTreeView::FillAddSchemaSubmenu(FMenuBuilder& MenuBuilder)
+{
+	static TArray<TSharedPtr<FString>> AllKnownSchemas = {
+		MakeShared<FString>(UsdToUnreal::ConvertToken(UnrealIdentifiers::ControlRigAPI)),
+		MakeShared<FString>(UsdToUnreal::ConvertToken(UnrealIdentifiers::GroomAPI)),
+		MakeShared<FString>(UsdToUnreal::ConvertToken(UnrealIdentifiers::GroomBindingAPI)),
+		MakeShared<FString>(UsdToUnreal::ConvertToken(UnrealIdentifiers::LiveLinkAPI)),
+		MakeShared<FString>(UsdToUnreal::ConvertToken(pxr::UsdShadeTokens->MaterialBindingAPI)),
+		MakeShared<FString>(UsdToUnreal::ConvertToken(pxr::UsdPhysicsTokens->PhysicsCollisionAPI)),
+		MakeShared<FString>(UsdToUnreal::ConvertToken(pxr::UsdPhysicsTokens->PhysicsMeshCollisionAPI)),
+		MakeShared<FString>(UsdToUnreal::ConvertToken(UnrealIdentifiers::SparseVolumeTextureAPI)),
+		MakeShared<FString>(UsdToUnreal::ConvertToken(pxr::UsdSkelTokens->SkelBindingAPI)),
+	};
+
+	// Only show on the list the schemas that can be applied to the selected prim
+	static TArray<TSharedPtr<FString>> AllowedKnownSchemas;
+	AllowedKnownSchemas.Reset(AllKnownSchemas.Num());
+	for (const TSharedPtr<FString>& KnownSchema : AllKnownSchemas)
+	{
+		if (CanApplySchema(**KnownSchema))
+		{
+			AllowedKnownSchemas.Add(KnownSchema);
+		}
+	}
+
+	static TSharedPtr<FString> CurrentKnownSchema = nullptr;
+	if (AllowedKnownSchemas.Num() > 0)
+	{
+		if (!CurrentKnownSchema || !CanApplySchema(**CurrentKnownSchema))
+		{
+			CurrentKnownSchema = AllowedKnownSchemas[0];
+		}
+	}
+	else
+	{
+		CurrentKnownSchema = nullptr;
+	}
+
+	static FText ManuallyInputText;
+
+	static TFunction<FText()> GetCurrentSchemaNameText = []() -> FText
+	{
+		return !ManuallyInputText.IsEmpty() ? ManuallyInputText : CurrentKnownSchema ? FText::FromString(*CurrentKnownSchema) : FText::GetEmpty();
+	};
+
+	// clang-format off
+	TSharedRef<SHorizontalBox> Box =
+		SNew(SHorizontalBox)
+
+		+SHorizontalBox::Slot()
+		.AutoWidth()
+		.Padding(FMargin(8.0f, 0.0f))
+		.VAlign(VAlign_Center)
+		.HAlign(HAlign_Left)
+		[
+			SNew(SComboBox<TSharedPtr<FString>>)
+			.OptionsSource(&AllowedKnownSchemas)
+			.OnGenerateWidget_Lambda([&](TSharedPtr<FString> Option)
+			{
+				TSharedPtr<SWidget> Widget = SNullWidget::NullWidget;
+				if (Option)
+				{
+					Widget = SNew(STextBlock)
+						.Text(FText::FromString(*Option))
+						.Font(FAppStyle::GetFontStyle("PropertyWindow.NormalFont"));
+				}
+
+				return Widget.ToSharedRef();
+			})
+			.OnSelectionChanged_Lambda([](TSharedPtr<FString> ChosenOption, ESelectInfo::Type SelectInfo)
+			{
+				CurrentKnownSchema = ChosenOption;
+				ManuallyInputText = FText::GetEmpty();
+			})
+			[
+				SNew(SEditableTextBox)
+				.Text_Lambda([]() -> FText
+				{
+					return GetCurrentSchemaNameText();
+				})
+				.Font(FAppStyle::GetFontStyle("PropertyWindow.NormalFont"))
+				.OnTextChanged_Lambda([](const FText& NewText)
+				{
+					ManuallyInputText = NewText;
+				})
+			]
+		]
+
+		+SHorizontalBox::Slot()
+		.AutoWidth()
+		.Padding(FMargin(8.0f, 0.0f))
+		.VAlign(VAlign_Center)
+		.HAlign(HAlign_Left)
+		[
+			SNew(SButton)
+			.Text(LOCTEXT("AddSchemaButtonText", "Add"))
+			.ToolTipText(LOCTEXT("AddSchemaButtonToolTip", "Adds the currently selected schema to the prim"))
+			.IsEnabled_Lambda([this]()->bool
+			{
+				FName SelectedSchema = *GetCurrentSchemaNameText().ToString();
+				return CanApplySchema(SelectedSchema);
+			})
+			.OnClicked_Lambda([this]() -> FReply
+			{
+				FName SelectedSchema = *GetCurrentSchemaNameText().ToString();
+				if (CanApplySchema(SelectedSchema))
+				{
+					OnApplySchema(SelectedSchema);
+				}
+
+				return FReply::Handled();
+			})
+			.ButtonStyle(&FAppStyle::Get(), "PrimaryButton")
+		];
+	// clang-format on
+
+	const bool bNoIndent = true;
+	MenuBuilder.AddWidget(Box, FText::GetEmpty(), bNoIndent);
+}
+
+void SUsdStageTreeView::FillRemoveSchemaSubmenu(FMenuBuilder& MenuBuilder)
+{
+	TSet<FString> RemovableSchemas;
+	TArray<FUsdPrimViewModelRef> MySelectedItems = GetSelectedItems();
+	for (FUsdPrimViewModelRef SelectedItem : MySelectedItems)
+	{
+		TArray<FName> Schemas = SelectedItem->UsdPrim.GetAppliedSchemas();
+		for (const FName& Schema : Schemas)
+		{
+			RemovableSchemas.Add(Schema.ToString());
+		}
+	}
+
+	TArray<FString> SortedRemovableSchemas = RemovableSchemas.Array();
+	SortedRemovableSchemas.Sort();
+
+	for (const FString& Schema : SortedRemovableSchemas)
+	{
+		if (!SUsdStageTreeView::CanRemoveSchema(*Schema))
+		{
+			continue;
+		}
+
+		MenuBuilder.AddMenuEntry(
+			FText::FromString(Schema),
+			FText::Format(LOCTEXT("RemoveSchemaToolTip", "Remove schema '{0}'"), FText::FromString(Schema)),
+			FSlateIcon(),
+			FUIAction(FExecuteAction::CreateSP(this, &SUsdStageTreeView::OnRemoveSchema, FName(*Schema))),
+			NAME_None,
+			EUserInterfaceActionType::Button
+		);
+	}
+}
+
 FReply SUsdStageTreeView::OnKeyDown(const FGeometry& MyGeometry, const FKeyEvent& InKeyEvent)
 {
 	if (UICommandList->ProcessCommandBindings(InKeyEvent))
diff --git a/Engine/Plugins/Importers/USDImporter/Source/USDStageEditor/Private/SUSDStageTreeView.h b/Engine/Plugins/Importers/USDImporter/Source/USDStageEditor/Private/SUSDStageTreeView.h
--- a/Engine/Plugins/Importers/USDImporter/Source/USDStageEditor/Private/SUSDStageTreeView.h
+++ b/Engine/Plugins/Importers/USDImporter/Source/USDStageEditor/Private/SUSDStageTreeView.h
@@ -64,6 +64,8 @@ private:
 	void OnToggleAllPayloads(EPayloadsTrigger PayloadsTrigger);
 
 	void FillDuplicateSubmenu(FMenuBuilder& MenuBuilder);
+	void FillAddSchemaSubmenu(FMenuBuilder& MenuBuilder);
+	void FillRemoveSchemaSubmenu(FMenuBuilder& MenuBuilder);
 
 	void OnAddChildPrim();
 	void OnCutPrim();
diff --git a/Engine/Plugins/Importers/USDImporter/Source/USDUtilities/Private/USDIntegrationUtils.cpp b/Engine/Plugins/Importers/USDImporter/Source/USDUtilities/Private/USDIntegrationUtils.cpp
--- a/Engine/Plugins/Importers/USDImporter/Source/USDUtilities/Private/USDIntegrationUtils.cpp
+++ b/Engine/Plugins/Importers/USDImporter/Source/USDUtilities/Private/USDIntegrationUtils.cpp
@@ -33,7 +33,14 @@ namespace UE::UsdIntegrationUtils::Private
 		FScopedUsdAllocs UsdAllocs;
 
 		pxr::TfType Schema = pxr::UsdSchemaRegistry::GetTypeFromSchemaTypeName(SchemaName);
-		if (!ensure(static_cast<bool>(Schema)))
+		if (!Schema)
+		{
+			return false;
+		}
+
+		// Query via this function first, because if we're given something like "Material", the following calls
+		// will emit USD errors when it tries checking if it can apply a non-single-apply API schema
+		if (!pxr::UsdSchemaRegistry::IsAppliedAPISchema(Schema))
 		{
 			return false;
 		}
@@ -119,7 +126,10 @@ bool UsdUtils::ApplySchema(const pxr::UsdPrim& Prim, const pxr::TfToken& SchemaT
 	FScopedUsdAllocs Allocs;
 
 	pxr::TfType Schema = pxr::UsdSchemaRegistry::GetTypeFromSchemaTypeName(SchemaToken);
-	ensure(static_cast<bool>(Schema));
+	if (!Schema)
+	{
+		return false;
+	}
 
 	return Prim.ApplyAPI(Schema);
 }
@@ -140,7 +150,10 @@ bool UsdUtils::RemoveSchema(const pxr::UsdPrim& Prim, const pxr::TfToken& Schema
 	FScopedUsdAllocs Allocs;
 
 	pxr::TfType Schema = pxr::UsdSchemaRegistry::GetTypeFromSchemaTypeName(SchemaToken);
-	ensure(static_cast<bool>(Schema));
+	if (!Schema)
+	{
+		return false;
+	}
 
 	return Prim.RemoveAPI(Schema);
 }
diff --git a/Engine/Plugins/Importers/USDImporter/Source/UnrealUSDWrapper/Private/UnrealUSDWrapper.cpp b/Engine/Plugins/Importers/USDImporter/Source/UnrealUSDWrapper/Private/UnrealUSDWrapper.cpp
--- a/Engine/Plugins/Importers/USDImporter/Source/UnrealUSDWrapper/Private/UnrealUSDWrapper.cpp
+++ b/Engine/Plugins/Importers/USDImporter/Source/UnrealUSDWrapper/Private/UnrealUSDWrapper.cpp
@@ -103,6 +103,15 @@ namespace UnrealIdentifiers
 	const TfToken UnrealUseFKControlRig = TfToken("unreal:controlRig:useFKControlRig");
 	const TfToken UnrealControlRigReduceKeys = TfToken("unreal:controlRig:reduceKeys");
 	const TfToken UnrealControlRigReductionTolerance = TfToken("unreal:controlRig:reductionTolerance");
+
+	const TfToken SparseVolumeTextureAPI = TfToken("SparseVolumeTextureAPI");
+	const TfToken UnrealSVTMappedFields = TfToken("unreal:SVT:mappedFields");
+	const TfToken UnrealSVTMappedMaterialParameters = TfToken("unreal:SVT:mappedMaterialParameters");
+	const TfToken UnrealSVTMappedGridComponents = TfToken("unreal:SVT:mappedGridComponents");
+	const TfToken UnrealSVTMappedAttributeChannels = TfToken("unreal:SVT:mappedAttributeChannels");
+	const TfToken UnrealSVTAttributesADataType = TfToken("unreal:SVT:attributesAdataType");
+	const TfToken UnrealSVTAttributesBDataType = TfToken("unreal:SVT:attributesBdataType");
+
 	const TfToken DiffuseColor = TfToken("diffuseColor");
 	const TfToken EmissiveColor = TfToken("emissiveColor");
 	const TfToken Metallic = TfToken("metallic");
diff --git a/Engine/Plugins/Importers/USDImporter/Source/UnrealUSDWrapper/Public/UnrealUSDWrapper.h b/Engine/Plugins/Importers/USDImporter/Source/UnrealUSDWrapper/Public/UnrealUSDWrapper.h
--- a/Engine/Plugins/Importers/USDImporter/Source/UnrealUSDWrapper/Public/UnrealUSDWrapper.h
+++ b/Engine/Plugins/Importers/USDImporter/Source/UnrealUSDWrapper/Public/UnrealUSDWrapper.h
@@ -382,6 +382,14 @@ namespace UnrealIdentifiers
 	extern UNREALUSDWRAPPER_API const pxr::TfToken UnrealControlRigReduceKeys;
 	extern UNREALUSDWRAPPER_API const pxr::TfToken UnrealControlRigReductionTolerance;
 
+	extern UNREALUSDWRAPPER_API const pxr::TfToken SparseVolumeTextureAPI;
+	extern UNREALUSDWRAPPER_API const pxr::TfToken UnrealSVTMappedFields;
+	extern UNREALUSDWRAPPER_API const pxr::TfToken UnrealSVTMappedMaterialParameters;
+	extern UNREALUSDWRAPPER_API const pxr::TfToken UnrealSVTMappedGridComponents;
+	extern UNREALUSDWRAPPER_API const pxr::TfToken UnrealSVTMappedAttributeChannels;
+	extern UNREALUSDWRAPPER_API const pxr::TfToken UnrealSVTAttributesADataType;
+	extern UNREALUSDWRAPPER_API const pxr::TfToken UnrealSVTAttributesBDataType;
+
 	extern UNREALUSDWRAPPER_API const pxr::TfToken DiffuseColor;
 	extern UNREALUSDWRAPPER_API const pxr::TfToken EmissiveColor;
 	extern UNREALUSDWRAPPER_API const pxr::TfToken Metallic;
