USD: Fix issues with USD file format listings on save/open/reference/payload dialogs.

Fix being able to pick a non-native USD file by clicking on the UsdStageActor root layer property;
Restrict references and payloads to USD-compatible formats;
Make it clear when a USD file will be binary or ASCII on all dialogs.

diff --git a/Engine/Plugins/Importers/USDImporter/Source/USDExporter/Private/AnimSequenceExporterUSD.cpp b/Engine/Plugins/Importers/USDImporter/Source/USDExporter/Private/AnimSequenceExporterUSD.cpp
--- a/Engine/Plugins/Importers/USDImporter/Source/USDExporter/Private/AnimSequenceExporterUSD.cpp
+++ b/Engine/Plugins/Importers/USDImporter/Source/USDExporter/Private/AnimSequenceExporterUSD.cpp
@@ -66,17 +66,7 @@ namespace UE::AnimSequenceExporterUSD::Private
 UAnimSequenceExporterUSD::UAnimSequenceExporterUSD()
 {
 #if USE_USD_SDK
-	for (const FString& Extension : UnrealUSDWrapper::GetNativeFileFormats())
-	{
-		// USDZ is not supported for writing for now
-		if (Extension.Equals(TEXT("usdz")))
-		{
-			continue;
-		}
-
-		FormatExtension.Add(Extension);
-		FormatDescription.Add(TEXT("Universal Scene Description file"));
-	}
+	UnrealUSDWrapper::AddUsdExportFileFormatDescriptions(FormatExtension, FormatDescription);
 	SupportedClass = UAnimSequence::StaticClass();
 	bText = false;
 #endif	  // #if USE_USD_SDK
diff --git a/Engine/Plugins/Importers/USDImporter/Source/USDExporter/Private/GeometryCacheExporterUSD.cpp b/Engine/Plugins/Importers/USDImporter/Source/USDExporter/Private/GeometryCacheExporterUSD.cpp
--- a/Engine/Plugins/Importers/USDImporter/Source/USDExporter/Private/GeometryCacheExporterUSD.cpp
+++ b/Engine/Plugins/Importers/USDImporter/Source/USDExporter/Private/GeometryCacheExporterUSD.cpp
@@ -63,17 +63,7 @@ namespace UE::GeometryCacheExporterUSD::Private
 UGeometryCacheExporterUSD::UGeometryCacheExporterUSD()
 {
 #if USE_USD_SDK
-	for (const FString& Extension : UnrealUSDWrapper::GetNativeFileFormats())
-	{
-		// USDZ is not supported for writing for now
-		if (Extension.Equals(TEXT("usdz")))
-		{
-			continue;
-		}
-
-		FormatExtension.Add(Extension);
-		FormatDescription.Add(TEXT("Universal Scene Description file"));
-	}
+	UnrealUSDWrapper::AddUsdExportFileFormatDescriptions(FormatExtension, FormatDescription);
 	SupportedClass = UGeometryCache::StaticClass();
 	bText = false;
 #endif	  // #if USE_USD_SDK
diff --git a/Engine/Plugins/Importers/USDImporter/Source/USDExporter/Private/LevelExporterUSD.cpp b/Engine/Plugins/Importers/USDImporter/Source/USDExporter/Private/LevelExporterUSD.cpp
--- a/Engine/Plugins/Importers/USDImporter/Source/USDExporter/Private/LevelExporterUSD.cpp
+++ b/Engine/Plugins/Importers/USDImporter/Source/USDExporter/Private/LevelExporterUSD.cpp
@@ -13,17 +13,7 @@
 ULevelExporterUSD::ULevelExporterUSD()
 {
 #if USE_USD_SDK
-	for (const FString& Extension : UnrealUSDWrapper::GetNativeFileFormats())
-	{
-		// USDZ is not supported for writing for now
-		if (Extension.Equals(TEXT("usdz")))
-		{
-			continue;
-		}
-
-		FormatExtension.Add(Extension);
-		FormatDescription.Add(TEXT("Universal Scene Description file"));
-	}
+	UnrealUSDWrapper::AddUsdExportFileFormatDescriptions(FormatExtension, FormatDescription);
 	SupportedClass = UWorld::StaticClass();
 	bText = false;
 #endif	  // #if USE_USD_SDK
diff --git a/Engine/Plugins/Importers/USDImporter/Source/USDExporter/Private/LevelSequenceExporterUSD.cpp b/Engine/Plugins/Importers/USDImporter/Source/USDExporter/Private/LevelSequenceExporterUSD.cpp
--- a/Engine/Plugins/Importers/USDImporter/Source/USDExporter/Private/LevelSequenceExporterUSD.cpp
+++ b/Engine/Plugins/Importers/USDImporter/Source/USDExporter/Private/LevelSequenceExporterUSD.cpp
@@ -1483,17 +1483,7 @@ namespace UE::LevelSequenceExporterUSD::Private
 ULevelSequenceExporterUsd::ULevelSequenceExporterUsd()
 {
 #if USE_USD_SDK
-	for (const FString& Extension : UnrealUSDWrapper::GetNativeFileFormats())
-	{
-		// USDZ is not supported for writing for now
-		if (Extension.Equals(TEXT("usdz")))
-		{
-			continue;
-		}
-
-		FormatExtension.Add(Extension);
-		FormatDescription.Add(TEXT("Universal Scene Description file"));
-	}
+	UnrealUSDWrapper::AddUsdExportFileFormatDescriptions(FormatExtension, FormatDescription);
 	SupportedClass = ULevelSequence::StaticClass();
 	bText = false;
 #endif	  // #if USE_USD_SDK
diff --git a/Engine/Plugins/Importers/USDImporter/Source/USDExporter/Private/MaterialExporterUSD.cpp b/Engine/Plugins/Importers/USDImporter/Source/USDExporter/Private/MaterialExporterUSD.cpp
--- a/Engine/Plugins/Importers/USDImporter/Source/USDExporter/Private/MaterialExporterUSD.cpp
+++ b/Engine/Plugins/Importers/USDImporter/Source/USDExporter/Private/MaterialExporterUSD.cpp
@@ -186,17 +186,7 @@ namespace UE::MaterialExporterUSD::Private
 UMaterialExporterUsd::UMaterialExporterUsd()
 {
 #if USE_USD_SDK
-	for (const FString& Extension : UnrealUSDWrapper::GetNativeFileFormats())
-	{
-		// USDZ is not supported for writing for now
-		if (Extension.Equals(TEXT("usdz")))
-		{
-			continue;
-		}
-
-		FormatExtension.Add(Extension);
-		FormatDescription.Add(TEXT("Universal Scene Description file"));
-	}
+	UnrealUSDWrapper::AddUsdExportFileFormatDescriptions(FormatExtension, FormatDescription);
 	SupportedClass = UMaterialInterface::StaticClass();
 	bText = false;
 #endif	  // #if USE_USD_SDK
diff --git a/Engine/Plugins/Importers/USDImporter/Source/USDExporter/Private/SkeletalMeshExporterUSD.cpp b/Engine/Plugins/Importers/USDImporter/Source/USDExporter/Private/SkeletalMeshExporterUSD.cpp
--- a/Engine/Plugins/Importers/USDImporter/Source/USDExporter/Private/SkeletalMeshExporterUSD.cpp
+++ b/Engine/Plugins/Importers/USDImporter/Source/USDExporter/Private/SkeletalMeshExporterUSD.cpp
@@ -89,17 +89,7 @@ namespace UE::SkeletalMeshExporterUSD::Private
 USkeletalMeshExporterUsd::USkeletalMeshExporterUsd()
 {
 #if USE_USD_SDK
-	for (const FString& Extension : UnrealUSDWrapper::GetNativeFileFormats())
-	{
-		// USDZ is not supported for writing for now
-		if (Extension.Equals(TEXT("usdz")))
-		{
-			continue;
-		}
-
-		FormatExtension.Add(Extension);
-		FormatDescription.Add(TEXT("Universal Scene Description file"));
-	}
+	UnrealUSDWrapper::AddUsdExportFileFormatDescriptions(FormatExtension, FormatDescription);
 	SupportedClass = USkeletalMesh::StaticClass();
 	bText = false;
 #endif	  // #if USE_USD_SDK
diff --git a/Engine/Plugins/Importers/USDImporter/Source/USDExporter/Private/StaticMeshExporterUSD.cpp b/Engine/Plugins/Importers/USDImporter/Source/USDExporter/Private/StaticMeshExporterUSD.cpp
--- a/Engine/Plugins/Importers/USDImporter/Source/USDExporter/Private/StaticMeshExporterUSD.cpp
+++ b/Engine/Plugins/Importers/USDImporter/Source/USDExporter/Private/StaticMeshExporterUSD.cpp
@@ -96,17 +96,7 @@ bool UStaticMeshExporterUsd::IsUsdAvailable()
 UStaticMeshExporterUsd::UStaticMeshExporterUsd()
 {
 #if USE_USD_SDK
-	for (const FString& Extension : UnrealUSDWrapper::GetNativeFileFormats())
-	{
-		// USDZ is not supported for writing for now
-		if (Extension.Equals(TEXT("usdz")))
-		{
-			continue;
-		}
-
-		FormatExtension.Add(Extension);
-		FormatDescription.Add(TEXT("Universal Scene Description file"));
-	}
+	UnrealUSDWrapper::AddUsdExportFileFormatDescriptions(FormatExtension, FormatDescription);
 	SupportedClass = UStaticMesh::StaticClass();
 	bText = false;
 #endif	  // #if USE_USD_SDK
diff --git a/Engine/Plugins/Importers/USDImporter/Source/USDStage/Private/USDStageActor.cpp b/Engine/Plugins/Importers/USDImporter/Source/USDStage/Private/USDStageActor.cpp
--- a/Engine/Plugins/Importers/USDImporter/Source/USDStage/Private/USDStageActor.cpp
+++ b/Engine/Plugins/Importers/USDImporter/Source/USDStage/Private/USDStageActor.cpp
@@ -1074,7 +1074,7 @@ AUsdStageActor::AUsdStageActor()
 			FProperty* Property = *PropertyIterator;
 			if (Property && Property->GetFName() == GET_MEMBER_NAME_CHECKED(AUsdStageActor, RootLayer))
 			{
-				TArray<FString> SupportedExtensions = UnrealUSDWrapper::GetAllSupportedFileFormats();
+				TArray<FString> SupportedExtensions = UnrealUSDWrapper::GetNativeFileFormats();
 				if (SupportedExtensions.Num() > 0)
 				{
 					FString JoinedExtensions = FString::Join(SupportedExtensions, TEXT("; *."));	// Combine "usd" and "usda" into "usd; *.usda"
diff --git a/Engine/Plugins/Importers/USDImporter/Source/USDStageImporter/Private/USDStageAssetImportFactory.cpp b/Engine/Plugins/Importers/USDImporter/Source/USDStageImporter/Private/USDStageAssetImportFactory.cpp
--- a/Engine/Plugins/Importers/USDImporter/Source/USDStageImporter/Private/USDStageAssetImportFactory.cpp
+++ b/Engine/Plugins/Importers/USDImporter/Source/USDStageImporter/Private/USDStageAssetImportFactory.cpp
@@ -33,10 +33,7 @@ UUsdStageAssetImportFactory::UUsdStageAssetImportFactory(const FObjectInitialize
 	bEditorImport = true;
 	bText = false;
 
-	for (const FString& Extension : UnrealUSDWrapper::GetNativeFileFormats())
-	{
-		Formats.Add(FString::Printf(TEXT("%s; Universal Scene Description files"), *Extension));
-	}
+	UnrealUSDWrapper::AddUsdImportFileFormatDescriptions(Formats);
 }
 
 bool UUsdStageAssetImportFactory::DoesSupportClass(UClass* Class)
diff --git a/Engine/Plugins/Importers/USDImporter/Source/USDStageImporter/Private/USDStageImportFactory.cpp b/Engine/Plugins/Importers/USDImporter/Source/USDStageImporter/Private/USDStageImportFactory.cpp
--- a/Engine/Plugins/Importers/USDImporter/Source/USDStageImporter/Private/USDStageImportFactory.cpp
+++ b/Engine/Plugins/Importers/USDImporter/Source/USDStageImporter/Private/USDStageImportFactory.cpp
@@ -24,10 +24,7 @@ UUsdStageImportFactory::UUsdStageImportFactory(const FObjectInitializer& ObjectI
 	bEditorImport = true;
 	bText = false;
 
-	for (const FString& Extension : UnrealUSDWrapper::GetNativeFileFormats())
-	{
-		Formats.Add(FString::Printf(TEXT("%s; Universal Scene Description files"), *Extension));
-	}
+	UnrealUSDWrapper::AddUsdImportFileFormatDescriptions(Formats);
 }
 
 UObject* UUsdStageImportFactory::FactoryCreateFile(
diff --git a/Engine/Plugins/Importers/USDImporter/Source/USDUtilities/Private/USDLayerUtils.cpp b/Engine/Plugins/Importers/USDImporter/Source/USDUtilities/Private/USDLayerUtils.cpp
--- a/Engine/Plugins/Importers/USDImporter/Source/USDUtilities/Private/USDLayerUtils.cpp
+++ b/Engine/Plugins/Importers/USDImporter/Source/USDUtilities/Private/USDLayerUtils.cpp
@@ -198,13 +198,29 @@ TOptional<FString> UsdUtils::BrowseUsdFile(EBrowseFileMode Mode)
 
 	TArray<FString> OutFiles;
 
+	TArray<FString> NativeTextExtensions;
+	TArray<FString> NativePossiblyBinaryExtensions;
+	UnrealUSDWrapper::GetNativeFileFormats(NativeTextExtensions, NativePossiblyBinaryExtensions);
+
+	TSet<FString> NativeTextExtensionsSet{NativeTextExtensions};
+	TSet<FString> NativePossiblyBinaryExtensionsSet{NativePossiblyBinaryExtensions};
+
 	// When browsing files for the purposes of opening a stage or saving layers,
 	// we offer the native USD file formats as options. Browsing files in order
 	// to use them as the targets of composition arcs (e.g. sublayers,
 	// references, payloads, etc.) also offers any plugin file formats that are
 	// registered.
-	TArray<FString> SupportedExtensions = (Mode == EBrowseFileMode::Composition) ? UnrealUSDWrapper::GetAllSupportedFileFormats()
-																				 : UnrealUSDWrapper::GetNativeFileFormats();
+	TArray<FString> SupportedExtensions;
+	if (Mode == EBrowseFileMode::Composition)
+	{
+		SupportedExtensions = UnrealUSDWrapper::GetAllSupportedFileFormats();
+	}
+	else
+	{
+		SupportedExtensions = NativeTextExtensions;
+		SupportedExtensions.Append(NativePossiblyBinaryExtensions);
+	}
+
 	if (SupportedExtensions.Num() == 0)
 	{
 		UE_LOG(LogUsd, Error, TEXT("No file extensions supported by the USD SDK!"));
@@ -221,9 +237,19 @@ TOptional<FString> UsdUtils::BrowseUsdFile(EBrowseFileMode Mode)
 	FString FileTypes = FString::Printf(TEXT("Universal Scene Description files (*.%s)|*.%s|"), *JoinedExtensions, *JoinedExtensions);
 	for (const FString& SupportedExtension : SupportedExtensions)
 	{
+		const bool bIsTextNative = NativeTextExtensionsSet.Contains(SupportedExtension);
+		const bool bIsBinaryNative = NativePossiblyBinaryExtensionsSet.Contains(SupportedExtension);
+
 		// The '(*.%s)' on the actual name (before the '|') is not optional: We need the name part to be different for each format
 		// or else the options will overwrite each other on the Mac
-		FileTypes += FString::Printf(TEXT("Universal Scene Description file (*.%s)|*.%s|"), *SupportedExtension, *SupportedExtension);
+		FileTypes += FString::Printf(
+			TEXT("Universal Scene Description %sfile (*.%s)|*.%s|"),
+			bIsTextNative	  ? TEXT("text ")
+			: bIsBinaryNative ? TEXT("binary ")
+							  : TEXT(""),
+			*SupportedExtension,
+			*SupportedExtension
+		);
 	}
 
 	switch (Mode)
diff --git a/Engine/Plugins/Importers/USDImporter/Source/UnrealUSDWrapper/Private/UnrealUSDWrapper.cpp b/Engine/Plugins/Importers/USDImporter/Source/UnrealUSDWrapper/Private/UnrealUSDWrapper.cpp
--- a/Engine/Plugins/Importers/USDImporter/Source/UnrealUSDWrapper/Private/UnrealUSDWrapper.cpp
+++ b/Engine/Plugins/Importers/USDImporter/Source/UnrealUSDWrapper/Private/UnrealUSDWrapper.cpp
@@ -6,6 +6,7 @@
 #include "USDLog.h"
 #include "USDMemory.h"
 #include "USDProjectSettings.h"
+#include "USDReferenceOptions.h"
 
 #include "UsdWrappers/SdfLayer.h"
 #include "UsdWrappers/UsdAttribute.h"
@@ -612,8 +613,7 @@ TArray<FString> UnrealUSDWrapper::GetNativeFileFormats()
 		pxr::UsdUsdFileFormatTokens->Id,
 		pxr::UsdUsdaFileFormatTokens->Id,
 		pxr::UsdUsdcFileFormatTokens->Id,
-		pxr::UsdUsdzFileFormatTokens->Id
-	};
+		pxr::UsdUsdzFileFormatTokens->Id};
 
 	std::set<std::string> FileExtensions;
 
@@ -640,6 +640,102 @@ TArray<FString> UnrealUSDWrapper::GetNativeFileFormats()
 	return Result;
 }
 
+void UnrealUSDWrapper::GetNativeFileFormats(TArray<FString>& OutTextFormats, TArray<FString>& OutPossiblyBinaryFormats)
+{
+#if USE_USD_SDK
+	FScopedUsdAllocs Allocs;
+
+	std::set<std::string> TextExtensions;
+	const pxr::SdfFileFormatConstPtr TextFormat = pxr::SdfFileFormat::FindById(pxr::UsdUsdaFileFormatTokens->Id);
+	if (TextFormat)
+	{
+		for (const std::string& FileExtension : TextFormat->GetFileExtensions())
+		{
+			TextExtensions.insert(FileExtension);
+		}
+	}
+
+	OutTextFormats.Reset(TextExtensions.size());
+	for (const std::string& Extension : TextExtensions)
+	{
+		OutTextFormats.Emplace(ANSI_TO_TCHAR(Extension.c_str()));
+	}
+
+	const static std::vector<pxr::TfToken> BinaryFormatIds{
+		pxr::UsdUsdFileFormatTokens->Id,
+		pxr::UsdUsdcFileFormatTokens->Id,
+		pxr::UsdUsdzFileFormatTokens->Id};
+
+	std::set<std::string> PossiblyBinaryExtensions;
+	for (const pxr::TfToken& FormatId : BinaryFormatIds)
+	{
+		const pxr::SdfFileFormatConstPtr Format = pxr::SdfFileFormat::FindById(FormatId);
+		if (!Format)
+		{
+			continue;
+		}
+
+		for (const std::string& FileExtension : Format->GetFileExtensions())
+		{
+			PossiblyBinaryExtensions.insert(FileExtension);
+		}
+	}
+
+	OutPossiblyBinaryFormats.Reset(PossiblyBinaryExtensions.size());
+	for (const std::string& Extension : PossiblyBinaryExtensions)
+	{
+		OutPossiblyBinaryFormats.Emplace(ANSI_TO_TCHAR(Extension.c_str()));
+	}
+#endif	  // #if USE_USD_SDK
+}
+
+void UnrealUSDWrapper::AddUsdExportFileFormatDescriptions(TArray<FString>& OutFormatExtensions, TArray<FString>& OutFormatDescriptions)
+{
+	TArray<FString> TextExtensions;
+	TArray<FString> PossiblyBinaryExtensions;
+	UnrealUSDWrapper::GetNativeFileFormats(TextExtensions, PossiblyBinaryExtensions);
+
+	for (const FString& Extension : PossiblyBinaryExtensions)
+	{
+		// USDZ is not supported for writing for now
+		if (Extension.Equals(TEXT("usdz")))
+		{
+			continue;
+		}
+
+		OutFormatExtensions.Add(Extension);
+		OutFormatDescriptions.Add(TEXT("Universal Scene Description binary file"));
+	}
+
+	for (const FString& Extension : TextExtensions)
+	{
+		OutFormatExtensions.Add(Extension);
+		OutFormatDescriptions.Add(TEXT("Universal Scene Description text file"));
+	}
+}
+
+void UnrealUSDWrapper::AddUsdImportFileFormatDescriptions(TArray<FString>& OutFormats)
+{
+	TArray<FString> TextExtensions;
+	TArray<FString> PossiblyBinaryExtensions;
+	UnrealUSDWrapper::GetNativeFileFormats(TextExtensions, PossiblyBinaryExtensions);
+
+	// Ideally we'd generate a combined entry that allows all USD file types at once, but it seems that
+	// GenerateFactoryFileExtensions and InternalGetFormatInfo on ObjectTools.cpp are really hard-coded
+	// to only accept a single extension at a time... I guess the user can still just pick *.* to see
+	// all file types at once
+
+	// Add isolated entries for each supported file type
+	for (const FString& Extension : PossiblyBinaryExtensions)
+	{
+		OutFormats.Add(FString::Printf(TEXT("%s; Universal Scene Description binary files"), *Extension));
+	}
+	for (const FString& Extension : TextExtensions)
+	{
+		OutFormats.Add(FString::Printf(TEXT("%s; Universal Scene Description text files"), *Extension));
+	}
+}
+
 namespace UE::UnrealUSDWrapper::Private
 {
 	UE::FUsdStage OpenStageImpl(
@@ -1074,6 +1170,27 @@ public:
 
 		FUsdMemoryManager::Initialize();
 		UnrealUSDWrapper::SetupDiagnosticDelegate();
+
+#if WITH_EDITOR
+		// Update the supported filetype filters for reference/payload picker dialogs
+		for (TFieldIterator<FProperty> PropertyIterator(UUsdReferenceOptions::StaticClass()); PropertyIterator; ++PropertyIterator)
+		{
+			FProperty* Property = *PropertyIterator;
+			if (Property && Property->GetFName() == GET_MEMBER_NAME_CHECKED(UUsdReferenceOptions, TargetFile))
+			{
+				TArray<FString> SupportedExtensions = UnrealUSDWrapper::GetAllSupportedFileFormats();
+				if (SupportedExtensions.Num() > 0)
+				{
+					FString JoinedExtensions = FString::Join(SupportedExtensions, TEXT("; *."));	// Combine "usd" and "usda" into "usd; *.usda"
+					Property->SetMetaData(
+						TEXT("FilePathFilter"),
+						FString::Printf(TEXT("Universal Scene Description files|*.%s"), *JoinedExtensions, *JoinedExtensions)
+					);
+				}
+				break;
+			}
+		}
+#endif	  // WITH_EDITOR
 	}
 
 	virtual void ShutdownModule() override
diff --git a/Engine/Plugins/Importers/USDImporter/Source/UnrealUSDWrapper/Public/UnrealUSDWrapper.h b/Engine/Plugins/Importers/USDImporter/Source/UnrealUSDWrapper/Public/UnrealUSDWrapper.h
--- a/Engine/Plugins/Importers/USDImporter/Source/UnrealUSDWrapper/Public/UnrealUSDWrapper.h
+++ b/Engine/Plugins/Importers/USDImporter/Source/UnrealUSDWrapper/Public/UnrealUSDWrapper.h
@@ -245,6 +245,19 @@ public:
 	 * These are the extensions for the file formats built into USD (i.e. "usd", "usda", "usdc", and "usdz").
 	 */
 	UNREALUSDWRAPPER_API static TArray<FString> GetNativeFileFormats();
+	UNREALUSDWRAPPER_API static void GetNativeFileFormats(TArray<FString>& OutTextFormats, TArray<FString>& OutPossiblyBinaryFormats);
+
+	/**
+	 * Internally used by all the export factories, this will query the supported export USD file formats and fill out
+	 * OutFormatExtensions with entries like "usd", "usda", "usdc", etc., and OutFormatDescriptions with entries like
+	 * "Universal Scene Description binary file" or "Universal Scene Description text file"
+	 */
+	UNREALUSDWRAPPER_API static void AddUsdExportFileFormatDescriptions(TArray<FString>& OutFormatExtensions, TArray<FString>& OutFormatDescriptions);
+	/**
+	 * Internally used by all the import factories, this will query the supported export USD file formats and fill out
+	 * OutFormats with entries like "usda; Universal Scene Description text files", "usdc; Universal Scene Description binary files", etc.
+	 */
+	UNREALUSDWRAPPER_API static void AddUsdImportFileFormatDescriptions(TArray<FString>& OutFormats);
 
 	/**
 	 * Opens a USD stage from a file on disk or existing layers, with a population mask or not.
