PR #11195: Make USD Unreal Shader support saving and loading material assets with unicode name

diff --git a/Engine/Plugins/Importers/USDImporter/Source/USDUtilities/Private/USDTypesConversion.cpp b/Engine/Plugins/Importers/USDImporter/Source/USDUtilities/Private/USDTypesConversion.cpp
--- a/Engine/Plugins/Importers/USDImporter/Source/USDUtilities/Private/USDTypesConversion.cpp
+++ b/Engine/Plugins/Importers/USDImporter/Source/USDUtilities/Private/USDTypesConversion.cpp
@@ -83,12 +83,12 @@ namespace UsdToUnreal
 {
 	FString ConvertString(const std::string& InString)
 	{
-		return FString(ANSI_TO_TCHAR(InString.c_str()));
+		return FString(UTF8_TO_TCHAR(InString.c_str()));
 	}
 
 	FString ConvertString(const char* InString)
 	{
-		return FString(ANSI_TO_TCHAR(InString));
+		return FString(UTF8_TO_TCHAR(InString));
 	}
 
 	FString ConvertPath(const pxr::SdfPath& Path)
@@ -300,7 +300,7 @@ namespace UnrealToUsd
 {
 	TUsdStore<std::string> ConvertString(const TCHAR* InString)
 	{
-		return MakeUsdStore<std::string>(TCHAR_TO_ANSI(InString));
+		return MakeUsdStore<std::string>(TCHAR_TO_UTF8(InString));
 	}
 
 	TUsdStore<pxr::SdfPath> ConvertPath(const TCHAR* InString)
