USD: Fix crashes caused by generating static meshes with no valid render data.

diff --git a/Engine/Plugins/Importers/USDImporter/Source/USDClasses/Private/USDAssetCache2.cpp b/Engine/Plugins/Importers/USDImporter/Source/USDClasses/Private/USDAssetCache2.cpp
--- a/Engine/Plugins/Importers/USDImporter/Source/USDClasses/Private/USDAssetCache2.cpp
+++ b/Engine/Plugins/Importers/USDImporter/Source/USDClasses/Private/USDAssetCache2.cpp
@@ -1022,7 +1022,7 @@ bool UUsdAssetCache2::AddAssetReference(const UObject* Asset, const UObject* Ref
 
 bool UUsdAssetCache2::RemoveAssetReference(const UObject* Asset, const UObject* Referencer)
 {
-	if (!Asset || !Referencer)
+	if (!Asset)
 	{
 		return false;
 	}
@@ -1032,12 +1032,27 @@ bool UUsdAssetCache2::RemoveAssetReference(const UObject* Asset, const UObject*
 	bool bRemovedSomething = false;
 	if (FCachedAssetInfo* Info = LRUCache.Find(Asset))
 	{
-		const int32 NumRemoved = Info->Referencers.Remove(FObjectKey{Referencer});
-		if (NumRemoved > 0)
+		if (Referencer)
 		{
-			UE_LOG(LogUsd, Verbose, TEXT("Removed referencer '%s' from asset '%s'"), *Referencer->GetPathName(), *Asset->GetPathName());
+			const int32 NumRemoved = Info->Referencers.Remove(FObjectKey{Referencer});
+			if (NumRemoved > 0)
+			{
+				UE_LOG(LogUsd, Verbose, TEXT("Removed referencer '%s' from asset '%s'"), *Referencer->GetPathName(), *Asset->GetPathName());
 
-			bRemovedSomething = true;
+				bRemovedSomething = true;
+			}
+		}
+		else
+		{
+			const int32 NumRemoved = Info->Referencers.Num();
+			Info->Referencers.Reset();
+
+			if (NumRemoved > 0)
+			{
+				UE_LOG(LogUsd, Verbose, TEXT("Removed all '%d' referencers for asset '%s'"), NumRemoved, *Asset->GetPathName());
+
+				bRemovedSomething = true;
+			}
 		}
 	}
 
diff --git a/Engine/Plugins/Importers/USDImporter/Source/USDClasses/Public/USDAssetCache2.h b/Engine/Plugins/Importers/USDImporter/Source/USDClasses/Public/USDAssetCache2.h
--- a/Engine/Plugins/Importers/USDImporter/Source/USDClasses/Public/USDAssetCache2.h
+++ b/Engine/Plugins/Importers/USDImporter/Source/USDClasses/Public/USDAssetCache2.h
@@ -88,9 +88,10 @@ public:
 
 	/**
 	 * Removes an UObject referencer from a particular asset, returning true if the operation succeeded.
+	 * If no specific Referencer is provided, all referencers to Asset will be removed.
 	 */
 	UFUNCTION(BlueprintCallable, Category = "Caching", meta = (CallInEditor = "true"))
-	bool RemoveAssetReference(const UObject* Asset, const UObject* Referencer);
+	bool RemoveAssetReference(const UObject* Asset, const UObject* Referencer = nullptr);
 
 	/**
 	 * Removes the particular referencer to all assets tracked by the cache, if it was a referencer to any of them.
diff --git a/Engine/Plugins/Importers/USDImporter/Source/USDSchemas/Private/USDGeomMeshTranslator.cpp b/Engine/Plugins/Importers/USDImporter/Source/USDSchemas/Private/USDGeomMeshTranslator.cpp
--- a/Engine/Plugins/Importers/USDImporter/Source/USDSchemas/Private/USDGeomMeshTranslator.cpp
+++ b/Engine/Plugins/Importers/USDImporter/Source/USDSchemas/Private/USDGeomMeshTranslator.cpp
@@ -584,6 +584,10 @@ namespace UsdGeomMeshTranslatorImpl
 		// the worst of both worlds as some of these will be arbitrarily recomputed anyway, and some will be left invalid
 		EComputeNTBsFlags Options = GSkipMeshTangentComputation ? EComputeNTBsFlags::None
 																: EComputeNTBsFlags::UseMikkTSpace | EComputeNTBsFlags::Tangents;
+
+		// Repairing can take a long time for degenerate triangles (UE-194839)
+		Options |= EComputeNTBsFlags::IgnoreDegenerateTriangles;
+
 		if (InvalidNormalFraction >= GMeshNormalRepairThreshold)
 		{
 			Options |= EComputeNTBsFlags::Normals;
@@ -701,6 +705,7 @@ namespace UsdGeomMeshTranslatorImpl
 				SourceModel.BuildSettings.bRecomputeNormals = false;
 				SourceModel.BuildSettings.bRecomputeTangents = false;
 				SourceModel.BuildSettings.bBuildReversedIndexBuffer = false;
+				SourceModel.BuildSettings.bRemoveDegenerates = true;	// Note: This may get rid of the entire mesh if it is all invalid
 
 				FMeshDescription* StaticMeshDescription = StaticMesh->CreateMeshDescription(LODIndex);
 				check(StaticMeshDescription);
@@ -1315,7 +1320,7 @@ namespace UE::UsdCollision::Private
 	// Approximation types from PhysicsMeshCollisionAPI and UE-specific approximations
 	enum class EUsdCollisionType : uint8
 	{
-		None, // no approximation so equivalent to CTF_UseComplexAsSimple
+		None,	 // no approximation so equivalent to CTF_UseComplexAsSimple
 		ConvexDecomposition,
 		ConvexHull,
 		Sphere,
@@ -1795,10 +1800,45 @@ void FBuildStaticMeshTaskChain::SetupTasks()
 				FeatureLevel = World->GetFeatureLevel();
 			}
 
-			if (!UsdGeomMeshTranslatorImpl::BuildStaticMesh(*StaticMesh, FeatureLevel, LODIndexToMeshDescription))
+			bool bSuccess = UsdGeomMeshTranslatorImpl::BuildStaticMesh(*StaticMesh, FeatureLevel, LODIndexToMeshDescription);
+			if (bSuccess)
+			{
+				FStaticMeshRenderData* RenderData = StaticMesh->GetRenderData();
+
+				// We want to discard our StaticMesh if it generated no valid RenderData (e.g. was compeletely degenerate).
+				// They are not really going to be useful, and this is also handy as in some cases the StaticMeshComponents that use them
+				// may get confused as to why they have a StaticMesh with no valid RenderData at points where they expected to have them
+				// (e.g. see UE-201946)
+				bSuccess = RenderData && !RenderData->Bounds.ContainsNaN() && RenderData->GetCurrentFirstLODIdx(0) != INDEX_NONE;
+			}
+
+			// Build failed, abandon mesh
+			if (!bSuccess)
 			{
-				// Build failed, discard the mesh
-				StaticMesh = nullptr;
+				Context->InfoCache->RemoveAllAssetPrimLinks(StaticMesh);
+
+				FString Hash = Context->AssetCache->GetHashForAsset(StaticMesh);
+				if (!Hash.IsEmpty())
+				{
+					Context->AssetCache->RemoveAssetReference(StaticMesh);
+					UObject* RemovedAsset = Context->AssetCache->RemoveAsset(Hash);
+					if (ensure(RemovedAsset))
+					{
+						const TCHAR* NewName = nullptr;
+						UObject* NewOuter = GetTransientPackage();
+						RemovedAsset->Rename(NewName, NewOuter);
+
+						StaticMesh = nullptr;
+
+						UE_LOG(
+							LogUsd,
+							Warning,
+							TEXT("Discarding StaticMesh generated for prim '%s' as it didn't produce any valid RenderData (likely all triangles "
+								 "were degenerate)"),
+							*PrimPath.GetString()
+						);
+					}
+				}
 
 				return false;
 			}
diff --git a/Engine/Plugins/Importers/USDImporter/Source/USDSchemas/Private/USDInfoCache.cpp b/Engine/Plugins/Importers/USDImporter/Source/USDSchemas/Private/USDInfoCache.cpp
--- a/Engine/Plugins/Importers/USDImporter/Source/USDSchemas/Private/USDInfoCache.cpp
+++ b/Engine/Plugins/Importers/USDImporter/Source/USDSchemas/Private/USDInfoCache.cpp
@@ -1581,6 +1581,30 @@ void FUsdInfoCache::LinkAssetToPrim(const UE::FSdfPath& Path, UObject* Asset)
 	ImplPtr->AssetToPrimPaths.FindOrAdd(Asset).AddUnique(Path);
 }
 
+void FUsdInfoCache::UnlinkAssetFromPrim(const UE::FSdfPath& Path, UObject* Asset)
+{
+	FUsdInfoCacheImpl* ImplPtr = Impl.Get();
+	if (!ImplPtr)
+	{
+		return;
+	}
+	FWriteScopeLock ScopeLock(ImplPtr->PrimPathToAssetsLock);
+
+	UE_LOG(LogUsd, Verbose, TEXT("Unlinking asset '%s' to prim '%s'"),
+		*Asset->GetPathName(),
+		*Path.GetString()
+	);
+
+	if (TArray<TWeakObjectPtr<UObject>>* FoundAssetsForPrim = ImplPtr->PrimPathToAssets.Find(Path))
+	{
+		FoundAssetsForPrim->Remove(Asset);
+	}
+	if (TArray<UE::FSdfPath>* FoundPrimPathsForAsset = ImplPtr->AssetToPrimPaths.Find(Asset))
+	{
+		FoundPrimPathsForAsset->Remove(Path);
+	}
+}
+
 TArray<TWeakObjectPtr<UObject>> FUsdInfoCache::RemoveAllAssetPrimLinks(const UE::FSdfPath& Path)
 {
 	FUsdInfoCacheImpl* ImplPtr = Impl.Get();
@@ -1606,6 +1630,31 @@ TArray<TWeakObjectPtr<UObject>> FUsdInfoCache::RemoveAllAssetPrimLinks(const UE:
 	return Assets;
 }
 
+TArray<UE::FSdfPath> FUsdInfoCache::RemoveAllAssetPrimLinks(const UObject* Asset)
+{
+	FUsdInfoCacheImpl* ImplPtr = Impl.Get();
+	if (!ImplPtr)
+	{
+		return {};
+	}
+	FWriteScopeLock ScopeLock(ImplPtr->PrimPathToAssetsLock);
+
+	UE_LOG(LogUsd, Verbose, TEXT("Removing asset prim links for asset '%s'"), *Asset->GetPathName());
+
+	TArray<UE::FSdfPath> PrimPaths;
+	ImplPtr->AssetToPrimPaths.RemoveAndCopyValue(const_cast<UObject*>(Asset), PrimPaths);
+
+	for (const UE::FSdfPath& Path : PrimPaths)
+	{
+		if (TArray<TWeakObjectPtr<UObject>>* Assets = ImplPtr->PrimPathToAssets.Find(Path))
+		{
+			Assets->Remove(const_cast<UObject*>(Asset));
+		}
+	}
+
+	return PrimPaths;
+}
+
 void FUsdInfoCache::RemoveAllAssetPrimLinks()
 {
 	FUsdInfoCacheImpl* ImplPtr = Impl.Get();
diff --git a/Engine/Plugins/Importers/USDImporter/Source/USDSchemas/Public/USDInfoCache.h b/Engine/Plugins/Importers/USDImporter/Source/USDSchemas/Public/USDInfoCache.h
--- a/Engine/Plugins/Importers/USDImporter/Source/USDSchemas/Public/USDInfoCache.h
+++ b/Engine/Plugins/Importers/USDImporter/Source/USDSchemas/Public/USDInfoCache.h
@@ -89,8 +89,10 @@ public:
 
 public:
 	void LinkAssetToPrim(const UE::FSdfPath& Path, UObject* Asset);
+	void UnlinkAssetFromPrim(const UE::FSdfPath& Path, UObject* Asset);
 
 	TArray<TWeakObjectPtr<UObject>> RemoveAllAssetPrimLinks(const UE::FSdfPath& Path);
+	TArray<UE::FSdfPath> RemoveAllAssetPrimLinks(const UObject* Asset);
 	void RemoveAllAssetPrimLinks();
 
 	TArray<TWeakObjectPtr<UObject>> GetAllAssetsForPrim(const UE::FSdfPath& Path) const;
