USD: Fix crash when exporting layers to paths that are invalid for the USD SDK.

diff --git a/Engine/Plugins/Importers/USDImporter/Source/USDStageEditor/Private/SUSDLayersTreeView.cpp b/Engine/Plugins/Importers/USDImporter/Source/USDStageEditor/Private/SUSDLayersTreeView.cpp
--- a/Engine/Plugins/Importers/USDImporter/Source/USDStageEditor/Private/SUSDLayersTreeView.cpp
+++ b/Engine/Plugins/Importers/USDImporter/Source/USDStageEditor/Private/SUSDLayersTreeView.cpp
@@ -43,6 +43,12 @@ namespace UE::USDLayersTreeViewImpl::Private
 
 		// Clone the layer so that we don't modify the currently opened stage when we do the remapping below
 		UE::FSdfLayer OutputLayer = UE::FSdfLayer::CreateNew(*TargetPath);
+		if (!OutputLayer)
+		{
+			UE_LOG(LogUsd, Warning, TEXT("Failed to export USD layer to path '%s'"), *TargetPath);
+			return;
+		}
+
 		OutputLayer.TransferContent(LayerToExport);
 
 		// Update references to assets (e.g. textures) so that they're absolute and also work from the new file
diff --git a/Engine/Plugins/Importers/USDImporter/Source/USDUtilities/Private/USDLayerUtils.cpp b/Engine/Plugins/Importers/USDImporter/Source/USDUtilities/Private/USDLayerUtils.cpp
--- a/Engine/Plugins/Importers/USDImporter/Source/USDUtilities/Private/USDLayerUtils.cpp
+++ b/Engine/Plugins/Importers/USDImporter/Source/USDUtilities/Private/USDLayerUtils.cpp
@@ -821,6 +821,11 @@ bool UsdUtils::IsSessionLayerWithinStage(const pxr::SdfLayerRefPtr& Layer, const
 
 void UsdUtils::ConvertAssetRelativePathsToAbsolute(UE::FSdfLayer& LayerToConvert, const UE::FSdfLayer& AnchorLayer)
 {
+	if (!LayerToConvert || !AnchorLayer)
+	{
+		return;
+	}
+
 	FScopedUsdAllocs Allocs;
 
 	pxr::SdfLayerRefPtr UsdLayer{LayerToConvert};
