USD: Fix File->ExportAll silently outputting layers to Engine/Binaries.

diff --git a/Engine/Plugins/Importers/USDImporter/Source/USDStageEditor/Private/SUSDStage.cpp b/Engine/Plugins/Importers/USDImporter/Source/USDStageEditor/Private/SUSDStage.cpp
--- a/Engine/Plugins/Importers/USDImporter/Source/USDStageEditor/Private/SUSDStage.cpp
+++ b/Engine/Plugins/Importers/USDImporter/Source/USDStageEditor/Private/SUSDStage.cpp
@@ -2276,7 +2276,7 @@ void SUsdStage::FileExportAllLayers(const FString& OutputDirectory)
 		FString LayerPath = Layer.GetRealPath();
 		FPaths::NormalizeFilename(LayerPath);
 
-		FString TargetPath = FPaths::Combine(OutputDirectory, Layer.GetDisplayName());
+		FString TargetPath = FPaths::Combine(OutputDirectoryCopy, Layer.GetDisplayName());
 		FPaths::NormalizeFilename(TargetPath);
 
 		// Filename collision (should be rare, but possible given that we're discarding the folder structure)
