USD: Add EngineTests plugin to help with USD automated testing;

Deprecate old testing functions that were exposed to users;
Update automated tests to use new plugin;
Use new function to fix test_draw_modes.py failure related to skeletal mesh animations.

diff --git a/Engine/Plugins/Importers/USDImporter/Source/USDTests/Public/USDTestsBlueprintLibrary.h b/Engine/Plugins/Importers/USDImporter/Source/USDTests/Public/USDTestsBlueprintLibrary.h
--- a/Engine/Plugins/Importers/USDImporter/Source/USDTests/Public/USDTestsBlueprintLibrary.h
+++ b/Engine/Plugins/Importers/USDImporter/Source/USDTests/Public/USDTestsBlueprintLibrary.h
@@ -21,6 +21,7 @@ public:
 	 * Helps test the effects of blueprint recompilation on the spawned actors and assets when a stage is opened.
 	 * Returns whether it compiled successfully or not.
 	 */
+	UE_DEPRECATED(5.4, "These functions were meant for internal testing only and will be removed in a future release.")
 	UFUNCTION(BlueprintCallable, Category = "Blueprint")
 	static bool RecompileBlueprintStageActor(AUsdStageActor* BlueprintDerivedStageActor);
 
@@ -29,6 +30,7 @@ public:
 	 * This is useful for testing how the stage actor behaves when going into PIE with a dirty blueprint, as that usually triggers
 	 * a recompile at the very sensitive PIE transition
 	 */
+	UE_DEPRECATED(5.4, "These functions were meant for internal testing only and will be removed in a future release.")
 	UFUNCTION(BlueprintCallable, Category = "Blueprint")
 	static void DirtyStageActorBlueprint(AUsdStageActor* BlueprintDerivedStageActor);
 
@@ -36,6 +38,7 @@ public:
 	 * Queries a subtree vertex count using the stage actor's info cache, which is not yet exposed to blueprint.
 	 * May return -1 in case of an error.
 	 */
+	UE_DEPRECATED(5.4, "These functions were meant for internal testing only and will be removed in a future release.")
 	UFUNCTION(BlueprintCallable, Category = "Counts")
 	static int64 GetSubtreeVertexCount(AUsdStageActor* StageActor, const FString& PrimPath);
 
@@ -43,18 +46,21 @@ public:
 	 * Queries a subtree material slot count using the stage actor's info cache, which is not yet exposed to blueprint.
 	 * May return -1 in case of an error.
 	 */
+	UE_DEPRECATED(5.4, "These functions were meant for internal testing only and will be removed in a future release.")
 	UFUNCTION(BlueprintCallable, Category = "Counts")
 	static int64 GetSubtreeMaterialSlotCount(AUsdStageActor* StageActor, const FString& PrimPath);
 
 	/**
 	 * Internally opens the provided stage and sets it on the stage actor via C++
 	 */
+	UE_DEPRECATED(5.4, "These functions were meant for internal testing only and will be removed in a future release.")
 	UFUNCTION(BlueprintCallable, Category = "Counts")
 	static void SetUsdStageCpp(AUsdStageActor* StageActor, const FString& NewStageRootLayer);
 
 	/**
 	 * Clears the editor transaction history (useful to be run before GC, to check for otherwise unreferenced assets)
 	 */
+	UE_DEPRECATED(5.4, "These functions were meant for internal testing only and will be removed in a future release.")
 	UFUNCTION(BlueprintCallable, Category = "Transactions")
 	static void ClearTransactionHistory();
 };
