Control Rig: Fix compilation error
#rnx

diff --git a/Engine/Plugins/Importers/USDImporter/Source/USDStage/Private/USDLevelSequenceHelper.cpp b/Engine/Plugins/Importers/USDImporter/Source/USDStage/Private/USDLevelSequenceHelper.cpp
--- a/Engine/Plugins/Importers/USDImporter/Source/USDStage/Private/USDLevelSequenceHelper.cpp
+++ b/Engine/Plugins/Importers/USDImporter/Source/USDStage/Private/USDLevelSequenceHelper.cpp
@@ -387,6 +387,7 @@ namespace UsdLevelSequenceHelperImpl
 				SkeletalMeshComp,
 				bReduceKeys,
 				Tolerance,
+				true,
 				ControlRigSectionStartFrame,
 				EMovieSceneKeyInterpolation::SmartAuto
 			);
