USD: Fixed ensure when prim is missing a material assignment.
Added cvars for some static mesh build and Nanite settings.

diff --git a/Engine/Plugins/Importers/USDImporter/Source/USDSchemas/Private/USDGeomMeshTranslator.cpp b/Engine/Plugins/Importers/USDImporter/Source/USDSchemas/Private/USDGeomMeshTranslator.cpp
--- a/Engine/Plugins/Importers/USDImporter/Source/USDSchemas/Private/USDGeomMeshTranslator.cpp
+++ b/Engine/Plugins/Importers/USDImporter/Source/USDSchemas/Private/USDGeomMeshTranslator.cpp
@@ -73,6 +73,41 @@ static FAutoConsoleVariableRef CVarSkipMeshTangentComputation(
 	TEXT("Skip computing tangents for meshes. With meshes with a huge numer of vertices, it can take a very long time to compute them.")
 );
 
+static bool GBuildReversedIndexBuffer = false;
+static FAutoConsoleVariableRef CVarBuildReversedIndexBuffer(
+	TEXT("USD.StaticMesh.BuildSettings.BuildReversedIndexBuffer"),
+	GBuildReversedIndexBuffer,
+	TEXT("Enable to optimize mesh in mirrored transform. Double index buffer size.")
+);
+
+static bool GUseFullPrecisionUVs = false;
+static FAutoConsoleVariableRef CVarUseFullPrecisionUVs(
+	TEXT("USD.StaticMesh.BuildSettings.UseFullPrecisionUVs"),
+	GUseFullPrecisionUVs,
+	TEXT("If true, UVs will be stored at full floating point precision.")
+);
+
+static bool GUseHighPrecisionTangentBasis = false;
+static FAutoConsoleVariableRef CVarUseHighPrecisionTangentBasis(
+	TEXT("USD.StaticMesh.BuildSettings.UseHighPrecisionTangentBasis"),
+	GUseHighPrecisionTangentBasis,
+	TEXT("If true, Tangents will be stored at 16 bit vs 8 bit precision.")
+);
+
+static int32 GNaniteSettingsNormalPrecision = -1;
+static FAutoConsoleVariableRef CVarNaniteSettingsNormalPrecision(
+	TEXT("USD.Nanite.Settings.NormalPrecision"),
+	GNaniteSettingsNormalPrecision,
+	TEXT("Normal Precision in bits. -1 is auto.")
+);
+
+static int32 GNaniteSettingsTangentPrecision = -1;
+static FAutoConsoleVariableRef CVarNaniteSettingsTangentPrecision(
+	TEXT("USD.Nanite.Settings.TangentPrecision"),
+	GNaniteSettingsTangentPrecision,
+	TEXT("Tangent Precision in bits. -1 is auto.")
+);
+
 namespace UsdGeomMeshTranslatorImpl
 {
 	bool ShouldEnableNanite(
@@ -704,7 +739,9 @@ namespace UsdGeomMeshTranslatorImpl
 				SourceModel.BuildSettings.bGenerateLightmapUVs = false;
 				SourceModel.BuildSettings.bRecomputeNormals = false;
 				SourceModel.BuildSettings.bRecomputeTangents = false;
-				SourceModel.BuildSettings.bBuildReversedIndexBuffer = false;
+				SourceModel.BuildSettings.bBuildReversedIndexBuffer = GBuildReversedIndexBuffer;
+				SourceModel.BuildSettings.bUseFullPrecisionUVs = GUseFullPrecisionUVs;
+				SourceModel.BuildSettings.bUseHighPrecisionTangentBasis = GUseHighPrecisionTangentBasis;
 				SourceModel.BuildSettings.bRemoveDegenerates = true;	// Note: This may get rid of the entire mesh if it is all invalid
 
 				FMeshDescription* StaticMeshDescription = StaticMesh->CreateMeshDescription(LODIndex);
@@ -1676,6 +1713,8 @@ void FBuildStaticMeshTaskChain::SetupTasks()
 
 #if WITH_EDITOR
 			   StaticMesh->NaniteSettings.bEnabled = bShouldEnableNanite;
+			   StaticMesh->NaniteSettings.NormalPrecision = GNaniteSettingsNormalPrecision;
+			   StaticMesh->NaniteSettings.TangentPrecision = GNaniteSettingsTangentPrecision;
 #endif	  // WITH_EDITOR
 
 			   if (UUsdMeshAssetUserData* UserData = UsdUtils::GetOrCreateAssetUserData<UUsdMeshAssetUserData>(StaticMesh))
diff --git a/Engine/Plugins/Importers/USDImporter/Source/USDUtilities/Private/USDGeomMeshConversion.cpp b/Engine/Plugins/Importers/USDImporter/Source/USDUtilities/Private/USDGeomMeshConversion.cpp
--- a/Engine/Plugins/Importers/USDImporter/Source/USDUtilities/Private/USDGeomMeshConversion.cpp
+++ b/Engine/Plugins/Importers/USDImporter/Source/USDUtilities/Private/USDGeomMeshConversion.cpp
@@ -3497,7 +3497,7 @@ UsdUtils::FUsdPrimMaterialAssignmentInfo UsdUtils::GetPrimMaterialAssignments(
 	// If we have any slot without an actual material assignment yet, copy over the material assignment from
 	// the "main" slot, or fallback to displayColor. This is how we have unspecified faces or geomsubsets without
 	// assignmets "fallback" to using the main material assignment
-	if (Result.Slots.Num() > 1)
+	if (Result.Slots.Num() >= 1)
 	{
 		FString FallbackMaterialSource;
 		EPrimAssignmentType FallbackAssignmentType = EPrimAssignmentType::None;
@@ -5032,8 +5032,24 @@ void UsdUtils::HashGeomMeshPrim(const UE::FUsdStage& Stage, const FString& PrimP
 		return;
 	}
 
-	HashArrayAttribute<GfVec3f>(InOutHashState, UsdMesh.GetPointsAttr(), TimeCode);
-	HashArrayAttribute<GfVec3f>(InOutHashState, UsdMesh.GetNormalsAttr(), TimeCode);
+	if (pxr::UsdGeomPrimvar PointsPrimvar = pxr::UsdGeomPrimvar(UsdPrim.GetAttribute(UnrealIdentifiers::PrimvarsPoints)))
+	{
+		HashArrayAttribute<GfVec3f>(InOutHashState, UsdPrim.GetAttribute(UnrealIdentifiers::PrimvarsPoints), TimeCode);
+	}
+	else
+	{
+		HashArrayAttribute<GfVec3f>(InOutHashState, UsdMesh.GetPointsAttr(), TimeCode);
+	}
+
+	if (pxr::UsdGeomPrimvar NormalsPrimvar = pxr::UsdGeomPrimvar(UsdPrim.GetAttribute(UnrealIdentifiers::PrimvarsNormals)))
+	{
+		HashArrayAttribute<GfVec3f>(InOutHashState, UsdPrim.GetAttribute(UnrealIdentifiers::PrimvarsNormals), TimeCode);
+	}
+	else
+	{
+		HashArrayAttribute<GfVec3f>(InOutHashState, UsdMesh.GetNormalsAttr(), TimeCode);
+	}
+
 	HashArrayPrimvar<GfVec3f>(InOutHashState, UsdMesh.GetDisplayColorPrimvar(), TimeCode);
 	HashArrayPrimvar<float>(InOutHashState, UsdMesh.GetDisplayOpacityPrimvar(), TimeCode);
 
