USD: Prevent layers from being released while being traversed within UsdUtils::ReplaceUnrealMaterialsWithBaked.

diff --git a/Engine/Plugins/Importers/USDImporter/Source/USDUtilities/Private/USDGeomMeshConversion.cpp b/Engine/Plugins/Importers/USDImporter/Source/USDUtilities/Private/USDGeomMeshConversion.cpp
--- a/Engine/Plugins/Importers/USDImporter/Source/USDUtilities/Private/USDGeomMeshConversion.cpp
+++ b/Engine/Plugins/Importers/USDImporter/Source/USDUtilities/Private/USDGeomMeshConversion.cpp
@@ -4631,10 +4631,6 @@ void UsdUtils::ReplaceUnrealMaterialsWithBaked(
 
 	pxr::UsdStageRefPtr UsdStage{Stage};
 
-	// UsedLayers here instead of layer stack because we may be exporting using payloads, and payload layers
-	// don't show up on the layer stack list but do show up on the UsedLayers list
-	std::vector<pxr::SdfLayerHandle> LayersToTraverse = UsdStage->GetUsedLayers();
-
 	// Recursively traverses the stage, doing the material assignment replacements.
 	// This handles Mesh prims as well as GeomSubset prims.
 	// Note how we receive the stage as an argument instead of capturing it from the outer scope:
@@ -4647,7 +4643,7 @@ void UsdUtils::ReplaceUnrealMaterialsWithBaked(
 	)>
 		TraverseForMaterialReplacement;
 	TraverseForMaterialReplacement =
-		[&TraverseForMaterialReplacement, &LayersToTraverse, &LayerToAuthorIn, &BakedMaterials, bIsAssetLayer, bUsePayload, &StageMatScope](
+		[&TraverseForMaterialReplacement, &UsdStage, &LayerToAuthorIn, &BakedMaterials, bIsAssetLayer, bUsePayload, &StageMatScope](
 			pxr::UsdStageRefPtr StageToTraverse,
 			pxr::UsdPrim Prim,
 			TOptional<FMaterialScopePrim>& MatPrimScope,
@@ -4787,8 +4783,13 @@ void UsdUtils::ReplaceUnrealMaterialsWithBaked(
 		}
 
 		// We always want to replace things in whatever layer they were authored, and not just override with
-		// a stronger opinion, so search through all sublayers to find the ones with the specs we are targeting
-		for (pxr::SdfLayerHandle Layer : LayersToTraverse)
+		// a stronger opinion, so search through all sublayers to find the ones with the specs we are targeting.
+		// UsedLayers here instead of layer stack because we may be exporting using payloads, and payload layers
+		// don't show up on the layer stack list but do show up on the UsedLayers list.
+		// We fetch these layers every time because variant switching may cause referenced layers to be dropped,
+		// in case they were only used by prims inside a particular variant. This means we can also discover new
+		// layers as we switch into other layers, so we really need to call this every time.
+		for (pxr::SdfLayerHandle Layer : UsdStage->GetUsedLayers())
 		{
 			pxr::SdfAttributeSpecHandle UnrealMaterialAttrSpec = Layer->GetAttributeAtPath(UnrealMaterialAttrPath);
 			pxr::SdfPrimSpecHandle UnrealMaterialPrimSpec = Layer->GetPrimAtPath(UnrealMaterialPrimPath);
