USD: Load UnrealUSDWrapper if a static function is called before the module is loaded.

diff --git a/Engine/Plugins/Importers/USDImporter/Source/UnrealUSDWrapper/Private/UnrealUSDWrapper.cpp b/Engine/Plugins/Importers/USDImporter/Source/UnrealUSDWrapper/Private/UnrealUSDWrapper.cpp
--- a/Engine/Plugins/Importers/USDImporter/Source/UnrealUSDWrapper/Private/UnrealUSDWrapper.cpp
+++ b/Engine/Plugins/Importers/USDImporter/Source/UnrealUSDWrapper/Private/UnrealUSDWrapper.cpp
@@ -547,6 +547,32 @@ double UnrealUSDWrapper::GetDefaultTimeCode()
 }
 #endif	  // USE_USD_SDK
 
+namespace UE::UnrealUSDWrapper::Private
+{
+	void EnsureModuleIsLoaded()
+	{
+		FModuleManager& Manager = FModuleManager::Get();
+
+		const FName ModuleName = TEXT("UnrealUSDWrapper");
+		if (!Manager.IsModuleLoaded(ModuleName))
+		{
+			UE_LOG(
+				LogUsd,
+				Warning,
+				TEXT(
+					"Attempted to call a static function from UnrealUSDWrapper before the module is actually loaded! The module will be loaded now, but not all static functions have this check. In general, please ensure an Unreal module is loaded before calling any of its static functions, for example by calling 'FModuleManager::LoadModuleChecked<IUnrealUSDWrapperModule>(\"UnrealUSDWrapper\");' beforehand."
+				)
+			);
+
+			IModuleInterface* LoadedModule = Manager.LoadModule(ModuleName);
+			if (!LoadedModule)
+			{
+				UE_LOG(LogUsd, Warning, TEXT("Failed to load the UnrealUSDWrapper module!"));
+			}
+		}
+	}
+};
+
 TArray<FString> UnrealUSDWrapper::RegisterPlugins(const FString& PathToPlugInfo)
 {
 	return RegisterPlugins(TArray<FString>{PathToPlugInfo});
@@ -554,6 +580,8 @@ TArray<FString> UnrealUSDWrapper::RegisterPlugins(const FString& PathToPlugInfo)
 
 TArray<FString> UnrealUSDWrapper::RegisterPlugins(const TArray<FString>& PathsToPlugInfo)
 {
+	UE::UnrealUSDWrapper::Private::EnsureModuleIsLoaded();
+
 	TArray<FString> PluginNames;
 
 #if USE_USD_SDK
@@ -580,6 +608,8 @@ TArray<FString> UnrealUSDWrapper::RegisterPlugins(const TArray<FString>& PathsTo
 
 TArray<FString> UnrealUSDWrapper::GetAllSupportedFileFormats()
 {
+	UE::UnrealUSDWrapper::Private::EnsureModuleIsLoaded();
+
 	TArray<FString> Result;
 
 #if USE_USD_SDK
@@ -604,6 +634,8 @@ TArray<FString> UnrealUSDWrapper::GetAllSupportedFileFormats()
 
 TArray<FString> UnrealUSDWrapper::GetNativeFileFormats()
 {
+	UE::UnrealUSDWrapper::Private::EnsureModuleIsLoaded();
+
 	TArray<FString> Result;
 
 #if USE_USD_SDK
@@ -642,6 +674,8 @@ TArray<FString> UnrealUSDWrapper::GetNativeFileFormats()
 
 void UnrealUSDWrapper::GetNativeFileFormats(TArray<FString>& OutTextFormats, TArray<FString>& OutPossiblyBinaryFormats)
 {
+	UE::UnrealUSDWrapper::Private::EnsureModuleIsLoaded();
+
 #if USE_USD_SDK
 	FScopedUsdAllocs Allocs;
 
@@ -691,6 +725,8 @@ void UnrealUSDWrapper::GetNativeFileFormats(TArray<FString>& OutTextFormats, TAr
 
 void UnrealUSDWrapper::AddUsdExportFileFormatDescriptions(TArray<FString>& OutFormatExtensions, TArray<FString>& OutFormatDescriptions)
 {
+	UE::UnrealUSDWrapper::Private::EnsureModuleIsLoaded();
+
 	TArray<FString> TextExtensions;
 	TArray<FString> PossiblyBinaryExtensions;
 	UnrealUSDWrapper::GetNativeFileFormats(TextExtensions, PossiblyBinaryExtensions);
@@ -716,6 +752,8 @@ void UnrealUSDWrapper::AddUsdExportFileFormatDescriptions(TArray<FString>& OutFo
 
 void UnrealUSDWrapper::AddUsdImportFileFormatDescriptions(TArray<FString>& OutFormats)
 {
+	UE::UnrealUSDWrapper::Private::EnsureModuleIsLoaded();
+
 	TArray<FString> TextExtensions;
 	TArray<FString> PossiblyBinaryExtensions;
 	UnrealUSDWrapper::GetNativeFileFormats(TextExtensions, PossiblyBinaryExtensions);
diff --git a/Engine/Plugins/Importers/USDImporter/USDImporter.uplugin b/Engine/Plugins/Importers/USDImporter/USDImporter.uplugin
--- a/Engine/Plugins/Importers/USDImporter/USDImporter.uplugin
+++ b/Engine/Plugins/Importers/USDImporter/USDImporter.uplugin
@@ -18,7 +18,7 @@
 		{
 			"Name": "UnrealUSDWrapper",
 			"Type": "Runtime",
-			"LoadingPhase": "PreDefault",
+			"LoadingPhase": "Default",
 			"PlatformAllowList": [
 				"Win64", "Mac", "Linux"
 			]
