USD: Fix USD skeletal mesh issues due to the refactor at cl 30711371.

diff --git a/Engine/Plugins/Importers/USDImporter/Source/USDUtilities/Private/USDSkeletalDataConversion.cpp b/Engine/Plugins/Importers/USDImporter/Source/USDUtilities/Private/USDSkeletalDataConversion.cpp
--- a/Engine/Plugins/Importers/USDImporter/Source/USDUtilities/Private/USDSkeletalDataConversion.cpp
+++ b/Engine/Plugins/Importers/USDImporter/Source/USDUtilities/Private/USDSkeletalDataConversion.cpp
@@ -32,6 +32,7 @@
 #include "Evaluation/MovieSceneSequenceTransform.h"
 #include "IMovieScenePlayer.h"
 #include "Interfaces/ITargetPlatformManagerModule.h"
+#include "MaterialDomain.h"
 #include "Misc/CoreMisc.h"
 #include "Misc/MemStack.h"
 #include "Modules/ModuleManager.h"
@@ -1759,8 +1760,6 @@ bool UsdToUnreal::ConvertSkinnedMesh(
 		Opacities.Add(1.0f);
 	}
 
-	SkelMeshImportData.NumTexCoords = 0;
-
 	bool bReverseOrder = IUsdPrim::GetGeometryOrientation(UsdMesh) == EUsdGeomOrientation::LeftHanded;
 
 	struct FUVSet
@@ -1834,6 +1833,11 @@ bool UsdToUnreal::ConvertSkinnedMesh(
 					if (UVSet.UVs.size() > 0)
 					{
 						UVSets.Add(MoveTemp(UVSet));
+
+						if (UVSets.Num() == MAX_TEXCOORDS)
+						{
+							break;
+						}
 					}
 				}
 			}
@@ -1844,6 +1848,11 @@ bool UsdToUnreal::ConvertSkinnedMesh(
 					if (UVSet.UVs.size() > 0)
 					{
 						UVSets.Add(MoveTemp(UVSet));
+
+						if (UVSets.Num() == MAX_TEXCOORDS)
+						{
+							break;
+						}
 					}
 				}
 			}
@@ -1856,6 +1865,8 @@ bool UsdToUnreal::ConvertSkinnedMesh(
 		++UVChannelIndex;
 	}
 
+	SkelMeshImportData.NumTexCoords = FMath::Clamp(FMath::Max(SkelMeshImportData.NumTexCoords, (uint32)UVSets.Num()), 0, MAX_TEXCOORDS);
+
 	SkelMeshImportData.Wedges.Reserve((NumExistingFaces + NumFaces) * 6);
 
 	uint32 NumProcessedFaceVertexIndices = 0;
@@ -1878,6 +1889,16 @@ bool UsdToUnreal::ConvertSkinnedMesh(
 		int32 RealMaterialIndex = LocalToCombinedMaterialIndex[LocalMaterialIndex];
 		SkelMeshImportData.MaxMaterialIndex = FMath::Max<uint32>(SkelMeshImportData.MaxMaterialIndex, RealMaterialIndex);
 
+		// The SkelMeshImportData now requires that the Materials array has number of entries that matches
+		// the max material index
+		SkelMeshImportData.Materials.SetNum(SkelMeshImportData.MaxMaterialIndex + 1);
+		for (int32 Index = 0; Index < SkelMeshImportData.Materials.Num(); ++Index)
+		{
+			SkeletalMeshImportData::FMaterial& Material = SkelMeshImportData.Materials[Index];
+			Material.MaterialImportName = LexToString(Index);
+			Material.Material = UMaterial::GetDefaultMaterial(MD_Surface);
+		}
+
 		// SkeletalMeshImportData uses triangle faces so quads will have to be split into triangles
 		const bool bIsQuad = (NumOriginalFaceVertices == 4);
 		const uint32 NumTriangles = bIsQuad ? 2 : 1;
@@ -2800,6 +2821,11 @@ USkeletalMesh* UsdToUnreal::GetSkeletalMeshFromImportData(
 	for (int32 LODIndex = 0; LODIndex < LODIndexToSkeletalMeshImportData.Num(); ++LODIndex)
 	{
 		FSkeletalMeshImportData& LODImportData = LODIndexToSkeletalMeshImportData[LODIndex];
+
+		// In the future it will be expected for bone data to be inside FSkeletalMeshImportData as well so we should
+		// probably do this
+		LODImportData.RefBonesBinary = InSkeletonBones;
+
 		ImportedResource->LODModels.Add(new FSkeletalMeshLODModel());
 		FSkeletalMeshLODModel& LODModel = ImportedResource->LODModels.Last();
 
@@ -2876,6 +2902,37 @@ USkeletalMesh* UsdToUnreal::GetSkeletalMeshFromImportData(
 			return nullptr;
 		}
 
+		// We must also provide the ImportData with morph target information now.
+		// Reference: FSkeletalMeshImportData::AddMorphTarget
+		// We don't use that function directly as matching the interface would involve copying our BlendShape.Vertices
+		// into a new FMorphTargetLODModel.
+		// TODO: Add in the morph target normal data (from FMorphTargetDelta::TangentZDelta) to the import data
+		// at the right location when FSkeletalMeshImportData::GetMeshDescription starts reading normal data from that location
+		LODImportData.MorphTargets.Reserve(InBlendShapesByPath.Num());
+		LODImportData.MorphTargetNames.Reserve(InBlendShapesByPath.Num());
+		LODImportData.MorphTargetModifiedPoints.Reserve(InBlendShapesByPath.Num());
+		for (const TPair<FString, UsdUtils::FUsdBlendShape>& Pair : InBlendShapesByPath)
+		{
+			const UsdUtils::FUsdBlendShape& BlendShape = Pair.Value;
+			if (!BlendShape.LODIndicesThatUseThis.Contains(LODIndex))
+			{
+				continue;
+			}
+
+			LODImportData.MorphTargetNames.Add(BlendShape.Name);
+			FSkeletalMeshImportData& MorphTarget = LODImportData.MorphTargets.Emplace_GetRef();
+			TSet<uint32>& NewModifiedPoints = LODImportData.MorphTargetModifiedPoints.Emplace_GetRef();
+
+			MorphTarget.Points.Reserve(LODImportData.Points.Num());
+			NewModifiedPoints.Reserve(BlendShape.Vertices.Num());
+
+			for (const FMorphTargetDelta& Delta : BlendShape.Vertices)
+			{
+				NewModifiedPoints.Add(Delta.SourceIdx);
+				MorphTarget.Points.Add(Delta.PositionDelta + LODImportData.Points[Delta.SourceIdx]);
+			}
+		}
+
 		// This is important because it will fill in the LODModel's RawSkeletalMeshBulkDataID,
 		// which is the part of the skeletal mesh's DDC key that is affected by the actual mesh data
 		PRAGMA_DISABLE_DEPRECATION_WARNINGS
