USD: Allow isolating session layers.

diff --git a/Engine/Plugins/Importers/USDImporter/Source/USDStage/Private/USDStageActor.cpp b/Engine/Plugins/Importers/USDImporter/Source/USDStage/Private/USDStageActor.cpp
--- a/Engine/Plugins/Importers/USDImporter/Source/USDStage/Private/USDStageActor.cpp
+++ b/Engine/Plugins/Importers/USDImporter/Source/USDStage/Private/USDStageActor.cpp
@@ -1333,20 +1333,36 @@ void AUsdStageActor::IsolateLayer(const UE::FSdfLayer& Layer)
 
 	if (UsdStage)
 	{
-		// Check if the layer we're trying to isolate is part of the stage. To do that we need to reopen a fresh copy
-		// of the current stage in case we have any layers muted, as muted layers don't show up on the layer stack.
-		// Note that we can't just check the list of muted layers either, as it's possible to mute *any* layer for
-		// a given stage, not only the layers that are currently used by it.
-		// We'll use an empty population mask though (which should prevent prim composition) and just use the layers
-		// that are already opened on the current stage anyway, so this should be cheap
-		UE::FUsdStage FreshCurrentStage = UnrealUSDWrapper::OpenMaskedStage(
-			*UsdStage.GetRootLayer().GetIdentifier(),
-			EUsdInitialLoadSet::LoadNone,
-			{}
-		);
-		ensure(FreshCurrentStage);
-		TArray<UE::FSdfLayer> CurrentLayerStack = FreshCurrentStage.GetLayerStack();
-		if (!CurrentLayerStack.Contains(Layer))
+		// We should only be allowed to isolate a layer belonging to UsdStage's local layer stack, but checking for that
+		// is not trivial given that layers can be muted.
+
+		TArray<UE::FSdfLayer> ValidLayers;
+		UE::FUsdStage FreshCurrentStage;
+		if (Layer.IsAnonymous() && UsdStage.GetRootLayer() != Layer)
+		{
+			// If we're an anonymous layer that is *not* the stage's root layer, it means that we're
+			// probably a session layer. We can't use the trick below as opening a new stage would
+			// get a brand new session layer, so for a proper check we must use the layer stack of the
+			// currently opened stage
+			const bool bIncludeSessionLayers = true;
+			ValidLayers = UsdStage.GetLayerStack(bIncludeSessionLayers);
+		}
+		else
+		{
+			// Check if the layer we're trying to isolate is part of the stage. To do that we need to reopen a fresh copy
+			// of the current stage in case we have any layers muted, as muted layers don't show up on the layer stack.
+			// Note that we can't just check the list of muted layers either, as it's possible to mute *any* layer for
+			// a given stage, not only the layers that are currently used by it.
+			// We'll use an empty population mask though (which should prevent prim composition) and just use the layers
+			// that are already opened on the current stage anyway, so this should be cheap
+			FreshCurrentStage = UnrealUSDWrapper::OpenMaskedStage(*UsdStage.GetRootLayer().GetIdentifier(), EUsdInitialLoadSet::LoadNone, {});
+			ensure(FreshCurrentStage);
+
+			const bool bIncludeSessionLayers = true;
+			ValidLayers = FreshCurrentStage.GetLayerStack(bIncludeSessionLayers);
+		}
+
+		if (!ValidLayers.Contains(Layer))
 		{
 			UE_LOG(
 				LogUsd,
