USD: Fix morph target curves not being used as controls for skeletal meshes generated without animation.

diff --git a/Engine/Plugins/Importers/USDImporter/Source/USDUtilities/Private/USDSkeletalDataConversion.cpp b/Engine/Plugins/Importers/USDImporter/Source/USDUtilities/Private/USDSkeletalDataConversion.cpp
--- a/Engine/Plugins/Importers/USDImporter/Source/USDUtilities/Private/USDSkeletalDataConversion.cpp
+++ b/Engine/Plugins/Importers/USDImporter/Source/USDUtilities/Private/USDSkeletalDataConversion.cpp
@@ -435,35 +435,6 @@ namespace SkelDataConversionImpl
 			return;
 		}
 
-		if (bAddCurveMetadataToSkeleton)
-		{
-			const bool bMaterialCurve = false;
-			const bool bMorphTargetCurve = true;
-			Skeleton->AccumulateCurveMetaData(CurveName, bMaterialCurve, bMorphTargetCurve);
-		}
-		else
-		{
-			// We should always have the preview mesh here, as we need to declare these curve names on it now
-			USkeletalMesh* Mesh = Sequence->GetPreviewMesh();
-			if (ensure(Mesh))
-			{
-				UAnimCurveMetaData* AnimCurveMetaData = Mesh->GetAssetUserData<UAnimCurveMetaData>();
-				if (AnimCurveMetaData == nullptr)
-				{
-					AnimCurveMetaData = NewObject<UAnimCurveMetaData>(Mesh, NAME_None, RF_Transactional);
-					Mesh->AddAssetUserData(AnimCurveMetaData);
-				}
-
-				AnimCurveMetaData->AddCurveMetaData(CurveName);
-
-				// Ensure we have a morph flag set
-				if (FCurveMetaData* CurveMetaData = AnimCurveMetaData->GetCurveMetaData(CurveName))
-				{
-					CurveMetaData->Type.bMorphtarget = true;
-				}
-			}
-		}
-
 		// Ignore curves that don't contribute to the animation
 		bool bHasNonZeroKey = false;
 		for (const FRichCurveKey& Key : SourceData.Keys)
@@ -2963,6 +2934,38 @@ USkeletalMesh* UsdToUnreal::GetSkeletalMeshFromImportData(
 
 	UsdToUnrealImpl::CreateMorphTargets(InBlendShapesByPath, LODIndexToSkeletalMeshImportData, SkeletalMesh);
 
+	// "Declare" the morph target curves on the skeleton or skeletal mesh according to bAddCurveMetadataToSkeleton.
+	// This is important otherwise the ControlRig will not hoist these curves as controls when using e.g. FKControlRig.
+	for (const TPair<FString, UsdUtils::FUsdBlendShape>& BlendShapeByPath : InBlendShapesByPath)
+	{
+		const UsdUtils::FUsdBlendShape& BlendShape = BlendShapeByPath.Value;
+		const FName CurveName = *BlendShape.Name;
+
+		if (bAddCurveMetadataToSkeleton)
+		{
+			const bool bMaterialCurve = false;
+			const bool bMorphTargetCurve = true;
+			Skeleton->AccumulateCurveMetaData(CurveName, bMaterialCurve, bMorphTargetCurve);
+		}
+		else
+		{
+			UAnimCurveMetaData* AnimCurveMetaData = SkeletalMesh->GetAssetUserData<UAnimCurveMetaData>();
+			if (AnimCurveMetaData == nullptr)
+			{
+				AnimCurveMetaData = NewObject<UAnimCurveMetaData>(SkeletalMesh, NAME_None, RF_Transactional);
+				SkeletalMesh->AddAssetUserData(AnimCurveMetaData);
+			}
+
+			AnimCurveMetaData->AddCurveMetaData(CurveName);
+
+			// Ensure we have a morph flag set
+			if (FCurveMetaData* CurveMetaData = AnimCurveMetaData->GetCurveMetaData(CurveName))
+			{
+				CurveMetaData->Type.bMorphtarget = true;
+			}
+		}
+	}
+
 	return SkeletalMesh;
 }
 
