USD: Downgrade "static function call before module load" warning into a log.

diff --git a/Engine/Plugins/Importers/USDImporter/Source/UnrealUSDWrapper/Private/UnrealUSDWrapper.cpp b/Engine/Plugins/Importers/USDImporter/Source/UnrealUSDWrapper/Private/UnrealUSDWrapper.cpp
--- a/Engine/Plugins/Importers/USDImporter/Source/UnrealUSDWrapper/Private/UnrealUSDWrapper.cpp
+++ b/Engine/Plugins/Importers/USDImporter/Source/UnrealUSDWrapper/Private/UnrealUSDWrapper.cpp
@@ -558,7 +558,7 @@ namespace UE::UnrealUSDWrapper::Private
 		{
 			UE_LOG(
 				LogUsd,
-				Warning,
+				Log,
 				TEXT(
 					"Attempted to call a static function from UnrealUSDWrapper before the module is actually loaded! The module will be loaded now, but not all static functions have this check. In general, please ensure an Unreal module is loaded before calling any of its static functions, for example by calling 'FModuleManager::LoadModuleChecked<IUnrealUSDWrapperModule>(\"UnrealUSDWrapper\");' beforehand."
 				)
