Usd relationship minimal wrapper

diff --git a/Engine/Plugins/Importers/USDImporter/Source/UnrealUSDWrapper/Private/UsdWrappers/UsdPrim.cpp b/Engine/Plugins/Importers/USDImporter/Source/UnrealUSDWrapper/Private/UsdWrappers/UsdPrim.cpp
--- a/Engine/Plugins/Importers/USDImporter/Source/UnrealUSDWrapper/Private/UsdWrappers/UsdPrim.cpp
+++ b/Engine/Plugins/Importers/USDImporter/Source/UnrealUSDWrapper/Private/UsdWrappers/UsdPrim.cpp
@@ -9,6 +9,7 @@
 #include "UsdWrappers/UsdPayloads.h"
 #include "UsdWrappers/UsdReferences.h"
 #include "UsdWrappers/UsdStage.h"
+#include "UsdWrappers/UsdRelationship.h"
 #include "UsdWrappers/UsdVariantSets.h"
 
 #if USE_USD_SDK
@@ -16,6 +17,7 @@
 #include "pxr/usd/sdf/schema.h"
 #include "pxr/usd/usd/attribute.h"
 #include "pxr/usd/usd/prim.h"
+#include "pxr/usd/usd/relationship.h"
 #include "USDIncludesEnd.h"
 #endif	  // #if USE_USD_SDK
 
@@ -407,6 +409,15 @@ namespace UE
 #endif	  // #if USE_USD_SDK
 	}
 
+	FUsdRelationship FUsdPrim::GetRelationship(const TCHAR* RelationshipName) const
+	{
+#if USE_USD_SDK
+		return FUsdRelationship(Impl->PxrUsdPrim.Get().GetRelationship(pxr::TfToken(TCHAR_TO_ANSI(RelationshipName))));
+#else
+		return FUsdRelationship();
+#endif	  // #if USE_USD_SDK
+	}
+
 	FName FUsdPrim::GetName() const
 	{
 #if USE_USD_SDK
diff --git a/Engine/Plugins/Importers/USDImporter/Source/UnrealUSDWrapper/Private/UsdWrappers/UsdRelationship.cpp b/Engine/Plugins/Importers/USDImporter/Source/UnrealUSDWrapper/Private/UsdWrappers/UsdRelationship.cpp
new file mode 100644
--- /dev/null
+++ b/Engine/Plugins/Importers/USDImporter/Source/UnrealUSDWrapper/Private/UsdWrappers/UsdRelationship.cpp
@@ -0,0 +1,173 @@
+﻿// Copyright Epic Games, Inc. All Rights Reserved.
+
+#include "UsdWrappers/UsdRelationship.h"
+#include "UsdWrappers/SdfPath.h"
+#include "USDMemory.h"
+
+#if USE_USD_SDK
+#include "USDIncludesStart.h"
+#include "pxr/usd/usd/relationship.h"
+#include "USDIncludesEnd.h"
+#endif	  // #if USE_USD_SDK
+
+namespace UE
+{
+	namespace Internal
+	{
+		class FUsdRelationshipImpl
+		{
+		public:
+			FUsdRelationshipImpl() = default;
+
+#if USE_USD_SDK
+			explicit FUsdRelationshipImpl(const pxr::UsdRelationship& InUsdRelationship)
+				: PxrUsdRelationship(InUsdRelationship)
+			{
+			}
+
+			explicit FUsdRelationshipImpl(pxr::UsdRelationship&& InUsdRelationship)
+				: PxrUsdRelationship(MoveTemp(InUsdRelationship))
+			{
+			}
+			
+			TUsdStore<pxr::UsdRelationship> PxrUsdRelationship;
+#endif	  // #if USE_USD_SDK
+		};
+	}	  // namespace Internal
+}
+
+namespace UE
+{
+	FUsdRelationship::FUsdRelationship()
+	{
+		FScopedUnrealAllocs UnrealAllocs;
+		Impl = MakeUnique<Internal::FUsdRelationshipImpl>();
+	}
+
+	FUsdRelationship::FUsdRelationship(const FUsdRelationship& Other)
+	{
+#if USE_USD_SDK
+		FScopedUnrealAllocs UnrealAllocs;
+		Impl = MakeUnique<Internal::FUsdRelationshipImpl>(Other.Impl->PxrUsdRelationship.Get());
+#endif	  // #if USE_USD_SDK
+	}
+
+	FUsdRelationship::FUsdRelationship(FUsdRelationship&& Other) = default;
+
+	FUsdRelationship::~FUsdRelationship()
+	{
+		FScopedUnrealAllocs UnrealAllocs;
+		Impl.Reset();
+	}
+
+	FUsdRelationship& FUsdRelationship::operator=(const FUsdRelationship& Other)
+	{
+#if USE_USD_SDK
+		FScopedUnrealAllocs UnrealAllocs;
+		Impl = MakeUnique<Internal::FUsdRelationshipImpl>(Other.Impl->PxrUsdRelationship.Get());
+#endif	  // #if USE_USD_SDK
+		return *this;
+	}
+
+	FUsdRelationship& FUsdRelationship::operator=(FUsdRelationship&& Other)
+	{
+		FScopedUnrealAllocs UnrealAllocs;
+		Impl = MoveTemp(Other.Impl);
+
+		return *this;
+	}
+
+	bool FUsdRelationship::operator==(const FUsdRelationship& Other) const
+	{
+#if USE_USD_SDK
+		return Impl->PxrUsdRelationship.Get() == Other.Impl->PxrUsdRelationship.Get();
+#else
+		return false;
+#endif	  // #if USE_USD_SDK
+	}
+
+	bool FUsdRelationship::operator!=(const FUsdRelationship& Other) const
+	{
+		return !(*this == Other);
+	}
+
+	FUsdRelationship::operator bool() const
+	{
+#if USE_USD_SDK
+		return (bool)Impl->PxrUsdRelationship.Get();
+#else
+		return false;
+#endif	  // #if USE_USD_SDK
+	}
+
+#if USE_USD_SDK
+	FUsdRelationship::FUsdRelationship(const pxr::UsdRelationship& InUsdRelationship)
+	{
+		FScopedUnrealAllocs UnrealAllocs;
+		Impl = MakeUnique<Internal::FUsdRelationshipImpl>(InUsdRelationship);
+	}
+
+	FUsdRelationship::FUsdRelationship(pxr::UsdRelationship&& InUsdRelationship)
+	{
+		FScopedUnrealAllocs UnrealAllocs;
+		Impl = MakeUnique<Internal::FUsdRelationshipImpl>(MoveTemp(InUsdRelationship));
+	}
+
+	FUsdRelationship& FUsdRelationship::operator=(const pxr::UsdRelationship& InUsdRelationship)
+	{
+		FScopedUnrealAllocs UnrealAllocs;
+		Impl = MakeUnique<Internal::FUsdRelationshipImpl>(InUsdRelationship);
+		return *this;
+	}
+
+	FUsdRelationship& FUsdRelationship::operator=(pxr::UsdRelationship&& InUsdRelationship)
+	{
+		FScopedUnrealAllocs UnrealAllocs;
+		Impl = MakeUnique<Internal::FUsdRelationshipImpl>(MoveTemp(InUsdRelationship));
+		return *this;
+	}
+
+	FUsdRelationship::operator pxr::UsdRelationship&()
+	{
+		return Impl->PxrUsdRelationship.Get();
+	}
+
+	FUsdRelationship::operator const pxr::UsdRelationship&() const
+	{
+		return Impl->PxrUsdRelationship.Get();
+	}
+
+	FUsdRelationship::operator pxr::UsdProperty&()
+	{
+		return Impl->PxrUsdRelationship.Get();
+	}
+
+	FUsdRelationship::operator const pxr::UsdProperty&() const
+	{
+		return Impl->PxrUsdRelationship.Get();
+	}
+#endif
+	
+	bool FUsdRelationship::GetTargets(TArray<UE::FSdfPath>& TargetsPath) const
+	{
+		TargetsPath.Reset();
+#if USE_USD_SDK
+		FScopedUsdAllocs Allocs;
+
+		pxr::SdfPathVector Targets;
+		if(Impl->PxrUsdRelationship.Get().GetTargets(&Targets))
+		{
+			const uint32 NumTargets = Targets.size();
+			for(uint32 TargetIndex = 0; TargetIndex < NumTargets; ++TargetIndex)
+			{
+				TargetsPath.Add(UE::FSdfPath(Targets[TargetIndex]));
+			}
+			return true;
+		}
+		
+#endif	  // USE_USD_SDK
+		return false;
+	}
+	
+	
+}	 // namespace UE
diff --git a/Engine/Plugins/Importers/USDImporter/Source/UnrealUSDWrapper/Public/UsdWrappers/UsdPrim.h b/Engine/Plugins/Importers/USDImporter/Source/UnrealUSDWrapper/Public/UsdWrappers/UsdPrim.h
--- a/Engine/Plugins/Importers/USDImporter/Source/UnrealUSDWrapper/Public/UsdWrappers/UsdPrim.h
+++ b/Engine/Plugins/Importers/USDImporter/Source/UnrealUSDWrapper/Public/UsdWrappers/UsdPrim.h
@@ -28,6 +28,8 @@ namespace UE
 	class FUsdReferences;
 	class FUsdVariantSet;
 	class FUsdVariantSets;
+	class FUsdRelationship;
+	
 
 	namespace Internal
 	{
@@ -110,6 +112,7 @@ namespace UE
 
 		const FSdfPath GetPrimPath() const;
 		FUsdStage GetStage() const;
+		FUsdRelationship GetRelationship(const TCHAR* RelationshipName) const;
 
 		FName GetName() const;
 
diff --git a/Engine/Plugins/Importers/USDImporter/Source/UnrealUSDWrapper/Public/UsdWrappers/UsdRelationship.h b/Engine/Plugins/Importers/USDImporter/Source/UnrealUSDWrapper/Public/UsdWrappers/UsdRelationship.h
new file mode 100644
--- /dev/null
+++ b/Engine/Plugins/Importers/USDImporter/Source/UnrealUSDWrapper/Public/UsdWrappers/UsdRelationship.h
@@ -0,0 +1,72 @@
+﻿// Copyright Epic Games, Inc. All Rights Reserved.
+
+#pragma once
+
+#include "Templates/UniquePtr.h"
+#include "UnrealUSDWrapper.h"
+
+#if USE_USD_SDK
+#include "USDIncludesStart.h"
+#include "pxr/pxr.h"
+#include "USDIncludesEnd.h"
+
+PXR_NAMESPACE_OPEN_SCOPE
+	class UsdRelationship;
+	class UsdProperty;
+PXR_NAMESPACE_CLOSE_SCOPE
+#endif	  // #if USE_USD_SDK
+
+namespace UE
+{
+	namespace Internal
+	{
+		class FUsdRelationshipImpl;
+	}
+	class FSdfPath;
+
+	/**
+	 * Minimal pxr::UsdRelationship wrapper for Unreal that can be used from no-rtti modules.
+	 */
+	class UNREALUSDWRAPPER_API FUsdRelationship
+	{
+	public:
+		FUsdRelationship();
+
+		FUsdRelationship(const FUsdRelationship& Other);
+		FUsdRelationship(FUsdRelationship&& Other);
+		~FUsdRelationship();
+
+		FUsdRelationship& operator=(const FUsdRelationship& Other);
+		FUsdRelationship& operator=(FUsdRelationship&& Other);
+
+		bool operator==(const FUsdRelationship& Other) const;
+		bool operator!=(const FUsdRelationship& Other) const;
+
+		explicit operator bool() const;
+
+		// Auto conversion from/to pxr::UsdRelationship
+	public:
+#if USE_USD_SDK
+		explicit FUsdRelationship(const pxr::UsdRelationship& InUsdRelationship);
+		explicit FUsdRelationship(pxr::UsdRelationship&& InUsdRelationship);
+		FUsdRelationship& operator=(const pxr::UsdRelationship& InUsdRelationship);
+		FUsdRelationship& operator=(pxr::UsdRelationship&& InUsdRelationship);
+
+		operator pxr::UsdRelationship&();
+		operator const pxr::UsdRelationship&() const;
+
+		// Also define the conversion operators directly to the base class as we can't "forward declare the class
+		// inheritance", but it is nice to be able to pass an UE::FUsdRelationship to a function that takes an
+		// pxr::UsdProperty directly
+		operator pxr::UsdProperty&();
+		operator const pxr::UsdProperty&() const;
+#endif	  // #if USE_USD_SDK
+		
+	public:
+		/** Get the relationship targets (primitive path relative to the stage) */
+		bool GetTargets(TArray<UE::FSdfPath>& TargetsPath) const;
+		
+	private:
+		TUniquePtr<Internal::FUsdRelationshipImpl> Impl;
+	};
+}	 // namespace UE
