USD: Add support for parsing OpenVDB textures referenced by USD stages.

#virtualized

diff --git a/Engine/Plugins/Importers/USDImporter/Resources/UsdResources/Linux/plugins/unreal/resources/plugInfo.json b/Engine/Plugins/Importers/USDImporter/Resources/UsdResources/Linux/plugins/unreal/resources/plugInfo.json
--- a/Engine/Plugins/Importers/USDImporter/Resources/UsdResources/Linux/plugins/unreal/resources/plugInfo.json
+++ b/Engine/Plugins/Importers/USDImporter/Resources/UsdResources/Linux/plugins/unreal/resources/plugInfo.json
@@ -57,6 +57,20 @@
                             "UsdAPISchemaBase"
                         ], 
                         "schemaKind": "singleApplyAPI"
+                    }, 
+                    "unrealSparseVolumeTextureAPI": {
+                        "alias": {
+                            "UsdSchemaBase": "SparseVolumeTextureAPI"
+                        }, 
+                        "apiSchemaCanOnlyApplyTo": [
+                            "Volume", 
+                            "OpenVDBAsset"
+                        ], 
+                        "autoGenerated": true, 
+                        "bases": [
+                            "UsdAPISchemaBase"
+                        ], 
+                        "schemaKind": "singleApplyAPI"
                     }
                 }
             }, 
diff --git a/Engine/Plugins/Importers/USDImporter/Resources/UsdResources/Mac/plugins/unreal/resources/plugInfo.json b/Engine/Plugins/Importers/USDImporter/Resources/UsdResources/Mac/plugins/unreal/resources/plugInfo.json
--- a/Engine/Plugins/Importers/USDImporter/Resources/UsdResources/Mac/plugins/unreal/resources/plugInfo.json
+++ b/Engine/Plugins/Importers/USDImporter/Resources/UsdResources/Mac/plugins/unreal/resources/plugInfo.json
@@ -57,6 +57,20 @@
                             "UsdAPISchemaBase"
                         ], 
                         "schemaKind": "singleApplyAPI"
+                    }, 
+                    "unrealSparseVolumeTextureAPI": {
+                        "alias": {
+                            "UsdSchemaBase": "SparseVolumeTextureAPI"
+                        }, 
+                        "apiSchemaCanOnlyApplyTo": [
+                            "Volume", 
+                            "OpenVDBAsset"
+                        ], 
+                        "autoGenerated": true, 
+                        "bases": [
+                            "UsdAPISchemaBase"
+                        ], 
+                        "schemaKind": "singleApplyAPI"
                     }
                 }
             }, 
diff --git a/Engine/Plugins/Importers/USDImporter/Resources/UsdResources/Win64/plugins/unreal/resources/plugInfo.json b/Engine/Plugins/Importers/USDImporter/Resources/UsdResources/Win64/plugins/unreal/resources/plugInfo.json
--- a/Engine/Plugins/Importers/USDImporter/Resources/UsdResources/Win64/plugins/unreal/resources/plugInfo.json
+++ b/Engine/Plugins/Importers/USDImporter/Resources/UsdResources/Win64/plugins/unreal/resources/plugInfo.json
@@ -57,6 +57,20 @@
                             "UsdAPISchemaBase"
                         ], 
                         "schemaKind": "singleApplyAPI"
+                    }, 
+                    "unrealSparseVolumeTextureAPI": {
+                        "alias": {
+                            "UsdSchemaBase": "SparseVolumeTextureAPI"
+                        }, 
+                        "apiSchemaCanOnlyApplyTo": [
+                            "Volume", 
+                            "OpenVDBAsset"
+                        ], 
+                        "autoGenerated": true, 
+                        "bases": [
+                            "UsdAPISchemaBase"
+                        ], 
+                        "schemaKind": "singleApplyAPI"
                     }
                 }
             }, 
diff --git a/Engine/Plugins/Importers/USDImporter/Source/USDClasses/Private/USDClassesModule.cpp b/Engine/Plugins/Importers/USDImporter/Source/USDClasses/Private/USDClassesModule.cpp
--- a/Engine/Plugins/Importers/USDImporter/Source/USDClasses/Private/USDClassesModule.cpp
+++ b/Engine/Plugins/Importers/USDImporter/Source/USDClasses/Private/USDClassesModule.cpp
@@ -14,6 +14,7 @@
 #include "Engine/SkinnedAssetCommon.h"
 #include "Engine/StaticMesh.h"
 #include "Engine/Texture.h"
+#include "Engine/Texture2D.h"
 #include "EngineAnalytics.h"
 #include "GeometryCache.h"
 #include "GroomAsset.h"
@@ -31,6 +32,7 @@
 #include "Misc/SecureHash.h"
 #include "PhysicsEngine/PhysicsAsset.h"
 #include "Serialization/JsonSerializer.h"
+#include "SparseVolumeTexture/SparseVolumeTexture.h"
 #include "UObject/NameTypes.h"
 #include "UObject/ObjectSaveContext.h"
 #include "UObject/Package.h"
@@ -424,6 +426,10 @@ TSet<UObject*> IUsdClassesModule::GetAssetDependencies(UObject* Asset)
 	{
 		// Do nothing. AnimBlueprints have no additional dependencies
 	}
+	else if (USparseVolumeTexture* SparseVolumeTexture = Cast<USparseVolumeTexture>(Asset))
+	{
+		// Do nothing. SparseVolumeTextures have no additional dependencies
+	}
 	else
 	{
 		UE_LOG(LogUsd, Warning, TEXT("Unknown asset '%s' encountered when collecting dependencies."), Asset ? *Asset->GetName() : TEXT("nullptr"));
diff --git a/Engine/Plugins/Importers/USDImporter/Source/USDClasses/Public/USDAssetUserData.h b/Engine/Plugins/Importers/USDImporter/Source/USDClasses/Public/USDAssetUserData.h
--- a/Engine/Plugins/Importers/USDImporter/Source/USDClasses/Public/USDAssetUserData.h
+++ b/Engine/Plugins/Importers/USDImporter/Source/USDClasses/Public/USDAssetUserData.h
@@ -8,7 +8,7 @@
 
 #include "USDAssetUserData.generated.h"
 
-UCLASS()
+UCLASS(BlueprintType)
 class USDCLASSES_API UUsdAssetUserData : public UAssetUserData
 {
 	GENERATED_BODY()
@@ -26,7 +26,7 @@ public:
 	TMap<FString, FUsdCombinedPrimMetadata> StageIdentifierToMetadata;
 };
 
-UCLASS()
+UCLASS(BlueprintType)
 class USDCLASSES_API UUsdAnimSequenceAssetUserData : public UUsdAssetUserData
 {
 	GENERATED_BODY()
@@ -44,17 +44,17 @@ public:
 };
 
 /** Simple wrapper because we're not allowed to have TMap properties with TArray<FString> as values */
-USTRUCT()
+USTRUCT(BlueprintType)
 struct FUsdPrimPathList
 {
 	GENERATED_BODY()
 
-	UPROPERTY()
+	UPROPERTY(EditAnywhere, BlueprintReadWrite, Category = "USD")
 	TArray<FString> PrimPaths;
 };
 
 /** AssetImportData assigned to Unreal materials that are generated when parsing USD Material prims */
-UCLASS()
+UCLASS(BlueprintType)
 class USDCLASSES_API UUsdMaterialAssetUserData : public UUsdAssetUserData
 {
 	GENERATED_BODY()
@@ -65,7 +65,7 @@ public:
 	 * describes the primvar that the USD material is sampling for each texture.
 	 * e.g. {'BaseColor': 'st0', 'Metallic': 'st1'}
 	 */
-	UPROPERTY()
+	UPROPERTY(EditAnywhere, BlueprintReadWrite, Category = "USD")
 	TMap<FString, FString> ParameterToPrimvar;
 
 	/**
@@ -77,12 +77,12 @@ public:
 	 * material instance that is.
 	 * e.g. {'firstPrimvar': 0, 'st': 1, 'st1': 2}
 	 */
-	UPROPERTY()
+	UPROPERTY(EditAnywhere, BlueprintReadWrite, Category = "USD")
 	TMap<FString, int32> PrimvarToUVIndex;
 };
 
 /** We assign these to UStaticMeshes or USkeletalMeshes generated from USD */
-UCLASS()
+UCLASS(BlueprintType)
 class USDCLASSES_API UUsdMeshAssetUserData : public UUsdAssetUserData
 {
 	GENERATED_BODY()
@@ -93,23 +93,23 @@ public:
 	 * assignment. It can contain multiple prims in case we combine material slots and/or collapse prims
 	 * (e.g. {0: ['/Root/mesh', '/Root/othermesh/geomsubset0', '/Root/othermesh/geomsubset1'] }).
 	 */
-	UPROPERTY()
+	UPROPERTY(EditAnywhere, BlueprintReadWrite, Category = "USD")
 	TMap<int32, FUsdPrimPathList> MaterialSlotToPrimPaths;
 
 	/** Describes which primvars should be assigned to each UV index. */
-	UPROPERTY()
+	UPROPERTY(EditAnywhere, BlueprintReadWrite, Category = "USD")
 	TMap<FString, int32> PrimvarToUVIndex;
 };
 
 /** We assign these to UGeometryCaches generated from USD */
-UCLASS()
+UCLASS(BlueprintType)
 class USDCLASSES_API UUsdGeometryCacheAssetUserData : public UUsdMeshAssetUserData
 {
 	GENERATED_BODY()
 
 public:
 	// Check analogous comment on UUsdAnimSequenceAssetUserData
-	UPROPERTY()
+	UPROPERTY(EditAnywhere, BlueprintReadWrite, Category = "USD")
 	float LayerStartOffsetSeconds = 0.0f;
 };
 
@@ -117,7 +117,7 @@ public:
  * We assign these to persistent LevelSequences that bind to one of the actors/components that the stage actor spawns.
  * We need this as part of a mechanism to automatically repair those bindings when they break if we close/reload the stage.
  */
-UCLASS()
+UCLASS(BlueprintType)
 class USDCLASSES_API UUsdLevelSequenceAssetUserData : public UAssetUserData
 {
 	GENERATED_BODY()
@@ -127,13 +127,46 @@ public:
 	 * The LevelSequence has a Guid that is changed every time its state is modified.
 	 * We pay attention to that so that we can avoid reprocessing a LevelSequence that hasn't changed
 	 */
-	UPROPERTY()
+	UPROPERTY(EditAnywhere, BlueprintReadWrite, Category = "USD")
 	FGuid LastCheckedSignature;
 
 	/**
 	 * Set of binding GUIDs that we already handled in the past.
 	 * We use this so that we won't try and overwrite the changes in case the user manually clears/modifies a binding we previously setup.
 	 */
-	UPROPERTY()
+	UPROPERTY(EditAnywhere, BlueprintReadWrite, Category = "USD")
 	TSet<FGuid> HandledBindingGuids;
 };
+
+/**
+ * We use this mainly to help in mapping between stage timeCode and the FrameIndex for animated SVTs
+ */
+UCLASS(BlueprintType)
+class USDCLASSES_API UUsdSparseVolumeTextureAssetUserData : public UUsdAssetUserData
+{
+	GENERATED_BODY()
+
+public:
+	/** Paths to all the OpenVDBAsset prims that led to the generation of this SVT asset */
+	UPROPERTY(EditAnywhere, BlueprintReadWrite, Category = "USD")
+	TArray<FString> SourceOpenVDBAssetPrimPaths;
+
+	/** TimeCodes of all the filePath attribute time samples as seen on the OpenVDBAsset prim in its own layer, in order */
+	UPROPERTY(EditAnywhere, BlueprintReadWrite, Category = "USD")
+	TArray<double> TimeSamplePathTimeCodes;
+
+	/**
+	 * Corresponding indices of the frame of the SVT that should be played at a particular timeCode.
+	 * Example: TimeSamplePathTimeCodes is [10, 20] and TimeSamplePathIndices is [2, 7] --> At timeCode 10 the frame index 2
+	 * of the SVT should be played, i.e. the .vdb file that is the third entry within TimeSamplePaths
+	 */
+	UPROPERTY(EditAnywhere, BlueprintReadWrite, Category = "USD")
+	TArray<int32> TimeSamplePathIndices;
+
+	/**
+	 * File paths that originated each of the SVT frames, in order.
+	 * The SVT should have as many frames as there are entries in this array.
+	 */
+	UPROPERTY(EditAnywhere, BlueprintReadWrite, Category = "USD")
+	TArray<FString> TimeSamplePaths;
+};
\ No newline at end of file
diff --git a/Engine/Plugins/Importers/USDImporter/Source/USDClasses/Public/USDProjectSettings.h b/Engine/Plugins/Importers/USDImporter/Source/USDClasses/Public/USDProjectSettings.h
--- a/Engine/Plugins/Importers/USDImporter/Source/USDClasses/Public/USDProjectSettings.h
+++ b/Engine/Plugins/Importers/USDImporter/Source/USDClasses/Public/USDProjectSettings.h
@@ -110,6 +110,15 @@ public:
 
 	// Note that the below properties being FSoftObjectPath ensure that these materials are cooked into packaged games
 
+	/**
+	 * Material to use when handling .vdb files as Sparse Volume Textures. An instance of this material will be
+	 * added to the AHeterogeneousVolume, and will use the parsed SparseVolumeTexture as a texture parameter.
+	 * Note that alternatively Volume prims can have material bindings to Unreal materials, and the importer
+	 * will prioritize trying to use those as the volumetric materials for the Sparse Volume Textures instead.
+	 */
+	UPROPERTY(config, EditAnywhere, Category = "USD|Reference Materials", meta = (AllowedClasses = "/Script/Engine.MaterialInterface"))
+	FSoftObjectPath ReferenceDefaultSVTMaterial = FSoftObjectPath{TEXT("/USDImporter/Materials/ReferenceDefaultSVTMaterial.ReferenceDefaultSVTMaterial")};
+
 	/**
 	 * What material to use for UUsdDrawModeComponents with "Cards" draw mode and provided textures (corresponding to
 	 * UsdGeomModelAPI with the "cards" drawMode).
diff --git a/Engine/Plugins/Importers/USDImporter/Source/USDSchemas/Private/USDGeomXformableTranslator.cpp b/Engine/Plugins/Importers/USDImporter/Source/USDSchemas/Private/USDGeomXformableTranslator.cpp
--- a/Engine/Plugins/Importers/USDImporter/Source/USDSchemas/Private/USDGeomXformableTranslator.cpp
+++ b/Engine/Plugins/Importers/USDImporter/Source/USDSchemas/Private/USDGeomXformableTranslator.cpp
@@ -930,7 +930,7 @@ void FUsdGeomXformableTranslator::CreateAlternativeDrawModeAssets(EUsdDrawMode D
 	{
 		if (Attr && Attr.HasAuthoredValue())
 		{
-			const FString ResolvedPath = UsdUtils::GetResolvedTexturePath(Attr);
+			const FString ResolvedPath = UsdUtils::GetResolvedAssetPath(Attr);
 			if (ResolvedPath.IsEmpty())
 			{
 				pxr::SdfAssetPath TextureAssetPath;
diff --git a/Engine/Plugins/Importers/USDImporter/Source/USDSchemas/Private/USDSchemasModule.cpp b/Engine/Plugins/Importers/USDImporter/Source/USDSchemas/Private/USDSchemasModule.cpp
--- a/Engine/Plugins/Importers/USDImporter/Source/USDSchemas/Private/USDSchemasModule.cpp
+++ b/Engine/Plugins/Importers/USDImporter/Source/USDSchemas/Private/USDSchemasModule.cpp
@@ -24,6 +24,7 @@
 #include "USDMemory.h"
 #include "USDShadeMaterialTranslator.h"
 #include "USDSkelSkeletonTranslator.h"
+#include "USDVolVolumeTranslator.h"
 #endif	  // #if USE_USD_SDK
 
 class FUsdSchemasModule : public IUsdSchemasModule
@@ -51,7 +52,8 @@ public:
 			Registry.Register<FUsdGeomXformableTranslator>(TEXT("UsdGeomXformable")),
 			Registry.Register<FUsdShadeMaterialTranslator>(TEXT("UsdShadeMaterial")),
 			Registry.Register<FUsdLuxLightTranslator>(TEXT("UsdLuxBoundableLightBase")),
-			Registry.Register<FUsdLuxLightTranslator>(TEXT("UsdLuxNonboundableLightBase"))};
+			Registry.Register<FUsdLuxLightTranslator>(TEXT("UsdLuxNonboundableLightBase")),
+			Registry.Register<FUsdVolVolumeTranslator>(TEXT("UsdVolVolume"))};
 
 #if WITH_EDITOR
 		ShaderRegistry.Register(FMaterialXUsdShadeMaterialTranslator::MaterialXRenderContext);
diff --git a/Engine/Plugins/Importers/USDImporter/Source/USDSchemas/Private/USDVolVolumeTranslator.cpp b/Engine/Plugins/Importers/USDImporter/Source/USDSchemas/Private/USDVolVolumeTranslator.cpp
new file mode 100644
--- /dev/null
+++ b/Engine/Plugins/Importers/USDImporter/Source/USDSchemas/Private/USDVolVolumeTranslator.cpp
@@ -0,0 +1,1484 @@
+// Copyright Epic Games, Inc. All Rights Reserved.
+
+#include "USDVolVolumeTranslator.h"
+
+#include "USDAssetImportData.h"
+#include "USDAssetUserData.h"
+#include "USDConversionUtils.h"
+#include "USDDrawModeComponent.h"
+#include "USDIntegrationUtils.h"
+#include "USDLayerUtils.h"
+#include "USDLog.h"
+#include "USDPrimConversion.h"
+#include "USDProjectSettings.h"
+#include "USDShadeConversion.h"
+#include "USDTypesConversion.h"
+#include "UsdWrappers/SdfPath.h"
+#include "UsdWrappers/UsdPrim.h"
+
+#include "Components/HeterogeneousVolumeComponent.h"
+#include "Engine/Level.h"
+#include "MaterialDomain.h"
+#include "Materials/Material.h"
+#include "Materials/MaterialInstance.h"
+#include "Materials/MaterialInstanceConstant.h"
+#include "Materials/MaterialInstanceDynamic.h"
+#include "MaterialShared.h"
+#include "SparseVolumeTexture/SparseVolumeTexture.h"
+
+#if WITH_EDITOR
+#include "AssetImportTask.h"
+#include "OpenVDBImportOptions.h"
+#include "SparseVolumeTextureFactory.h"
+#endif	  // WITH_EDITOR
+
+#if USE_USD_SDK
+
+#include "USDIncludesStart.h"
+#include <pxr/usd/usd/prim.h>
+#include <pxr/usd/usd/stage.h>
+#include <pxr/usd/usdShade/material.h>
+#include <pxr/usd/usdShade/materialBindingAPI.h>
+#include <pxr/usd/usdVol/openVDBAsset.h>
+#include <pxr/usd/usdVol/volume.h>
+#include "USDIncludesEnd.h"
+
+static bool bRemoveDuplicates = true;
+static FAutoConsoleVariableRef CVarRemoveDuplicates(
+	TEXT("USD.Volume.RemoveDuplicateAnimatedFrames"),
+	bRemoveDuplicates,
+	TEXT(
+		"If this is true (default), the contents of a .VDB file are added only once to animated Sparse Volume Textures (SVT), even if the same file shows in multiple different time samples. If this is false, every OpenVDBAsset prim filePath time sample is parse as a new frame on the animated SVT."
+	)
+);
+
+namespace UE::UsdVolVolumeTranslator::Private
+{
+#if WITH_EDITOR
+	struct FSparseVolumeTextureInfo
+	{
+		FString SourceVDBFilePath;
+
+		TArray<FString> SourceOpenVDBAssetPrimPaths;
+
+		TArray<double> TimeSamplePathTimeCodes;
+		TArray<int32> TimeSamplePathIndices;
+		TArray<FString> TimeSamplePaths;	// First item may equal SourceVDBFilePath. May have duplicate entries
+
+		// Fields on the Volume prim that were connected to this .vdb file
+		TArray<FString> VolumeFieldNames;
+
+		// e.g. {'velocity': {'AttributesA.R': 'X', 'AttributesA.G': 'Y', 'AttributesA.B': 'Z'}}
+		TMap<FString, TMap<FString, FString>> GridNameToChannelComponentMapping;
+		TOptional<ESparseVolumeAttributesFormat> AttributesAFormat;
+		TOptional<ESparseVolumeAttributesFormat> AttributesBFormat;
+
+		USparseVolumeTexture* SparseVolumeTexture = nullptr;
+		FString PrefixedAssetHash;
+	};
+
+	static const int32 SparseVolumeTextureChannelCount = 8;
+
+	// Collect all the .vdb files that this prim wants to parse, and the desired fields/grids from them.
+	//
+	// In VDB terminology a "grid" is essentially a 3D texture, and can have formats like float, double3, half, etc.
+	// In USD the analogous term is "field", but essentially means the same thing. Possibly the terminology is abstracted
+	// to also fit the Field3D library, which we don't support it. Field/grid will be used interchangeably here.
+	//
+	// USD is very flexible and allows the user to reference specific grids from of each .vdb file. The syntax makes it
+	// difficult to find out at once all the grids we'll need to parse from each the .vdb files, so here we need to group them
+	// up first before deferring to the SparseVolumeTextureFactory.
+	//
+	// Note that USD allows a single Volume prim to reference grids from multiple .vdb files, and to also "timeSample" the
+	// file reference to allow for volume animations. This means that in UE a "Volume" prim corresponds to a single
+	// HeterogeneousVolumeActor, but which in turn can have any number of Sparse Volume Textures (one for each .vdb file referenced).
+	TMap<FString, FSparseVolumeTextureInfo> CollectSparseVolumeTextureInfoByFilePath(const pxr::UsdVolVolume& Volume)
+	{
+		TMap<FString, FSparseVolumeTextureInfo> FilePathHashToInfo;
+
+		FScopedUsdAllocs UsdAllocs;
+
+		pxr::UsdStageRefPtr Stage = Volume.GetPrim().GetStage();
+
+		const std::map<pxr::TfToken, pxr::SdfPath>& FieldMap = Volume.GetFieldPaths();
+		for (std::map<pxr::TfToken, pxr::SdfPath>::const_iterator Iter = FieldMap.cbegin(); Iter != FieldMap.end(); ++Iter)
+		{
+			// This field name is the name of the field for the Volume prim, which can be anything and differ from the
+			// grid name within the .vdb files
+			const pxr::TfToken& FieldName = Iter->first;
+			const pxr::SdfPath& AssetPrimPath = Iter->second;
+
+			pxr::UsdPrim OpenVDBPrim = Stage->GetPrimAtPath(AssetPrimPath);
+			pxr::UsdVolOpenVDBAsset OpenVDBPrimSchema = pxr::UsdVolOpenVDBAsset{OpenVDBPrim};
+			if (OpenVDBPrimSchema)
+			{
+				pxr::UsdAttribute FilePathAttr = OpenVDBPrimSchema.GetFilePathAttr();
+
+				FString ResolvedVDBPath = UsdUtils::GetResolvedAssetPath(FilePathAttr, pxr::UsdTimeCode::Default());
+
+				// Find timesampled paths, if any
+				TArray<double> TimeSamplePathTimeCodes;
+				TArray<int32> TimeSamplePathIndices;
+				TArray<FString> TimeSamplePaths;
+
+				TMap<FString, int32> PathToIndex;
+
+				std::vector<double> TimeSamples;
+				if (FilePathAttr.GetTimeSamples(&TimeSamples) && TimeSamples.size() > 0)
+				{
+					UE::FSdfLayerOffset CombinedOffset = UsdUtils::GetPrimToStageOffset(UE::FUsdPrim{OpenVDBPrim});
+
+					TimeSamplePathTimeCodes.Reserve(TimeSamples.size());
+					TimeSamplePaths.Reserve(TimeSamples.size());
+					TimeSamplePathIndices.Reserve(TimeSamples.size());
+					for (double TimeSample : TimeSamples)
+					{
+						// We always want to store on the AssetUserData (which is where this stuff will end up in)
+						// the time codes in the layer where the actual OpenVDBPrim is authored. If that layer is referenced
+						// by a parent layer through an offset and scale, TimeSample will contain that offset and scale here,
+						// which we need to undo
+						double LayerLocalTimeCode = (TimeSample - CombinedOffset.Offset) / CombinedOffset.Scale;
+						TimeSamplePathTimeCodes.Add(LayerLocalTimeCode);
+
+						FString ResolvedFramePath = UsdUtils::GetResolvedAssetPath(FilePathAttr, TimeSample);
+
+						// If we had no default time sample to act as the "main file", take the first frame
+						if (ResolvedVDBPath.IsEmpty())
+						{
+							ResolvedVDBPath = ResolvedFramePath;
+						}
+
+						if (bRemoveDuplicates)
+						{
+							if (int32* FoundIndex = PathToIndex.Find(ResolvedFramePath))
+							{
+								TimeSamplePathIndices.Add(*FoundIndex);
+							}
+							else
+							{
+								TimeSamplePaths.Add(ResolvedFramePath);
+
+								const int32 NewIndex = TimeSamplePaths.Num() - 1;
+								PathToIndex.Add(ResolvedFramePath, NewIndex);
+								TimeSamplePathIndices.Add(NewIndex);
+							}
+						}
+						else
+						{
+							TimeSamplePaths.Add(ResolvedFramePath);
+
+							const int32 NewIndex = TimeSamplePaths.Num() - 1;
+							TimeSamplePathIndices.Add(NewIndex);
+						}
+					}
+				}
+
+				// Hash all the relevant file paths here: The collection of file paths to parse determines the SVT, and
+				// we want one FSparseVolumeTextureInfo per SVT
+				FString FilePathHashString;
+				{
+					FSHA1 SHA1;
+					SHA1.UpdateWithString(*ResolvedVDBPath, ResolvedVDBPath.Len());
+					for (const FString& TimeSamplePath : TimeSamplePaths)
+					{
+						SHA1.UpdateWithString(*TimeSamplePath, TimeSamplePath.Len());
+					}
+					SHA1.Final();
+
+					FSHAHash FilePathHash;
+					SHA1.GetHash(&FilePathHash.Hash[0]);
+
+					FilePathHashString = FilePathHash.ToString();
+				}
+
+				if (!ResolvedVDBPath.IsEmpty())
+				{
+					FSparseVolumeTextureInfo& SparseVolumeTextureInfo = FilePathHashToInfo.FindOrAdd(FilePathHashString);
+					SparseVolumeTextureInfo.SourceOpenVDBAssetPrimPaths.AddUnique(UsdToUnreal::ConvertPath(AssetPrimPath));
+					SparseVolumeTextureInfo.SourceVDBFilePath = ResolvedVDBPath;
+					SparseVolumeTextureInfo.TimeSamplePathTimeCodes = MoveTemp(TimeSamplePathTimeCodes);
+					SparseVolumeTextureInfo.TimeSamplePathIndices = MoveTemp(TimeSamplePathIndices);
+					SparseVolumeTextureInfo.TimeSamplePaths = MoveTemp(TimeSamplePaths);
+
+					FString FieldNameStr = UsdToUnreal::ConvertToken(FieldName);
+					SparseVolumeTextureInfo.VolumeFieldNames.AddUnique(FieldNameStr);
+
+					pxr::TfToken GridName;
+					pxr::UsdAttribute Attr = OpenVDBPrimSchema.GetFieldNameAttr();
+					if (Attr && Attr.Get<pxr::TfToken>(&GridName))
+					{
+						FString GridNameStr = UsdToUnreal::ConvertToken(GridName);
+
+						// Note we want this to add an entry to SparseVolumeTexture.GridNameToChannelNames even if we won't
+						// find the schema on the prim, as we'll use these entries to make sure the generated Sparse Volume
+						// Texture contains theses desired fields
+						TMap<FString, FString>& ChannelToComponent = SparseVolumeTextureInfo.GridNameToChannelComponentMapping.FindOrAdd(GridNameStr);
+
+						if (UsdUtils::PrimHasSchema(OpenVDBPrim.GetPrim(), UnrealIdentifiers::SparseVolumeTextureAPI))
+						{
+							// Parse desired data types for AttributesA and AttributesB channels
+							TFunction<void(pxr::TfToken, TOptional<ESparseVolumeAttributesFormat>&)> HandleAttribute =
+								[&OpenVDBPrim,
+								 &AssetPrimPath,
+								 &ResolvedVDBPath](pxr::TfToken AttrName, TOptional<ESparseVolumeAttributesFormat>& AttributeFormat)
+							{
+								using FormatMapType = std::unordered_map<pxr::TfToken, ESparseVolumeAttributesFormat, pxr::TfHash>;
+								using InverseFormatMapType = std::unordered_map<ESparseVolumeAttributesFormat, FString>;
+
+								// clang-format off
+								const static FormatMapType FormatMap = {
+									{pxr::TfToken("unorm8"),  ESparseVolumeAttributesFormat::Unorm8},
+									{pxr::TfToken("float16"), ESparseVolumeAttributesFormat::Float16},
+									{pxr::TfToken("float32"), ESparseVolumeAttributesFormat::Float32},
+								};
+								const static InverseFormatMapType InverseFormatMap = {
+									{ESparseVolumeAttributesFormat::Unorm8,  TEXT("unorm8")},
+									{ESparseVolumeAttributesFormat::Float16, TEXT("float16")},
+									{ESparseVolumeAttributesFormat::Float32, TEXT("float32")},
+								};
+								// clang-format on
+
+								pxr::TfToken DataType;
+								pxr::UsdAttribute AttrA = OpenVDBPrim.GetAttribute(AttrName);
+								if (AttrA && AttrA.Get<pxr::TfToken>(&DataType))
+								{
+									FormatMapType::const_iterator Iter = FormatMap.find(DataType);
+									if (Iter != FormatMap.end())
+									{
+										ESparseVolumeAttributesFormat TargetFormat = Iter->second;
+
+										// Check in case multiple OpenVDBAsset prims want different values for the data type
+										const bool bIsSet = AttributeFormat.IsSet();
+										if (bIsSet && AttributeFormat.GetValue() != TargetFormat)
+										{
+											const FString* ExistingFormat = nullptr;
+											InverseFormatMapType::const_iterator ExistingIter = InverseFormatMap.find(AttributeFormat.GetValue());
+											if (ExistingIter != InverseFormatMap.end())
+											{
+												ExistingFormat = &ExistingIter->second;
+											}
+
+											UE_LOG(
+												LogUsd,
+												Warning,
+												TEXT(
+													"OpenVDBAsset prims disagree on the attribute channel format for the Sparse Volume Texture generated for VDB file '%s' (encountered '%s' and '%s'). If there are multiple opinions for the attribute channel formats from different OpenVDBAsset prims, they must all agree!"
+												),
+												*ResolvedVDBPath,
+												*UsdToUnreal::ConvertToken(DataType),
+												ExistingFormat ? **ExistingFormat : TEXT("unknown")
+											);
+										}
+										else if (!bIsSet)
+										{
+											AttributeFormat = TargetFormat;
+										}
+									}
+									else
+									{
+										UE_LOG(
+											LogUsd,
+											Warning,
+											TEXT(
+												"Invalid Sparse Volume Texture attribute channel format '%s'. Available formats: 'unorm8', 'float16' and 'float32'."
+											),
+											*UsdToUnreal::ConvertToken(DataType),
+											*UsdToUnreal::ConvertPath(AssetPrimPath)
+										);
+									}
+								}
+							};
+							HandleAttribute(UnrealIdentifiers::UnrealSVTAttributesADataType, SparseVolumeTextureInfo.AttributesAFormat);
+							HandleAttribute(UnrealIdentifiers::UnrealSVTAttributesBDataType, SparseVolumeTextureInfo.AttributesBFormat);
+
+							// Parse desired channel assignment
+							pxr::VtArray<pxr::TfToken> DesiredChannels;
+							pxr::VtArray<pxr::TfToken> DesiredComponents;
+							pxr::UsdAttribute ComponentsAttr = OpenVDBPrim.GetAttribute(UnrealIdentifiers::UnrealSVTMappedGridComponents);
+							pxr::UsdAttribute ChannelsAttr = OpenVDBPrim.GetAttribute(UnrealIdentifiers::UnrealSVTMappedAttributeChannels);
+							if (ChannelsAttr &&														 //
+								ComponentsAttr &&													 //
+								ChannelsAttr.Get<pxr::VtArray<pxr::TfToken>>(&DesiredChannels) &&	 //
+								ComponentsAttr.Get<pxr::VtArray<pxr::TfToken>>(&DesiredComponents))
+							{
+								// These must always match of course
+								if (DesiredChannels.size() == DesiredComponents.size())
+								{
+									// If we have more than one OpenVDBAsset prim reading from the same VDB file, the declared component to
+									// channel mappings must be compatible
+									for (uint32 Index = 0; Index < DesiredChannels.size(); ++Index)
+									{
+										FString Channel = UsdToUnreal::ConvertString(DesiredChannels[Index]);
+										FString Component = UsdToUnreal::ConvertString(DesiredComponents[Index]);
+
+										if (const FString* ExistingComponentMapping = ChannelToComponent.Find(Channel))
+										{
+											if (Component != *ExistingComponentMapping)
+											{
+												UE_LOG(
+													LogUsd,
+													Warning,
+													TEXT(
+														"Found multiple OpenVDBAsset prims (including '%s') targetting the same grid '%s', but with with conflicting grid component to Sparse Volume Texture attribute channel mapping (for example, both components '%s' and '%s' are mapped to the same channel '%s', which is not allowed)"
+													),
+													*UsdToUnreal::ConvertPath(AssetPrimPath),
+													*GridNameStr,
+													**ExistingComponentMapping,
+													*Component,
+													*Channel
+												);
+											}
+										}
+										else
+										{
+											ChannelToComponent.Add(Channel, Component);
+										}
+									}
+								}
+								else
+								{
+									UE_LOG(
+										LogUsd,
+										Warning,
+										TEXT(
+											"Failed to parse custom component to attribute mapping from OpenVDBAsset prim '%s': The '%s' and '%s' attributes should have the same number of entries, but the former has %u entries while the latter has %u"
+										),
+										*UsdToUnreal::ConvertPath(AssetPrimPath),
+										*UsdToUnreal::ConvertToken(UnrealIdentifiers::UnrealSVTMappedGridComponents),
+										*UsdToUnreal::ConvertToken(UnrealIdentifiers::UnrealSVTMappedAttributeChannels),
+										DesiredComponents.size(),
+										DesiredChannels.size()
+									);
+								}
+							}
+						}
+					}
+				}
+				else
+				{
+					UE_LOG(
+						LogUsd,
+						Warning,
+						TEXT("Failed to find the VDB file '%s' referenced by OpenVDBAsset prim at path '%s'"),
+						*ResolvedVDBPath,
+						*UsdToUnreal::ConvertPath(AssetPrimPath)
+					);
+				}
+			}
+			else
+			{
+				UE_LOG(
+					LogUsd,
+					Warning,
+					TEXT("Failed to find an OpenVDBAsset prim at path '%s' for field '%s' of prim '%s'"),
+					*UsdToUnreal::ConvertPath(AssetPrimPath),
+					*UsdToUnreal::ConvertToken(FieldName),
+					*UsdToUnreal::ConvertPath(Volume.GetPrim().GetPath())
+				);
+			}
+		}
+
+		return FilePathHashToInfo;
+	}
+
+	// Here we must stash into InOutPreviewData.ImportOptions the desired channel mapping for this SVT given all the GridNameToChannelNames
+	// mappings we pulled out of the prims if they had any instances of our SparseVolumeTextureAPI schema
+	void SetVDBImportOptions(const FSparseVolumeTextureInfo& ParsedTexture, FOpenVDBPreviewData& InOutPreviewData)
+	{
+		// Tweak the collected filenames for other frames: The OpenVDB importer will scan for similar filenames
+		// in the same folder as the main file, but through USD we expect the user to manually pick file paths
+		// for each time sample (which may or may not come from the same folder, or be in any particular order)
+		InOutPreviewData.SequenceFilenames = ParsedTexture.TimeSamplePaths;
+		InOutPreviewData.DefaultImportOptions.bIsSequence = ParsedTexture.TimeSamplePaths.Num() > 1;
+
+		// Apply manually specified channel formats, if any
+		check(InOutPreviewData.DefaultImportOptions.Attributes.Num() == 2);
+		if (ParsedTexture.AttributesAFormat.IsSet())
+		{
+			InOutPreviewData.DefaultImportOptions.Attributes[0].Format = ParsedTexture.AttributesAFormat.GetValue();
+		}
+		if (ParsedTexture.AttributesBFormat.IsSet())
+		{
+			InOutPreviewData.DefaultImportOptions.Attributes[1].Format = ParsedTexture.AttributesBFormat.GetValue();
+		}
+
+		static const TMap<FString, int32> ChannelIndexMap = {
+			{TEXT("AttributesA.R"), 0},
+			{TEXT("AttributesA.G"), 1},
+			{TEXT("AttributesA.B"), 2},
+			{TEXT("AttributesA.A"), 3},
+			{TEXT("AttributesB.R"), 4},
+			{TEXT("AttributesB.G"), 5},
+			{TEXT("AttributesB.B"), 6},
+			{TEXT("AttributesB.A"), 7}
+		  };
+
+		static const TMap<FString, int32> ComponentIndexMap = {
+			{TEXT("X"), 0},
+			{TEXT("Y"), 1},
+			{TEXT("Z"), 2},
+			{TEXT("W"), 3},
+			{TEXT("R"), 0},
+			{TEXT("G"), 1},
+			{TEXT("B"), 2},
+			{TEXT("A"), 3}
+		  };
+
+		// We'll use this to make sure we only try assigning one thing to each available attribute channel
+		FString UsedChannels[SparseVolumeTextureChannelCount];
+
+		TMap<FString, const FOpenVDBGridInfo*> GridInfoByName;
+		for (const FOpenVDBGridInfo& GridInfo : InOutPreviewData.GridInfo)
+		{
+			GridInfoByName.Add(GridInfo.Name, &GridInfo);
+		}
+
+		FString AvailableGrids;
+		static FString AvailableComponentNames;
+		static FString AvailableChannelNames;
+
+		// We'll collect our new mapping here and only apply to InOutPreviewData if we have a valid mapping,
+		// so that we don't wipe it clean if we don't have anything valid to add back anyway
+		TStaticArray<FOpenVDBSparseVolumeAttributesDesc, 2> NewChannelMapping;
+		bool bHadValidMapping = false;
+
+		for (const TPair<FString, TMap<FString, FString>>& GridToMapping : ParsedTexture.GridNameToChannelComponentMapping)
+		{
+			const FString& GridName = GridToMapping.Key;
+			if (const FOpenVDBGridInfo* FoundGridInfo = GridInfoByName.FindRef(GridName))
+			{
+				int32 GridIndex = static_cast<int32>(FoundGridInfo->Index);
+
+				for (const TPair<FString, FString>& ChannelToComponent : GridToMapping.Value)
+				{
+					// Get and validate the desired component as an index (e.g. whether this mapping refers to 'velocity.X' (index 0) or 'velocity.Y'
+					// (index 1), etc.)
+					const FString& DesiredComponent = ChannelToComponent.Value;
+					int32 ComponentIndex = INDEX_NONE;
+					if (const int32* FoundComponentIndex = ComponentIndexMap.Find(DesiredComponent))
+					{
+						ComponentIndex = *FoundComponentIndex;
+
+						check(ComponentIndex < SparseVolumeTextureChannelCount);
+						if (ComponentIndex < 0 || static_cast<uint32>(ComponentIndex) >= FoundGridInfo->NumComponents)
+						{
+							UE_LOG(
+								LogUsd,
+								Warning,
+								TEXT("Invalid component name '%s' for grid '%s' in VDB file '%s', as that particular grid only has %u components."),
+								*DesiredComponent,
+								*GridName,
+								*ParsedTexture.SourceVDBFilePath,
+								FoundGridInfo->NumComponents
+							);
+							ComponentIndex = INDEX_NONE;
+						}
+					}
+					else
+					{
+						if (AvailableComponentNames.IsEmpty())
+						{
+							TArray<FString> ComponentNames;
+							ComponentIndexMap.GetKeys(ComponentNames);
+
+							AvailableComponentNames = TEXT("'") + FString::Join(ComponentNames, TEXT("', '")) + TEXT("'");
+						}
+
+						UE_LOG(
+							LogUsd,
+							Warning,
+							TEXT(
+								"Desired component name '%s' for grid '%s' in VDB file '%s' is not a valid component name. Avaliable component names: %s"
+							),
+							*DesiredComponent,
+							*GridName,
+							*ParsedTexture.SourceVDBFilePath,
+							*AvailableComponentNames
+						);
+					}
+
+					// Get and validate desired channel (e.g. whether this mapping means to put something on 'AttributesA.R' or 'AttributesB.A', etc.)
+					const FString& DesiredChannel = ChannelToComponent.Key;
+					int32 ChannelIndex = INDEX_NONE;
+					if (const int32* FoundChannelIndex = ChannelIndexMap.Find(DesiredChannel))
+					{
+						ChannelIndex = *FoundChannelIndex;
+
+						const bool bIsSet = !UsedChannels[ChannelIndex].IsEmpty();
+						const FString GridAndComponent = GridName + TEXT(".") + DesiredComponent;
+
+						if (bIsSet && UsedChannels[ChannelIndex] != GridAndComponent)
+						{
+							UE_LOG(
+								LogUsd,
+								Warning,
+								TEXT(
+									"Cannot use attribute channel '%s' for grid '%s' in VDB file '%s', as the channel is already being used for the grid and component '%s'"
+								),
+								*DesiredChannel,
+								*GridName,
+								*ParsedTexture.SourceVDBFilePath,
+								*UsedChannels[ChannelIndex]
+							);
+							ChannelIndex = INDEX_NONE;
+						}
+						else
+						{
+							ChannelIndex = *FoundChannelIndex;
+							UsedChannels[ChannelIndex] = GridAndComponent;
+						}
+					}
+					else
+					{
+						if (AvailableChannelNames.IsEmpty())
+						{
+							TArray<FString> ChannelNames;
+							ChannelIndexMap.GetKeys(ChannelNames);
+
+							AvailableChannelNames = TEXT("'") + FString::Join(ChannelNames, TEXT("', '")) + TEXT("'");
+						}
+
+						UE_LOG(
+							LogUsd,
+							Warning,
+							TEXT(
+								"Desired attribute channel '%s' for grid '%s' in VDB file '%s' is not a valid channel name. Avaliable channel names: %s"
+							),
+							*DesiredChannel,
+							*GridName,
+							*ParsedTexture.SourceVDBFilePath,
+							*AvailableChannelNames
+						);
+					}
+
+					// Finally actually assign the desired grid/component mapping
+					if (ComponentIndex != INDEX_NONE && GridIndex != INDEX_NONE && ChannelIndex != INDEX_NONE)
+					{
+						// We track ChannelIndex from 0 through 7, but it's actually two groups of 4
+						int32 AttributeGroupIndex = ChannelIndex / 4;
+						int32 AttributeChannelIndex = ChannelIndex % 4;
+
+						FOpenVDBSparseVolumeComponentMapping& ComponentMapping = NewChannelMapping[AttributeGroupIndex]
+																					 .Mappings[AttributeChannelIndex];
+						ComponentMapping.SourceGridIndex = GridIndex;
+						ComponentMapping.SourceComponentIndex = ComponentIndex;
+
+						bHadValidMapping = true;
+					}
+				}
+			}
+			else
+			{
+				if (AvailableGrids.IsEmpty())
+				{
+					TArray<FString> GridNames;
+					GridInfoByName.GetKeys(GridNames);
+
+					if (GridNames.Num() > 0)
+					{
+						AvailableGrids = TEXT("'") + FString::Join(GridNames, TEXT("', '")) + TEXT("'");
+					}
+				}
+
+				UE_LOG(
+					LogUsd,
+					Warning,
+					TEXT("Faild to find grid named '%s' inside VDB file '%s'. Avaliable grid names: %s"),
+					*GridName,
+					*ParsedTexture.SourceVDBFilePath,
+					*AvailableGrids
+				);
+			}
+		}
+
+		if (bHadValidMapping)
+		{
+			InOutPreviewData.DefaultImportOptions.Attributes[0].Mappings = NewChannelMapping[0].Mappings;
+			InOutPreviewData.DefaultImportOptions.Attributes[1].Mappings = NewChannelMapping[1].Mappings;
+		}
+	}
+
+	void HashForSparseVolumeTexture(const FOpenVDBPreviewData& PreviewData, FSHA1& InOutHash)
+	{
+		InOutHash.Update(
+			reinterpret_cast<const uint8*>(PreviewData.LoadedFile.GetData()),
+			PreviewData.LoadedFile.Num() * PreviewData.LoadedFile.GetTypeSize()
+		);
+
+		// Hash other files
+		{
+			FMD5 MD5;
+
+			if (PreviewData.SequenceFilenames.Num() > 1)
+			{
+				// Skip first one as that should always be the "main" file, that we just hashed on PreviewData.LoadedFile above.
+				// Note: This could become a performance issue if we have many large frames
+				for (int32 Index = 1; Index < PreviewData.SequenceFilenames.Num(); ++Index)
+				{
+					const FString& FrameFilePath = PreviewData.SequenceFilenames[Index];
+
+					// Copied from FMD5Hash::HashFileFromArchive as it doesn't expose its FMD5
+					if (TUniquePtr<FArchive> Ar{IFileManager::Get().CreateFileReader(*FrameFilePath)})
+					{
+						TArray<uint8> LocalScratch;
+						LocalScratch.SetNumUninitialized(1024 * 64);
+
+						const int64 Size = Ar->TotalSize();
+						int64 Position = 0;
+
+						while (Position < Size)
+						{
+							const auto ReadNum = FMath::Min(Size - Position, (int64)LocalScratch.Num());
+							Ar->Serialize(LocalScratch.GetData(), ReadNum);
+							MD5.Update(LocalScratch.GetData(), ReadNum);
+
+							Position += ReadNum;
+						}
+					}
+				}
+
+				FMD5Hash Hash;
+				Hash.Set(MD5);
+				InOutHash.Update(Hash.GetBytes(), Hash.GetSize());
+			}
+		}
+
+		// The only other thing that affects the SVT asset hash is the grid to attribute channel mapping.
+		// i.e. if we have another Volume prim with entirely different field names but that ends up with the same grid names
+		// mapped to the same attribute channels, we want to reuse the generated SVT asset
+
+		InOutHash.Update(
+			reinterpret_cast<const uint8*>(&PreviewData.DefaultImportOptions.bIsSequence),
+			sizeof(PreviewData.DefaultImportOptions.bIsSequence)
+		);
+
+		const TStaticArray<FOpenVDBSparseVolumeAttributesDesc, 2>& AttributeMapping = PreviewData.DefaultImportOptions.Attributes;
+		for (int32 AttributesIndex = 0; AttributesIndex < AttributeMapping.Num(); ++AttributesIndex)
+		{
+			const FOpenVDBSparseVolumeAttributesDesc& AttributesDesc = AttributeMapping[AttributesIndex];
+
+			InOutHash.Update(reinterpret_cast<const uint8*>(&AttributesDesc.Format), sizeof(AttributesDesc.Format));
+
+			const TStaticArray<FOpenVDBSparseVolumeComponentMapping, 4>& ChannelMapping = AttributesDesc.Mappings;
+			for (int32 ChannelIndex = 0; ChannelIndex < ChannelMapping.Num(); ++ChannelIndex)
+			{
+				const FOpenVDBSparseVolumeComponentMapping& Map = ChannelMapping[ChannelIndex];
+				InOutHash.Update(reinterpret_cast<const uint8*>(&Map.SourceGridIndex), sizeof(Map.SourceGridIndex));
+				InOutHash.Update(reinterpret_cast<const uint8*>(&Map.SourceComponentIndex), sizeof(Map.SourceComponentIndex));
+			}
+		}
+	}
+
+	void HashForVolumetricMaterial(
+		const UMaterialInterface* ReferenceMaterial,
+		const TMap<FString, FSparseVolumeTextureInfo*>& MaterialParameterToTexture,
+		FSHA1& InOutHash
+	)
+	{
+		FString ReferenceMaterialPathName = ReferenceMaterial->GetPathName();
+		InOutHash.UpdateWithString(*ReferenceMaterialPathName, ReferenceMaterialPathName.Len());
+
+		TArray<TPair<FString, FSparseVolumeTextureInfo*>> MaterialParameterPairs;
+		MaterialParameterPairs.Reserve(MaterialParameterToTexture.Num());
+		for (const TPair<FString, FSparseVolumeTextureInfo*>& Pair : MaterialParameterToTexture)
+		{
+			MaterialParameterPairs.Add(Pair);
+		}
+
+		// Make sure we hash our SVTs deterministically, whether they have a specific material assignment due
+		// to the schema or not
+		MaterialParameterPairs.Sort(
+			[](const TPair<FString, FSparseVolumeTextureInfo*>& LHS, const TPair<FString, FSparseVolumeTextureInfo*>& RHS) -> bool
+			{
+				if (LHS.Key == RHS.Key)
+				{
+					return LHS.Value->PrefixedAssetHash < RHS.Value->PrefixedAssetHash;
+				}
+				return LHS.Key < RHS.Key;
+			}
+		);
+
+		for (const TPair<FString, FSparseVolumeTextureInfo*>& Pair : MaterialParameterPairs)
+		{
+			InOutHash.UpdateWithString(*Pair.Key, Pair.Key.Len());
+			InOutHash.UpdateWithString(*Pair.Value->PrefixedAssetHash, Pair.Value->PrefixedAssetHash.Len());
+		}
+	}
+
+	// This collects a mapping describing which Sparse Volume Texture (SVT) should be assigned to each SVT material parameter
+	// of the ReferenceMaterial.
+	// It will prefer checking the VolumePrim for a custom schema where that is manually described, then it will fall back
+	// to trying to map Volume prim field names to material parameter names, and finally will just distribute the SVTs over all
+	// available parameters in alphabetical order
+	TMap<FString, FSparseVolumeTextureInfo*> CollectMaterialParameterTextureAssignment(
+		const pxr::UsdPrim& VolumePrim,
+		const UMaterial* ReferenceMaterial,
+		const TMap<FString, FSparseVolumeTextureInfo>& FilePathHashToTextureInfo
+	)
+	{
+		TMap<FString, FSparseVolumeTextureInfo*> ResultParameterToInfo;
+
+		if (!VolumePrim || !ReferenceMaterial || FilePathHashToTextureInfo.Num() == 0)
+		{
+			return ResultParameterToInfo;
+		}
+
+		FScopedUsdAllocs Allocs;
+
+		// Collect which field was mapped to which .VDB asset (and so SVT asset)
+		// A field can only be mapped to a single .VDB, but multiple fields can be mapped to the same .VDB
+		TMap<FString, FSparseVolumeTextureInfo*> FieldNameToInfo;
+		FieldNameToInfo.Reserve(FilePathHashToTextureInfo.Num());
+		for (const TPair<FString, FSparseVolumeTextureInfo>& FilePathToInfo : FilePathHashToTextureInfo)
+		{
+			for (const FString& FieldName : FilePathToInfo.Value.VolumeFieldNames)
+			{
+				// We won't be modifying FSparseVolumeTextureInfo in this function anyway, it's just so that we can
+				// return non-const pointers
+				FieldNameToInfo.Add(FieldName, const_cast<FSparseVolumeTextureInfo*>(&FilePathToInfo.Value));
+			}
+		}
+
+		// Collect material parameter assignments manually specified via the custom schema, if any
+		bool bHadManualAssignment = false;
+		if (UsdUtils::PrimHasSchema(VolumePrim, UnrealIdentifiers::SparseVolumeTextureAPI))
+		{
+			pxr::VtArray<pxr::TfToken> MappedFields;
+			pxr::VtArray<pxr::TfToken> MappedParameters;
+			pxr::UsdAttribute FieldsAttr = VolumePrim.GetAttribute(UnrealIdentifiers::UnrealSVTMappedFields);
+			pxr::UsdAttribute ParametersAttr = VolumePrim.GetAttribute(UnrealIdentifiers::UnrealSVTMappedMaterialParameters);
+			if (FieldsAttr &&													//
+				ParametersAttr &&												//
+				FieldsAttr.Get<pxr::VtArray<pxr::TfToken>>(&MappedFields) &&	//
+				ParametersAttr.Get<pxr::VtArray<pxr::TfToken>>(&MappedParameters))
+			{
+				// These must always match of course
+				if (MappedFields.size() == MappedParameters.size())
+				{
+					for (uint32 Index = 0; Index < MappedFields.size(); ++Index)
+					{
+						FString FieldName = UsdToUnreal::ConvertString(MappedFields[Index]);
+						FString MaterialParameter = UsdToUnreal::ConvertString(MappedParameters[Index]);
+						bHadManualAssignment = true;
+
+						if (FSparseVolumeTextureInfo* FoundParsedTexture = FieldNameToInfo.FindRef(FieldName))
+						{
+							if (const FSparseVolumeTextureInfo* JobAtParam = ResultParameterToInfo.FindRef(MaterialParameter))
+							{
+								if (JobAtParam->SparseVolumeTexture != FoundParsedTexture->SparseVolumeTexture)
+								{
+									UE_LOG(
+										LogUsd,
+										Warning,
+										TEXT(
+											"Trying to assign different Sparse Volume Textures to the same material parameter '%s' on the material instantiated for Volume prim '%s' and field name '%s'! Only a single material can be assigned to a material parameter at a time."
+										),
+										*MaterialParameter,
+										*UsdToUnreal::ConvertPath(VolumePrim.GetPrimPath()),
+										*FieldName
+									);
+								}
+							}
+							else
+							{
+								// We won't be modifying anything in here, it's just so that we can return non-const pointers
+								ResultParameterToInfo.Add(MaterialParameter, FoundParsedTexture);
+							}
+						}
+					}
+				}
+				else
+				{
+					UE_LOG(
+						LogUsd,
+						Warning,
+						TEXT(
+							"Failed to parse custom parsed texture to material parameter mapping from volume prim '%s': The '%s' and '%s' attributes should have the same number of entries, but the former has %u entries while the latter has %u"
+						),
+						*UsdToUnreal::ConvertPath(VolumePrim.GetPrimPath()),
+						*UsdToUnreal::ConvertToken(UnrealIdentifiers::UnrealSVTMappedFields),
+						*UsdToUnreal::ConvertToken(UnrealIdentifiers::UnrealSVTMappedMaterialParameters),
+						MappedFields.size(),
+						MappedParameters.size()
+					);
+				}
+			}
+		}
+
+		// Collect available parameter names on this material instance
+		TSet<FString> SparseVolumeTextureParameterNames;
+		{
+			TMap<FMaterialParameterInfo, FMaterialParameterMetadata> SparseVolumeTextureParameters;
+			ReferenceMaterial->GetAllParametersOfType(EMaterialParameterType::SparseVolumeTexture, SparseVolumeTextureParameters);
+
+			SparseVolumeTextureParameterNames.Reserve(SparseVolumeTextureParameters.Num());
+			for (const TPair<FMaterialParameterInfo, FMaterialParameterMetadata>& ParameterPair : SparseVolumeTextureParameters)
+			{
+				const FMaterialParameterInfo& ParameterInfo = ParameterPair.Key;
+				FString ParameterName = ParameterInfo.Name.ToString();
+
+				SparseVolumeTextureParameterNames.Add(ParameterName);
+			}
+		}
+
+		// Validate that all parameters exist on the material, or else emit a warning
+		for (const TPair<FString, FSparseVolumeTextureInfo*>& AssignmentPair : ResultParameterToInfo)
+		{
+			if (!SparseVolumeTextureParameterNames.Contains(AssignmentPair.Key))
+			{
+				UE_LOG(
+					LogUsd,
+					Warning,
+					TEXT("Failed to assign Sparse Volume Texture '%s' to material '%s' as the desired material parameter '%s' doesn't exist on it"),
+					*AssignmentPair.Value->SparseVolumeTexture->GetPathName(),
+					*ReferenceMaterial->GetPathName(),
+					*AssignmentPair.Key
+				);
+			}
+		}
+
+		// No manual material parameter assignment specified via custom schema: First let's assume that the field names match material parameter names
+		bool bHadParameterNameMatch = false;
+		if (!bHadManualAssignment)
+		{
+			ensure(ResultParameterToInfo.Num() == 0);
+
+			TMap<FString, FString> CaseInsensitiveToSensitiveParameterNames;
+			CaseInsensitiveToSensitiveParameterNames.Reserve(SparseVolumeTextureParameterNames.Num());
+			for (const FString& ParameterName : SparseVolumeTextureParameterNames)
+			{
+				CaseInsensitiveToSensitiveParameterNames.Add(ParameterName.ToLower(), ParameterName);
+			}
+
+			for (const TPair<FString, FSparseVolumeTextureInfo*>& TextureJobPair : FieldNameToInfo)
+			{
+				const FString& FieldName = TextureJobPair.Key;
+
+				if (const FString* CaseSensitiveParameterName = CaseInsensitiveToSensitiveParameterNames.Find(FieldName.ToLower()))
+				{
+					ResultParameterToInfo.Add(*CaseSensitiveParameterName, TextureJobPair.Value);
+					bHadParameterNameMatch = true;
+				}
+			}
+		}
+
+		// Nothing yet, let's fall back to just distributing the SVTs across the available material parameter slots in alphabetical order
+		if (!bHadManualAssignment && !bHadParameterNameMatch)
+		{
+			ensure(ResultParameterToInfo.Num() == 0);
+
+			// If there aren't enough parameters, let the user know
+			if ((SparseVolumeTextureParameterNames.Num() < FilePathHashToTextureInfo.Num()) && FilePathHashToTextureInfo.Num() > 0)
+			{
+				UE_LOG(
+					LogUsd,
+					Warning,
+					TEXT(
+						"Material '%s' used for prim '%s' doesn't have enough Sparse Volume Texture params to fit all of its %d parsed textures! Some may be left unassigned."
+					),
+					*ReferenceMaterial->GetPathName(),
+					*UsdToUnreal::ConvertPath(VolumePrim.GetPrimPath()),
+					FilePathHashToTextureInfo.Num()
+				);
+			}
+
+			TArray<FString> SortedFieldNames;
+			FieldNameToInfo.GetKeys(SortedFieldNames);
+			SortedFieldNames.Sort();
+
+			TArray<FString> SortedParameterNames = SparseVolumeTextureParameterNames.Array();
+			SortedParameterNames.Sort();
+
+			for (int32 Index = 0; Index < SortedFieldNames.Num() && Index < SortedParameterNames.Num(); ++Index)
+			{
+				const FString& FieldName = SortedFieldNames[Index];
+				const FString& ParameterName = SortedParameterNames[Index];
+
+				FSparseVolumeTextureInfo* FoundTextureJob = FieldNameToInfo.FindRef(FieldName);
+				if (ensure(FoundTextureJob))
+				{
+					ResultParameterToInfo.Add(ParameterName, FoundTextureJob);
+				}
+			}
+		}
+
+		return ResultParameterToInfo;
+	}
+
+	UMaterialInstance* InstantiateMaterial(
+		FName InstanceName,
+		UMaterialInterface* ReferenceMaterial,
+		const FUsdSchemaTranslationContext& TranslationContext
+	)
+	{
+		if (!ReferenceMaterial)
+		{
+			return nullptr;
+		}
+
+		// Create an UMaterialInstanceConstant
+		if (GIsEditor)	  // Also have to prevent Standalone game from going with MaterialInstanceConstants
+		{
+			UMaterialInstanceConstant* NewMaterial = NewObject<UMaterialInstanceConstant>(
+				GetTransientPackage(),
+				InstanceName,
+				TranslationContext.ObjectFlags | EObjectFlags::RF_Transient
+			);
+
+			if (NewMaterial)
+			{
+				if (ensure(ReferenceMaterial))
+				{
+					// We can't blindly recreate all component render states when a level is being added, because we may end up first creating
+					// render states for some components, and UWorld::AddToWorld calls FScene::AddPrimitive which expects the component to not
+					// have primitives yet
+					FMaterialUpdateContext::EOptions::Type Options = FMaterialUpdateContext::EOptions::Default;
+					if (TranslationContext.Level && TranslationContext.Level->bIsAssociatingLevel)
+					{
+						Options = (FMaterialUpdateContext::EOptions::Type)(Options & ~FMaterialUpdateContext::EOptions::RecreateRenderStates);
+					}
+					FMaterialUpdateContext UpdateContext(Options, GMaxRHIShaderPlatform);
+					UpdateContext.AddMaterialInstance(NewMaterial);
+
+					NewMaterial->SetParentEditorOnly(ReferenceMaterial);
+
+					NewMaterial->PreEditChange(nullptr);
+					NewMaterial->PostEditChange();
+
+					return NewMaterial;
+				}
+			}
+		}
+		else
+		// Create an UMaterialInstanceDynamic
+		{
+			// SparseVolumeTextures can't be created at runtime so this branch should never really be taken for now, but anyway...
+			// Note: Some code in FNiagaraBakerRenderer::RenderSparseVolumeTexture suggests that this workflow wouldn't really work
+			// because the HeterogeneousVolumeComponent always creates its own MID from the material we give it, and creating a MID
+			// from another MID doesn't really work
+			if (ensure(ReferenceMaterial))
+			{
+				UMaterialInstanceDynamic* NewMaterial = UMaterialInstanceDynamic::Create(ReferenceMaterial, GetTransientPackage(), InstanceName);
+				if (NewMaterial)
+				{
+					NewMaterial->SetFlags(RF_Transient);
+				}
+				return NewMaterial;
+			}
+		}
+
+		return nullptr;
+	}
+
+	void AssignMaterialParameters(UMaterialInstance* MaterialInstance, const TMap<FString, FSparseVolumeTextureInfo*>& ParameterToTexture)
+	{
+		// Now that we finally have the parameter assignment for each SVT, assign them to the materials
+		for (const TPair<FString, FSparseVolumeTextureInfo*>& Pair : ParameterToTexture)
+		{
+			if (UMaterialInstanceConstant* Constant = Cast<UMaterialInstanceConstant>(MaterialInstance))
+			{
+				FMaterialParameterInfo Info;
+				Info.Name = *Pair.Key;
+				Constant->SetSparseVolumeTextureParameterValueEditorOnly(Info, Pair.Value->SparseVolumeTexture);
+			}
+			else if (UMaterialInstanceDynamic* Dynamic = Cast<UMaterialInstanceDynamic>(MaterialInstance))
+			{
+				Dynamic->SetSparseVolumeTextureParameterValue(*Pair.Key, Pair.Value->SparseVolumeTexture);
+			}
+		}
+	}
+#endif	  // WITH_EDITOR
+}	 // namespace UE::UsdVolVolumeTranslator::Private
+
+void FUsdVolVolumeTranslator::CreateAssets()
+{
+#if WITH_EDITOR
+	TRACE_CPUPROFILER_EVENT_SCOPE(FUsdVolVolumeTranslator::CreateAssets);
+
+	using namespace UE::UsdVolVolumeTranslator::Private;
+
+	if (!Context->AssetCache || !Context->InfoCache)
+	{
+		return;
+	}
+
+	// Don't bother generating assets if we're going to just draw some bounds for this prim instead
+	EUsdDrawMode DrawMode = UsdUtils::GetAppliedDrawMode(GetPrim());
+	if (DrawMode != EUsdDrawMode::Default)
+	{
+		CreateAlternativeDrawModeAssets(DrawMode);
+		return;
+	}
+
+	if (!Context->bAllowParsingSparseVolumeTextures)
+	{
+		return;
+	}
+
+	const FString VolumePrimPathString = PrimPath.GetString();
+	pxr::UsdPrim VolumePrim = GetPrim();
+	pxr::UsdVolVolume Volume{VolumePrim};
+	if (!Volume)
+	{
+		return;
+	}
+	pxr::UsdStageRefPtr Stage = VolumePrim.GetStage();
+
+	FString VolumePrimHashPrefix = UsdUtils::GetAssetHashPrefix(GetPrim(), Context->bReuseIdenticalAssets);
+
+	// Create the info structs from the requested files
+	TMap<FString, FSparseVolumeTextureInfo> FilePathHashToTextureInfo = CollectSparseVolumeTextureInfoByFilePath(Volume);
+
+	// Create SVT assets from the info structs
+	for (TPair<FString, FSparseVolumeTextureInfo>& FilePathHashToInfo : FilePathHashToTextureInfo)
+	{
+		const FString& VDBFilePath = FilePathHashToInfo.Value.SourceVDBFilePath;
+		FSparseVolumeTextureInfo& ParsedTexture = FilePathHashToInfo.Value;
+
+		if (!FPaths::FileExists(VDBFilePath))
+		{
+			UE_LOG(
+				LogUsd,
+				Warning,
+				TEXT("Failed to find a VDB file at path '%s' when parsing Volume prim '%s'"),
+				*VDBFilePath,
+				*VolumePrimPathString
+			);
+			continue;
+		}
+
+		// Here we're going to pick how to map between the grids from the .vdb files (each one is a separate volumetric texture,
+		// like "density" or "temperature", etc.) into the SVT texture's 8 attribute channels (AttributesA.RGBA and AttributesB.RGBA),
+		// and also pick the attribute channel data types.
+		//
+		// By default we'll defer to LoadOpenVDBPreviewData which has some heuristics based on the grid names and data types.
+		// In practice these .vdb files should only have 1-3 grids each with some common names so the heuristics should hopefully
+		// be fine for a sensible result.
+		//
+		// Users can also add a custom schema to the OpenVDBAsset prims in order to manually control how to map the grids to the SVT
+		// attributes, in a similar way to how blendshapes are mapped. We'll check for those in SetVDBImportOptions
+		TStrongObjectPtr<UOpenVDBImportOptionsObject> ImportOptions{NewObject<UOpenVDBImportOptionsObject>()};
+		LoadOpenVDBPreviewData(VDBFilePath, &ImportOptions->PreviewData);
+		SetVDBImportOptions(ParsedTexture, ImportOptions->PreviewData);
+
+		// Collect a hash for this VDB asset
+		FSHAHash VDBAndAssignmentHash;
+		{
+			FSHA1 SHA1;
+			HashForSparseVolumeTexture(ImportOptions->PreviewData, SHA1);
+
+			SHA1.Final();
+			SHA1.GetHash(&VDBAndAssignmentHash.Hash[0]);
+		}
+		ParsedTexture.PrefixedAssetHash = VolumePrimHashPrefix + VDBAndAssignmentHash.ToString();
+
+		USparseVolumeTexture* SparseVolumeTexture = Cast<USparseVolumeTexture>(Context->AssetCache->GetCachedAsset(ParsedTexture.PrefixedAssetHash));
+
+		// Need to create a brand new asset
+		if (!SparseVolumeTexture)
+		{
+			TStrongObjectPtr<USparseVolumeTextureFactory> SparseVolumeTextureFactory{NewObject<USparseVolumeTextureFactory>()};
+
+			// We use the asset import task to indicate it's an automated import, and also to transmit our import options
+			TStrongObjectPtr<UAssetImportTask> AssetImportTask{NewObject<UAssetImportTask>()};
+			AssetImportTask->Filename = VDBFilePath;
+			AssetImportTask->bAutomated = true;
+			AssetImportTask->bSave = false;
+			AssetImportTask->Options = ImportOptions.Get();
+			AssetImportTask->Factory = SparseVolumeTextureFactory.Get();
+
+			bool bOperationCanceled = false;
+			const TCHAR* Parms = nullptr;
+			const FName AssetName = MakeUniqueObjectName(
+				GetTransientPackage(),
+				USparseVolumeTexture::StaticClass(),
+				*IUsdClassesModule::SanitizeObjectName(FPaths::GetBaseFilename(VDBFilePath))	// File path instead of prim path in case we have
+																								// multiple .vdb files in the same Volume prim
+			);
+
+			// Call FactoryCreateFile directly here or else the usual AssetToolsModule.Get().ImportAssetTasks()
+			// workflow would end up creating a package for every asset, which we don't care about since
+			// we'll rename them into the asset cache anyway
+			SparseVolumeTextureFactory->SetAssetImportTask(AssetImportTask.Get());
+			SparseVolumeTexture = Cast<USparseVolumeTexture>(SparseVolumeTextureFactory->FactoryCreateFile(
+				USparseVolumeTexture::StaticClass(),
+				GetTransientPackage(),
+				AssetName,
+				Context->ObjectFlags | EObjectFlags::RF_Public | EObjectFlags::RF_Transient,
+				VDBFilePath,
+				Parms,
+				GWarn,
+				bOperationCanceled
+			));
+			if (!SparseVolumeTexture)
+			{
+				UE_LOG(LogUsd, Error, TEXT("Failed to generate Sparse Volume Texture from OpenVDB file '%s'"), *VDBFilePath);
+				return;
+			}
+
+			SparseVolumeTexture->PostEditChange();
+
+			if (UStreamableSparseVolumeTexture* StreamableTexture = Cast<UStreamableSparseVolumeTexture>(SparseVolumeTexture))
+			{
+				// Set an asset import data into the texture as it won't do that on its own, and we would otherwise
+				// lose the source .vdb file information downstream
+				UUsdAssetImportData* ImportData = NewObject<UUsdAssetImportData>(SparseVolumeTexture);
+				ImportData->UpdateFilenameOnly(VDBFilePath);
+
+				StreamableTexture->AssetImportData = ImportData;
+			}
+
+			Context->AssetCache->CacheAsset(ParsedTexture.PrefixedAssetHash, SparseVolumeTexture);
+		}
+
+		if (SparseVolumeTexture)
+		{
+			Context->InfoCache->LinkAssetToPrim(PrimPath, SparseVolumeTexture);
+
+			if (UUsdSparseVolumeTextureAssetUserData* UserData = UsdUtils::GetOrCreateAssetUserData<UUsdSparseVolumeTextureAssetUserData>(
+					SparseVolumeTexture
+				))
+			{
+				UserData->PrimPaths.AddUnique(VolumePrimPathString);
+				UserData->SourceOpenVDBAssetPrimPaths = ParsedTexture.SourceOpenVDBAssetPrimPaths;
+				UserData->TimeSamplePaths = ParsedTexture.TimeSamplePaths;
+				UserData->TimeSamplePathIndices = ParsedTexture.TimeSamplePathIndices;
+				UserData->TimeSamplePathTimeCodes = ParsedTexture.TimeSamplePathTimeCodes;
+
+				if (Context->MetadataOptions.bCollectMetadata)
+				{
+					UsdToUnreal::ConvertMetadata(
+						VolumePrim,
+						UserData,
+						Context->MetadataOptions.BlockedPrefixFilters,
+						Context->MetadataOptions.bInvertFilters,
+						Context->MetadataOptions.bCollectFromEntireSubtrees
+					);
+				}
+				else
+				{
+					UserData->StageIdentifierToMetadata.Remove(UsdToUnreal::ConvertString(Stage->GetRootLayer()->GetIdentifier()));
+				}
+			}
+		}
+
+		ParsedTexture.SparseVolumeTexture = SparseVolumeTexture;
+	}
+
+	// Create SVT materials
+	// Sparse volume textures are really 3D textures, and our actor essentially has a 3D cube mesh and will
+	// draw these textures on the level. There's one step missing: The material to use.
+	//
+	// By default we'll spawn an instance of a reference material that we ship, that is basically just a simple
+	// volume domain material with "add" blend mode, that connects AttributesA.R to the "extinction" material output,
+	// and AttributesB.RGB into "albedo".
+	//
+	// The default material should be enough to get "something to show up", but realistically for the correct look the
+	// user would need to set up a custom material. Even more so because these .vdb files can contain grids that are
+	// meant to be drawn as level sets, or float values that are meant to go through look-up tables, usually don't have
+	// any color, etc.
+	//
+	// Volume prims are Gprims however, and can have material bindings, which is what we'll fetch below. If this happens
+	// to be an UnrealMaterial, we'll try to use it as the SVT material instead of our default. We'll only need to find
+	// the correct material parameter to put our SVT assets in, and once again we can use the custom schema to let the
+	// user specify the correct material parameter name for each SVT, in case there are more than one option.
+
+	UMaterialInterface* ReferenceMaterial = nullptr;
+
+	// Check to see if the Volume prim has a material binding to an UnrealMaterial we can use
+	static FName UnrealRenderContext = *UsdToUnreal::ConvertToken(UnrealIdentifiers::Unreal);
+	if (Context->RenderContext == UnrealRenderContext)
+	{
+		pxr::TfToken MaterialPurposeToken = pxr::UsdShadeTokens->allPurpose;
+		if (!Context->MaterialPurpose.IsNone())
+		{
+			MaterialPurposeToken = UnrealToUsd::ConvertToken(*Context->MaterialPurpose.ToString()).Get();
+		}
+
+		pxr::UsdShadeMaterialBindingAPI BindingAPI{VolumePrim};
+		if (pxr::UsdShadeMaterial ShadeMaterial = BindingAPI.ComputeBoundMaterial(MaterialPurposeToken))
+		{
+			if (TOptional<FString> UnrealMaterial = UsdUtils::GetUnrealSurfaceOutput(ShadeMaterial.GetPrim()))
+			{
+				ReferenceMaterial = Cast<UMaterialInterface>(FSoftObjectPath{UnrealMaterial.GetValue()}.TryLoad());
+			}
+		}
+	}
+	// Fall back to the default SVT material instead
+	if (!ReferenceMaterial)
+	{
+		const UUsdProjectSettings* ProjectSettings = GetDefault<UUsdProjectSettings>();
+		if (!ProjectSettings)
+		{
+			return;
+		}
+
+		ReferenceMaterial = Cast<UMaterialInterface>(ProjectSettings->ReferenceDefaultSVTMaterial.TryLoad());
+	}
+
+	const UMaterial* Material = nullptr;
+	if (ReferenceMaterial)
+	{
+		Material = ReferenceMaterial->GetMaterial();
+
+		// Warn in case the used material can't be used for SVTs
+		if (Material && Material->MaterialDomain != EMaterialDomain::MD_Volume)
+		{
+			UE_LOG(
+				LogUsd,
+				Warning,
+				TEXT("The material '%s' used for the Volume prim '%s' may not be capable of using Sparse Volume Textures as it does not have the "
+					 "Volume material domain."),
+				*ReferenceMaterial->GetPathName(),
+				*VolumePrimPathString
+			);
+		}
+	}
+	if (!Material)
+	{
+		return;
+	}
+
+	TMap<FString, FSparseVolumeTextureInfo*> MaterialParameterToTexture = CollectMaterialParameterTextureAssignment(
+		VolumePrim,
+		Material,
+		FilePathHashToTextureInfo
+	);
+
+	FSHAHash MaterialHash;
+	{
+		FSHA1 SHA1;
+		HashForVolumetricMaterial(ReferenceMaterial, MaterialParameterToTexture, SHA1);
+		SHA1.Final();
+		SHA1.GetHash(&MaterialHash.Hash[0]);
+	}
+	const FString PrefixedMaterialHash = VolumePrimHashPrefix + MaterialHash.ToString();
+
+	UMaterialInstance* MaterialInstance = nullptr;
+
+	if (Context->AssetCache)
+	{
+		MaterialInstance = Cast<UMaterialInstance>(Context->AssetCache->GetCachedAsset(PrefixedMaterialHash));
+	}
+
+	// Create new material instance
+	if (!MaterialInstance)
+	{
+		const FName InstanceName = MakeUniqueObjectName(
+			GetTransientPackage(),
+			UMaterialInstance::StaticClass(),
+			*IUsdClassesModule::SanitizeObjectName(FPaths::GetBaseFilename(VolumePrimPathString))
+		);
+
+		MaterialInstance = InstantiateMaterial(InstanceName, ReferenceMaterial, Context.Get());
+
+		if (MaterialInstance)
+		{
+			UUsdAssetUserData* UserData = NewObject<UUsdAssetUserData>(MaterialInstance, TEXT("USDAssetUserData"));
+			UserData->PrimPaths = {PrimPath.GetString()};
+			MaterialInstance->AddAssetUserData(UserData);
+
+			AssignMaterialParameters(MaterialInstance, MaterialParameterToTexture);
+
+			Context->AssetCache->CacheAsset(PrefixedMaterialHash, MaterialInstance);
+		}
+	}
+
+	if (MaterialInstance)
+	{
+		Context->InfoCache->LinkAssetToPrim(PrimPath, MaterialInstance);
+
+		if (UUsdAssetUserData* UserData = UsdUtils::GetOrCreateAssetUserData<UUsdAssetUserData>(MaterialInstance))
+		{
+			UserData->PrimPaths.AddUnique(VolumePrimPathString);
+
+			if (Context->MetadataOptions.bCollectMetadata)
+			{
+				UsdToUnreal::ConvertMetadata(
+					VolumePrim,
+					UserData,
+					Context->MetadataOptions.BlockedPrefixFilters,
+					Context->MetadataOptions.bInvertFilters,
+					Context->MetadataOptions.bCollectFromEntireSubtrees
+				);
+			}
+			else
+			{
+				UserData->StageIdentifierToMetadata.Remove(UsdToUnreal::ConvertString(Stage->GetRootLayer()->GetIdentifier()));
+			}
+		}
+	}
+#endif	  // WITH_EDITOR
+}
+
+USceneComponent* FUsdVolVolumeTranslator::CreateComponents()
+{
+	TRACE_CPUPROFILER_EVENT_SCOPE(FUsdVolVolumeTranslator::CreateComponents);
+
+	USceneComponent* SceneComponent = nullptr;
+
+#if WITH_EDITOR
+	EUsdDrawMode DrawMode = UsdUtils::GetAppliedDrawMode(GetPrim());
+	if (DrawMode == EUsdDrawMode::Default)
+	{
+		if (Context->bAllowParsingSparseVolumeTextures)
+		{
+			const bool bNeedsActor = true;
+			SceneComponent = CreateComponentsEx({UHeterogeneousVolumeComponent::StaticClass()}, bNeedsActor);
+		}
+	}
+	else
+	{
+		SceneComponent = CreateAlternativeDrawModeComponents(DrawMode);
+	}
+
+	UpdateComponents(SceneComponent);
+#endif	  // WITH_EDITOR
+
+	return SceneComponent;
+}
+
+void FUsdVolVolumeTranslator::UpdateComponents(USceneComponent* SceneComponent)
+{
+#if WITH_EDITOR
+	// Set volumetric material onto the spawned component
+	if (UHeterogeneousVolumeComponent* VolumeComponent = Cast<UHeterogeneousVolumeComponent>(SceneComponent))
+	{
+		const int32 ElementIndex = 0;
+		UMaterialInterface* CurrentMaterial = VolumeComponent->GetMaterial(ElementIndex);
+
+		if (UMaterialInstance* MaterialForPrim = Context->InfoCache->GetSingleAssetForPrim<UMaterialInstance>(PrimPath))
+		{
+			if (MaterialForPrim != CurrentMaterial)
+			{
+				// We need to call PostLoad here or else it won't render the material properly (reference:
+				// SNiagaraVolumeTextureViewport::Construct)
+				VolumeComponent->SetMaterial(ElementIndex, MaterialForPrim);
+				VolumeComponent->PostLoad();
+
+				CurrentMaterial = MaterialForPrim;
+			}
+		}
+
+		// Animate first SVT parameter if we have an animated one
+		if (CurrentMaterial && !Context->bSequencerIsAnimating)
+		{
+			TArray<FMaterialParameterInfo> ParameterInfo;
+			TArray<FGuid> ParameterIds;
+			CurrentMaterial->GetAllSparseVolumeTextureParameterInfo(ParameterInfo, ParameterIds);
+
+			for (const FMaterialParameterInfo& Info : ParameterInfo)
+			{
+				USparseVolumeTexture* SparseVolumeTexture = nullptr;
+				if (CurrentMaterial->GetSparseVolumeTextureParameterValue(Info, SparseVolumeTexture) && SparseVolumeTexture)
+				{
+					if (SparseVolumeTexture->GetNumFrames() > 1)
+					{
+						if (UUsdSparseVolumeTextureAssetUserData* UserData = Cast<UUsdSparseVolumeTextureAssetUserData>(
+								UsdUtils::GetAssetUserData(SparseVolumeTexture)
+							))
+						{
+							UE::FUsdPrim VolumePrim = GetPrim();
+							UE::FUsdStage Stage = VolumePrim.GetStage();
+
+							UE::FUsdPrim PrimForOffsetCalculation = VolumePrim;
+							if (UserData->SourceOpenVDBAssetPrimPaths.Num() > 0)
+							{
+								const FString& FirstAssetPrimPath = UserData->SourceOpenVDBAssetPrimPaths[0];
+								UE::FUsdPrim FirstAssetPrim = Stage.GetPrimAtPath(UE::FSdfPath{*FirstAssetPrimPath});
+								if (FirstAssetPrim)
+								{
+									PrimForOffsetCalculation = FirstAssetPrim;
+								}
+							}
+
+							UE::FSdfLayerOffset CombinedOffset = UsdUtils::GetPrimToStageOffset(PrimForOffsetCalculation);
+							const double LayerTimeCode = ((Context->Time - CombinedOffset.Offset) / CombinedOffset.Scale);
+
+							// The SVTs will have all the volume frames packed next to each other with no time information,
+							// and are indexed by a "frame index" where 0 is the first frame and N-1 is the last frame.
+							// These is also no linear interpolation: The frame index is basically floor()'d and the integer
+							// value is used as the index into the Frames array
+							int32 TargetIndex = 0;
+							for (; TargetIndex + 1 < UserData->TimeSamplePathTimeCodes.Num(); ++TargetIndex)
+							{
+								if (UserData->TimeSamplePathTimeCodes[TargetIndex + 1] > LayerTimeCode)
+								{
+									break;
+								}
+							}
+							TargetIndex = FMath::Clamp(TargetIndex, 0, UserData->TimeSamplePathTimeCodes.Num() - 1);
+
+							// At this point TargetIndex points at the index of the biggest timeCode that is
+							// still <= LayerTimeCode. We may have an index mapping though, like when the bRemoveDuplicates
+							// cvar is true
+
+							if (UserData->TimeSamplePathIndices.IsValidIndex(TargetIndex))
+							{
+								TargetIndex = UserData->TimeSamplePathIndices[TargetIndex];
+							}
+
+							// Now TargetIndex should be pointing at the index of the desired frame within the SVT
+
+							const bool bPlaying = false;
+							VolumeComponent->SetPlaying(bPlaying);
+							VolumeComponent->SetFrame(static_cast<float>(TargetIndex));
+						}
+					}
+				}
+
+				break;
+			}
+		}
+	}
+#endif	  // WITH_EDITOR
+
+	Super::UpdateComponents(SceneComponent);
+}
+
+bool FUsdVolVolumeTranslator::CollapsesChildren(ECollapsingType CollapsingType) const
+{
+	// If we have a custom draw mode, it means we should draw bounds/cards/etc. instead
+	// of our entire subtree, which is basically the same thing as collapsing
+	EUsdDrawMode DrawMode = UsdUtils::GetAppliedDrawMode(GetPrim());
+	if (DrawMode != EUsdDrawMode::Default)
+	{
+		return true;
+	}
+
+	return false;
+}
+
+bool FUsdVolVolumeTranslator::CanBeCollapsed(ECollapsingType CollapsingType) const
+{
+	return false;
+}
+
+TSet<UE::FSdfPath> FUsdVolVolumeTranslator::CollectAuxiliaryPrims() const
+{
+	if (!Context->bIsBuildingInfoCache)
+	{
+		return Context->InfoCache->GetAuxiliaryPrims(PrimPath);
+	}
+
+	TSet<UE::FSdfPath> Result;
+#if WITH_EDITOR
+	{
+		FScopedUsdAllocs UsdAllocs;
+
+		if (pxr::UsdVolVolume Volume{GetPrim()})
+		{
+			const std::map<pxr::TfToken, pxr::SdfPath>& FieldMap = Volume.GetFieldPaths();
+			for (std::map<pxr::TfToken, pxr::SdfPath>::const_iterator Iter = FieldMap.cbegin(); Iter != FieldMap.end(); ++Iter)
+			{
+				const pxr::SdfPath& AssetPrimPath = Iter->second;
+
+				Result.Add(UE::FSdfPath{*UsdToUnreal::ConvertPath(AssetPrimPath)});
+			}
+		}
+	}
+#endif	  // WITH_EDITOR
+	return Result;
+}
+
+#endif	  // #if USE_USD_SDK
diff --git a/Engine/Plugins/Importers/USDImporter/Source/USDSchemas/Public/USDSchemaTranslator.h b/Engine/Plugins/Importers/USDImporter/Source/USDSchemas/Public/USDSchemaTranslator.h
--- a/Engine/Plugins/Importers/USDImporter/Source/USDSchemas/Public/USDSchemaTranslator.h
+++ b/Engine/Plugins/Importers/USDImporter/Source/USDSchemas/Public/USDSchemaTranslator.h
@@ -283,6 +283,9 @@ struct USDSCHEMAS_API FUsdSchemaTranslationContext : public TSharedFromThis<FUsd
 	/** If true, means we will try generating GroomAssets, GroomCaches and GroomBindings */
 	bool bAllowParsingGroomAssets = true;
 
+	/** If true, means we will try generating Sparse Volume Textures */
+	bool bAllowParsingSparseVolumeTextures = true;
+
 	/** Skip the import of materials that aren't being used by any prim on the stage */
 	bool bTranslateOnlyUsedMaterials = false;
 
diff --git a/Engine/Plugins/Importers/USDImporter/Source/USDSchemas/Public/USDVolVolumeTranslator.h b/Engine/Plugins/Importers/USDImporter/Source/USDSchemas/Public/USDVolVolumeTranslator.h
new file mode 100644
--- /dev/null
+++ b/Engine/Plugins/Importers/USDImporter/Source/USDSchemas/Public/USDVolVolumeTranslator.h
@@ -0,0 +1,29 @@
+// Copyright Epic Games, Inc. All Rights Reserved.
+
+#pragma once
+
+#include "USDGeomXformableTranslator.h"
+
+#if USE_USD_SDK
+
+class USceneComponent;
+
+class USDSCHEMAS_API FUsdVolVolumeTranslator : public FUsdGeomXformableTranslator
+{
+	using Super = FUsdGeomXformableTranslator;
+
+public:
+
+	using FUsdGeomXformableTranslator::FUsdGeomXformableTranslator;
+
+	virtual void CreateAssets() override;
+	virtual USceneComponent* CreateComponents() override;
+	virtual void UpdateComponents(USceneComponent* SceneComponent) override;
+
+	virtual bool CollapsesChildren(ECollapsingType CollapsingType) const override;
+	virtual bool CanBeCollapsed(ECollapsingType CollapsingType) const override;
+
+	virtual TSet<UE::FSdfPath> CollectAuxiliaryPrims() const override;
+};
+
+#endif	  // #if USE_USD_SDK
diff --git a/Engine/Plugins/Importers/USDImporter/Source/USDSchemas/USDSchemas.Build.cs b/Engine/Plugins/Importers/USDImporter/Source/USDSchemas/USDSchemas.Build.cs
--- a/Engine/Plugins/Importers/USDImporter/Source/USDSchemas/USDSchemas.Build.cs
+++ b/Engine/Plugins/Importers/USDImporter/Source/USDSchemas/USDSchemas.Build.cs
@@ -71,6 +71,7 @@ namespace UnrealBuildTool.Rules
 						"MeshUtilities",
 						"PhysicsUtilities", // For generating UPhysicsAssets for SkeletalMeshes and ConvexDecompTool
 						"PropertyEditor",
+						"SparseVolumeTexture",
 						"UnrealEd",
 					}
 				);
diff --git a/Engine/Plugins/Importers/USDImporter/Source/USDStage/Private/USDLevelSequenceHelper.cpp b/Engine/Plugins/Importers/USDImporter/Source/USDStage/Private/USDLevelSequenceHelper.cpp
--- a/Engine/Plugins/Importers/USDImporter/Source/USDStage/Private/USDLevelSequenceHelper.cpp
+++ b/Engine/Plugins/Importers/USDImporter/Source/USDStage/Private/USDLevelSequenceHelper.cpp
@@ -36,6 +36,7 @@
 #include "CineCameraComponent.h"
 #include "Compilation/MovieSceneCompiledDataManager.h"
 #include "Components/DirectionalLightComponent.h"
+#include "Components/HeterogeneousVolumeComponent.h"
 #include "Components/LightComponent.h"
 #include "Components/LightComponentBase.h"
 #include "Components/PointLightComponent.h"
@@ -68,6 +69,7 @@
 #include "Sections/MovieSceneSubSection.h"
 #include "Sequencer/MovieSceneControlRigParameterSection.h"
 #include "Sequencer/MovieSceneControlRigParameterTrack.h"
+#include "SparseVolumeTexture/SparseVolumeTexture.h"
 #include "Templates/SharedPointer.h"
 #include "Tracks/MovieScene3DTransformTrack.h"
 #include "Tracks/MovieSceneBoolTrack.h"
@@ -381,13 +383,14 @@ namespace UsdLevelSequenceHelperImpl
 			// of the playback range, which we do here
 			ControlRigSectionStartFrame -= UE::MovieScene::DiscreteInclusiveLower(MovieScene->GetPlaybackRange());
 
+			const bool bResetControls = true;
 			ParamSection->LoadAnimSequenceIntoThisSection(
 				AnimSequence,
 				MovieScene,
 				SkeletalMeshComp,
 				bReduceKeys,
 				Tolerance,
-				true,
+				bResetControls,
 				ControlRigSectionStartFrame,
 				EMovieSceneKeyInterpolation::SmartAuto
 			);
@@ -638,6 +641,7 @@ private:
 	void AddSkeletalTracks(const UUsdPrimTwin& PrimTwin, const UE::FUsdPrim& Prim);
 	void AddGroomTracks(const UUsdPrimTwin& PrimTwin, const UE::FUsdPrim& Prim);
 	void AddGeometryCacheTracks(const UUsdPrimTwin& PrimTwin, const UE::FUsdPrim& Prim);
+	void AddVolumeTracks(const UUsdPrimTwin& PrimTwin, const UE::FUsdPrim& Prim);
 
 	template<typename TrackType>
 	TrackType* AddTrack(
@@ -2378,6 +2382,175 @@ void FUsdLevelSequenceHelperImpl::AddGroomTracks(const UUsdPrimTwin& PrimTwin, c
 	PrimPathByLevelSequenceName.AddUnique(GroomAnimationSequence->GetFName(), PrimPath);
 }
 
+void FUsdLevelSequenceHelperImpl::AddVolumeTracks(const UUsdPrimTwin& PrimTwin, const UE::FUsdPrim& Prim)
+{
+	using namespace UnrealIdentifiers;
+
+	UHeterogeneousVolumeComponent* VolumeComponent = Cast<UHeterogeneousVolumeComponent>(PrimTwin.GetSceneComponent());
+	if (!VolumeComponent)
+	{
+		return;
+	}
+
+	UE::FSdfLayer PrimLayer = UsdUtils::FindLayerForPrim(Prim);
+	ULevelSequence* PrimSequence = FindSequenceForIdentifier(PrimLayer.GetIdentifier());
+	if (!PrimSequence)
+	{
+		return;
+	}
+
+	UMovieScene* MovieScene = PrimSequence->GetMovieScene();
+	if (!MovieScene)
+	{
+		return;
+	}
+
+	const FName PropertyPath = GET_MEMBER_NAME_CHECKED(UHeterogeneousVolumeComponent, Frame);
+
+	// Here we'll just get *any* of the filePath attrs from this Volume prim to check for the
+	// muted bool. We won't use the attribute itself for the baking though, as our timeSamples are already
+	// on the SparseVolumeTexture AssetUserData, and the keyframe values are just their indices
+	TArray<UE::FUsdAttribute> Attrs = UnrealToUsd::GetAttributesForProperty(Prim, PropertyPath);
+	if (Attrs.Num() < 1)
+	{
+		return;
+	}
+	UE::FUsdAttribute MainAttr = Attrs[0];
+	const bool bIsMuted = MainAttr && MainAttr.GetNumTimeSamples() > 0 && UsdUtils::IsAttributeMuted(MainAttr, UsdStage);
+
+	UUsdSparseVolumeTextureAssetUserData* UserData = nullptr;
+	const int32 ElementIndex = 0;
+	UMaterialInterface* CurrentMaterial = VolumeComponent->GetMaterial(ElementIndex);
+	if (CurrentMaterial)
+	{
+		TArray<FMaterialParameterInfo> ParameterInfo;
+		TArray<FGuid> ParameterIds;
+		CurrentMaterial->GetAllSparseVolumeTextureParameterInfo(ParameterInfo, ParameterIds);
+
+		for (const FMaterialParameterInfo& Info : ParameterInfo)
+		{
+			USparseVolumeTexture* SparseVolumeTexture = nullptr;
+			if (CurrentMaterial->GetSparseVolumeTextureParameterValue(Info, SparseVolumeTexture) && SparseVolumeTexture)
+			{
+				if (SparseVolumeTexture->GetNumFrames() > 1)
+				{
+					UserData = Cast<UUsdSparseVolumeTextureAssetUserData>(UsdUtils::GetAssetUserData(SparseVolumeTexture));
+				}
+			}
+
+			// Follow the theme of only ever caring about the first SVT parameter of the material, as that is all
+			// that the UHeterogeneousVolumeComponent will ever animate anyway. If we're in here we already know
+			// that this first SVT should be animated at any case
+			break;
+		}
+	}
+	if (!UserData)
+	{
+		return;
+	}
+
+	TArray<double>* TimeSamples = &UserData->TimeSamplePathTimeCodes;
+	TArray<int32>* FrameIndices = nullptr;
+	if (UserData->TimeSamplePathIndices.Num() == UserData->TimeSamplePathTimeCodes.Num())
+	{
+		FrameIndices = &UserData->TimeSamplePathIndices;
+	}
+	else
+	{
+		UE_LOG(
+			LogUsd,
+			Warning,
+			TEXT(
+				"Ignoring AssetUserData TimeSamplePathIndices when generating Sequencer tracks for Prim '%s' because it has %d entries, while it should have the same number of entries as TimeSamplePathTimeCodes (%d)"
+			),
+			*Prim.GetPrimPath().GetString(),
+			UserData->TimeSamplePathIndices.Num(),
+			UserData->TimeSamplePathTimeCodes.Num()
+		);
+	}
+
+	if (!TimeSamples || TimeSamples->Num() < 2)
+	{
+		return;
+	}
+
+	// Our TimeSamples are local to the layer where they were defined, but we need to convert them
+	// to be with respect to the stage in order to find the right locations for the key frames
+	UE::FUsdPrim PrimForOffsetCalculation = Prim;
+	if (UserData->SourceOpenVDBAssetPrimPaths.Num() > 0)
+	{
+		const FString& FirstAssetPrimPath = UserData->SourceOpenVDBAssetPrimPaths[0];
+		UE::FUsdPrim FirstAssetPrim = Prim.GetStage().GetPrimAtPath(UE::FSdfPath{*FirstAssetPrimPath});
+		if (FirstAssetPrim)
+		{
+			PrimForOffsetCalculation = FirstAssetPrim;
+		}
+	}
+	UE::FSdfLayerOffset CombinedOffset = UsdUtils::GetPrimToStageOffset(UE::FUsdPrim{PrimForOffsetCalculation});
+	TArray<double> ConvertedTimeSamples;
+	ConvertedTimeSamples.Reserve(TimeSamples->Num());
+	for (double TimeSample : *TimeSamples)
+	{
+		ConvertedTimeSamples.Add(TimeSample * CombinedOffset.Scale + CombinedOffset.Offset);
+	}
+
+	// We still have a Sequence transform though, as that converts the TimeSamples from being global
+	// to the stage to the particular subsequence where they are going to be added to.
+	//
+	// This may seem like it undoes the calculation in ConvertedTimeSamples, and it really does: For cases
+	// where we have a sublayer with an offset and scale, we'll end up adding the layer-local time samples
+	// to the subsequence, like we want. Using the CombinedOffset AND the SequenceTransform is needed for
+	// a different case however: Prim references with sublayer and offsets. In that case the prim itself
+	// may have a sublayer and offset, but its track will be placed in the levelsequence for the *referencer*
+	// layer: This means we want to see the keys on that layer instead, at times relative to it
+	FMovieSceneSequenceTransform SequenceTransform;
+	FMovieSceneSequenceID SequenceID = SequencesID.FindRef(PrimSequence);
+	if (FMovieSceneSubSequenceData* SubSequenceData = SequenceHierarchyCache.FindSubData(SequenceID))
+	{
+		SequenceTransform = SubSequenceData->RootToSequenceTransform;
+	}
+
+	// Unlike the other cases we can create our Reader right here, because the only thing we need to generate
+	// the track are the TimeSamples
+	UsdToUnreal::FPropertyTrackReader Reader;
+	Reader.FloatReader = [&ConvertedTimeSamples, FrameIndices, TargetIndex = 0](double UsdTimeCode) mutable -> float
+	{
+		// Reference: FUsdVolVolumeTranslator::UpdateComponents
+		for (; TargetIndex + 1 < ConvertedTimeSamples.Num(); ++TargetIndex)
+		{
+			if (ConvertedTimeSamples[TargetIndex + 1] > UsdTimeCode)
+			{
+				break;
+			}
+		}
+		TargetIndex = FMath::Clamp(TargetIndex, 0, ConvertedTimeSamples.Num() - 1);
+
+		if (FrameIndices && FrameIndices->IsValidIndex(TargetIndex))
+		{
+			TargetIndex = (*FrameIndices)[TargetIndex];
+		}
+
+		return static_cast<float>(TargetIndex);
+	};
+
+	if (UMovieSceneFloatTrack* FloatTrack = AddTrack<UMovieSceneFloatTrack>(PropertyPath, PrimTwin, *VolumeComponent, *PrimSequence, bIsMuted))
+	{
+		// The component won't really linearly interpolate anything and will just do the analogous of constant interpolation,
+		// so it would be nice if our keys showed that too
+		ERichCurveInterpMode InterpolationModeOverride = ERichCurveInterpMode::RCIM_Constant;
+		UsdToUnreal::ConvertFloatTimeSamples(
+			UsdStage,
+			ConvertedTimeSamples,
+			Reader.FloatReader,
+			*FloatTrack,
+			SequenceTransform,
+			InterpolationModeOverride
+		);
+	}
+
+	PrimPathByLevelSequenceName.AddUnique(PrimSequence->GetFName(), Prim.GetPrimPath().GetString());
+}
+
 void FUsdLevelSequenceHelperImpl::AddPrim(UUsdPrimTwin& PrimTwin, bool bForceVisibilityTracks, TOptional<bool> HasAnimatedBounds)
 {
 	if (!UsdStage)
@@ -2442,6 +2615,10 @@ void FUsdLevelSequenceHelperImpl::AddPrim(UUsdPrimTwin& PrimTwin, bool bForceVis
 	{
 		AddGroomTracks(PrimTwin, UsdPrim);
 	}
+	else if (UsdPrim.IsA(TEXT("Volume")))
+	{
+		AddVolumeTracks(PrimTwin, UsdPrim);
+	}
 
 	AddCommonTracks(PrimTwin, UsdPrim, bForceVisibilityTracks);
 
@@ -3882,6 +4059,16 @@ void FUsdLevelSequenceHelperImpl::HandleTrackChange(const UMovieSceneTrack& Trac
 
 				if (const UMovieSceneFloatTrack* FloatTrack = Cast<const UMovieSceneFloatTrack>(&Track))
 				{
+					// We won't need a SequenceTransform in this case because the FloatWriter will be ready to receive and write
+					// keyframes local to its own sequence/layer
+					if (const UHeterogeneousVolumeComponent* VolumeComponent = Cast<const UHeterogeneousVolumeComponent>(BoundSceneComponent))
+					{
+						if (PropertyTrack->GetPropertyName() == GET_MEMBER_NAME_CHECKED(UHeterogeneousVolumeComponent, Frame))
+						{
+							SequenceTransform = {};
+						}
+					}
+
 					UnrealToUsd::ConvertFloatTrack(*FloatTrack, SequenceTransform, Writer.FloatWriter, UsdPrim);
 				}
 				else if (const UMovieSceneBoolTrack* BoolTrack = Cast<const UMovieSceneBoolTrack>(&Track))
diff --git a/Engine/Plugins/Importers/USDImporter/Source/USDStageImporter/Private/USDStageAssetImportFactory.cpp b/Engine/Plugins/Importers/USDImporter/Source/USDStageImporter/Private/USDStageAssetImportFactory.cpp
--- a/Engine/Plugins/Importers/USDImporter/Source/USDStageImporter/Private/USDStageAssetImportFactory.cpp
+++ b/Engine/Plugins/Importers/USDImporter/Source/USDStageImporter/Private/USDStageAssetImportFactory.cpp
@@ -117,8 +117,20 @@ bool UUsdStageAssetImportFactory::CanReimport(UObject* Obj, TArray<FString>& Out
 {
 	if (UAssetImportData* ImportData = UsdUtils::GetAssetImportData(Obj))
 	{
-		OutFilenames.Add(ImportData->GetFirstFilename());
-		return true;
+		const FString FileName = ImportData->GetFirstFilename();
+		const FString FileExtension = FPaths::GetExtension(FileName);
+
+		// Reimporting from here means opening FileName as a USD stage and trying to re-read the same prims,
+		// so make sure we only claim we can reimport something if that would work. Otherwise we may intercept
+		// some other formats like .vdb files and then fail to open them as stages
+		for (const FString& Extension : UnrealUSDWrapper::GetNativeFileFormats())
+		{
+			if (Extension == FileExtension)
+			{
+				OutFilenames.Add(ImportData->GetFirstFilename());
+				return true;
+			}
+		}
 	}
 
 	return false;
diff --git a/Engine/Plugins/Importers/USDImporter/Source/USDStageImporter/Private/USDStageImportOptions.cpp b/Engine/Plugins/Importers/USDImporter/Source/USDStageImporter/Private/USDStageImportOptions.cpp
--- a/Engine/Plugins/Importers/USDImporter/Source/USDStageImporter/Private/USDStageImportOptions.cpp
+++ b/Engine/Plugins/Importers/USDImporter/Source/USDStageImporter/Private/USDStageImportOptions.cpp
@@ -18,6 +18,7 @@ UUsdStageImportOptions::UUsdStageImportOptions(const FObjectInitializer& ObjectI
 	bImportLevelSequences = true;
 	bImportMaterials = true;
 	bImportGroomAssets = true;
+	bImportSparseVolumeTextures = true;
 	bImportOnlyUsedMaterials = false;
 
 	PurposesToImport = (int32)(EUsdPurpose::Default | EUsdPurpose::Proxy | EUsdPurpose::Render | EUsdPurpose::Guide);
@@ -93,6 +94,7 @@ void UsdUtils::AddAnalyticsAttributes(const UUsdStageImportOptions& Options, TAr
 	InOutAttributes.Emplace(TEXT("ImportLevelSequences"), LexToString(Options.bImportLevelSequences));
 	InOutAttributes.Emplace(TEXT("ImportMaterials"), LexToString(Options.bImportMaterials));
 	InOutAttributes.Emplace(TEXT("ImportGroomAssets"), LexToString(Options.bImportGroomAssets));
+	InOutAttributes.Emplace(TEXT("ImportSparseVolumeTextures"), LexToString(Options.bImportSparseVolumeTextures));
 	InOutAttributes.Emplace(TEXT("ImportOnlyUsedMaterials"), LexToString(Options.bImportOnlyUsedMaterials));
 	if (Options.PrimsToImport != TArray<FString>{TEXT("/")})
 	{
diff --git a/Engine/Plugins/Importers/USDImporter/Source/USDStageImporter/Private/USDStageImporter.cpp b/Engine/Plugins/Importers/USDImporter/Source/USDStageImporter/Private/USDStageImporter.cpp
--- a/Engine/Plugins/Importers/USDImporter/Source/USDStageImporter/Private/USDStageImporter.cpp
+++ b/Engine/Plugins/Importers/USDImporter/Source/USDStageImporter/Private/USDStageImporter.cpp
@@ -55,6 +55,7 @@
 #include "PackageTools.h"
 #include "PhysicsEngine/PhysicsAsset.h"
 #include "Serialization/ArchiveReplaceObjectRef.h"
+#include "SparseVolumeTexture/SparseVolumeTexture.h"
 #include "Subsystems/AssetEditorSubsystem.h"
 #include "UObject/UObjectIterator.h"
 
@@ -613,6 +614,13 @@ namespace UE::USDStageImporter::Private
 
 			AssetPrefix = TEXT("ABP_");
 		}
+		else if (USparseVolumeTexture* SparseVolumeTexture = Cast<USparseVolumeTexture>(Asset))
+		{
+			AssetPrefix = TEXT("SVT_");
+
+			// Same situation as the UTexture case
+			AssetName = SparseVolumeTexture->GetFName().GetPlainNameString();
+		}
 
 		if (!AssetName.StartsWith(AssetPrefix))
 		{
@@ -645,7 +653,7 @@ namespace UE::USDStageImporter::Private
 		if (ImportData)
 		{
 			// Don't force update textures as they will already have this preset to their actual texture path
-			if (!Asset->IsA<UTexture>())
+			if (!Asset->IsA<UTexture>() && !Asset->IsA<USparseVolumeTexture>())
 			{
 				ImportData->UpdateFilenameOnly(MainFilePath);
 			}
@@ -924,6 +932,7 @@ namespace UE::USDStageImporter::Private
 		TArray<UObject*> Grooms;
 		TArray<UObject*> GroomCaches;
 		TArray<UObject*> GroomBindings;
+		TArray<UObject*> SparseVolumeTextures;
 
 		TSet<FString> UniqueAssetNames;
 		TMap<UObject*, FString> AssetToContentFolder;
@@ -1016,6 +1025,11 @@ namespace UE::USDStageImporter::Private
 				}
 				GroomBindings.Add(Asset);
 			}
+			else if (Asset->IsA(USparseVolumeTexture::StaticClass()))
+			{
+				AssetTypeFolderPtr = &TexturesFolder;
+				SparseVolumeTextures.Add(Asset);
+			}
 			else
 			{
 				// We don't know what to do with this asset
@@ -1083,6 +1097,7 @@ namespace UE::USDStageImporter::Private
 		PublishAssetType(GeometryCaches);
 		PublishAssetType(Materials);
 		PublishAssetType(Textures);
+		PublishAssetType(SparseVolumeTextures);
 	}
 
 	void ResolveComponentConflict(
@@ -1697,6 +1712,7 @@ namespace UE::USDStageImporter::Private
 			int32 NumGroomAssets = 0;
 			int32 NumGroomBindings = 0;
 			int32 NumGroomCaches = 0;
+			int32 NumSparseVolumeTextures = 0;
 			for (UObject* ImportedAsset : ImportedAssets)
 			{
 				if (!ImportedAsset)
@@ -1744,6 +1760,10 @@ namespace UE::USDStageImporter::Private
 				{
 					++NumGroomCaches;
 				}
+				else if (ImportedAsset->IsA<USparseVolumeTexture>())
+				{
+					++NumSparseVolumeTextures;
+				}
 			}
 			EventAttributes.Emplace(TEXT("NumStaticMeshes"), LexToString(NumStaticMeshes));
 			EventAttributes.Emplace(TEXT("NumSkeletalMeshes"), LexToString(NumSkeletalMeshes));
@@ -1755,6 +1775,7 @@ namespace UE::USDStageImporter::Private
 			EventAttributes.Emplace(TEXT("NumGroomAssets"), LexToString(NumGroomAssets));
 			EventAttributes.Emplace(TEXT("NumGroomBindings"), LexToString(NumGroomBindings));
 			EventAttributes.Emplace(TEXT("NumGroomCaches"), LexToString(NumGroomCaches));
+			EventAttributes.Emplace(TEXT("NumSparseVolumeTextures"), LexToString(NumSparseVolumeTextures));
 
 			FString RootLayerIdentifier = ImportContext.FilePath;
 			if (ImportContext.Stage)
@@ -1801,19 +1822,20 @@ namespace UE::USDStageImporter::Private
 				continue;
 			}
 
-			if ((!ImportContext.ImportOptions->bImportGeometry
-				 && (Asset->IsA<UStaticMesh>() || Asset->IsA<USkeletalMesh>() || Asset->IsA<USkeleton>() || Asset->IsA<UPhysicsAsset>()
-					 || Asset->IsA<UGeometryCache>()))
-				|| (!bImportSkeletalAnimations && (Asset->IsA<UAnimSequence>()))
-				|| (!ImportContext.ImportOptions->bImportLevelSequences && (Asset->IsA<ULevelSequence>()))
-				|| (!ImportContext.ImportOptions->bImportMaterials && (Asset->IsA<UMaterialInterface>() || Asset->IsA<UTexture>()))
-				|| (!ImportContext.ImportOptions->bImportGroomAssets
-					&& (Asset->IsA<UGroomAsset>() || Asset->IsA<UGroomCache>() || Asset->IsA<UGroomBindingAsset>())))
+			// clang-format off
+			if ((!ImportContext.ImportOptions->bImportGeometry && (Asset->IsA<UStaticMesh>() || Asset->IsA<USkeletalMesh>() || Asset->IsA<USkeleton>() || Asset->IsA<UPhysicsAsset>() || Asset->IsA<UGeometryCache>())) ||
+				(!bImportSkeletalAnimations && (Asset->IsA<UAnimSequence>())) ||
+				(!ImportContext.ImportOptions->bImportLevelSequences && (Asset->IsA<ULevelSequence>())) ||
+				(!ImportContext.ImportOptions->bImportMaterials && (Asset->IsA<UMaterialInterface>() || Asset->IsA<UTexture>())) ||
+				(!ImportContext.ImportOptions->bImportGroomAssets && (Asset->IsA<UGroomAsset>() || Asset->IsA<UGroomCache>() || Asset->IsA<UGroomBindingAsset>())) ||
+				(!ImportContext.ImportOptions->bImportSparseVolumeTextures && (Asset->IsA<USparseVolumeTexture>()))
+			)
 			{
 				ObjectsToRemap.Add(Asset, nullptr);
 				SoftObjectsToRemap.Add(Asset, nullptr);
 				It.RemoveCurrent();
 			}
+			// clang-format on
 		}
 	}
 
@@ -2018,6 +2040,7 @@ void UUsdStageImporter::ImportFromFile(FUsdStageImportContext& ImportContext)
 	TranslationContext->bAllowParsingSkeletalAnimations = ImportContext.ImportOptions->bImportGeometry
 														  && ImportContext.ImportOptions->bImportSkeletalAnimations;
 	TranslationContext->bAllowParsingGroomAssets = ImportContext.ImportOptions->bImportGroomAssets;
+	TranslationContext->bAllowParsingSparseVolumeTextures = ImportContext.ImportOptions->bImportSparseVolumeTextures;
 	TranslationContext->bTranslateOnlyUsedMaterials = ImportContext.ImportOptions->bImportOnlyUsedMaterials;
 	TranslationContext->InfoCache = InfoCache;
 	TranslationContext->BBoxCache = ImportContext.BBoxCache;
@@ -2201,6 +2224,7 @@ bool UUsdStageImporter::ReimportSingleAsset(
 	TranslationContext->bAllowParsingSkeletalAnimations = ImportContext.ImportOptions->bImportGeometry
 														  && ImportContext.ImportOptions->bImportSkeletalAnimations;
 	TranslationContext->bAllowParsingGroomAssets = ImportContext.ImportOptions->bImportGroomAssets;
+	TranslationContext->bAllowParsingSparseVolumeTextures = ImportContext.ImportOptions->bImportSparseVolumeTextures;
 	TranslationContext->bTranslateOnlyUsedMaterials = ImportContext.ImportOptions->bImportOnlyUsedMaterials;
 	TranslationContext->InfoCache = InfoCache;
 	TranslationContext->BBoxCache = ImportContext.BBoxCache;
diff --git a/Engine/Plugins/Importers/USDImporter/Source/USDStageImporter/Public/USDStageImportOptions.h b/Engine/Plugins/Importers/USDImporter/Source/USDStageImporter/Public/USDStageImportOptions.h
--- a/Engine/Plugins/Importers/USDImporter/Source/USDStageImporter/Public/USDStageImportOptions.h
+++ b/Engine/Plugins/Importers/USDImporter/Source/USDStageImporter/Public/USDStageImportOptions.h
@@ -74,6 +74,10 @@ public:
 	UPROPERTY(BlueprintReadWrite, config, EditAnywhere, Category = "DataToImport", meta = (DisplayName = "Groom Assets"))
 	bool bImportGroomAssets;
 
+	/** Whether to import OpenVDB volumes as Sparse Volume Textures */
+	UPROPERTY(BlueprintReadWrite, config, EditAnywhere, Category = "DataToImport", meta = (DisplayName = "Sparse Volume Textures"))
+	bool bImportSparseVolumeTextures;
+
 	/**
 	 * If this is checked, only materials actively used by the stage and import settings will be parsed.
 	 * If this is unchecked, all materials present on the stage will be parsed.
diff --git a/Engine/Plugins/Importers/USDImporter/Source/USDUtilities/Private/USDConversionUtils.cpp b/Engine/Plugins/Importers/USDImporter/Source/USDUtilities/Private/USDConversionUtils.cpp
--- a/Engine/Plugins/Importers/USDImporter/Source/USDUtilities/Private/USDConversionUtils.cpp
+++ b/Engine/Plugins/Importers/USDImporter/Source/USDUtilities/Private/USDConversionUtils.cpp
@@ -27,6 +27,7 @@
 #include "CineCameraActor.h"
 #include "CineCameraComponent.h"
 #include "Components/DirectionalLightComponent.h"
+#include "Components/HeterogeneousVolumeComponent.h"
 #include "Components/PointLightComponent.h"
 #include "Components/RectLightComponent.h"
 #include "Components/SkeletalMeshComponent.h"
@@ -49,6 +50,7 @@
 #include "LandscapeProxy.h"
 #include "Misc/PackageName.h"
 #include "PhysicsEngine/PhysicsAsset.h"
+#include "SparseVolumeTexture/SparseVolumeTexture.h"
 #include "Widgets/Notifications/SNotificationList.h"
 
 #if USE_USD_SDK
@@ -83,6 +85,8 @@
 #include "pxr/usd/usdSkel/cache.h"
 #include "pxr/usd/usdSkel/root.h"
 #include "pxr/usd/usdSkel/skeletonQuery.h"
+#include "pxr/usd/usdVol/openVDBAsset.h"
+#include "pxr/usd/usdVol/volume.h"
 #include "USDIncludesEnd.h"
 
 #include <string>
@@ -465,6 +469,10 @@ UClass* UsdUtils::GetActorTypeForPrim(const pxr::UsdPrim& Prim)
 	{
 		return ASkyLight::StaticClass();
 	}
+	else if (Prim.IsA<pxr::UsdVolVolume>())
+	{
+		return AHeterogeneousVolume::StaticClass();
+	}
 	else
 	{
 		return AActor::StaticClass();
@@ -969,6 +977,26 @@ bool UsdUtils::IsAnimated(const pxr::UsdPrim& Prim)
 			}
 		}
 	}
+	else if (pxr::UsdVolVolume Volume{Prim})
+	{
+		pxr::UsdStageRefPtr Stage = Prim.GetStage();
+
+		const std::map<pxr::TfToken, pxr::SdfPath>& FieldMap = Volume.GetFieldPaths();
+		for (std::map<pxr::TfToken, pxr::SdfPath>::const_iterator Iter = FieldMap.cbegin(); Iter != FieldMap.cend(); ++Iter)
+		{
+			const pxr::SdfPath& AssetPrimPath = Iter->second;
+
+			if (pxr::UsdVolOpenVDBAsset OpenVDBAsset{Stage->GetPrimAtPath(AssetPrimPath)})
+			{
+				std::vector<double> TimeSamples;
+				pxr::UsdAttribute FilePathAttr = OpenVDBAsset.GetFilePathAttr();
+				if (FilePathAttr && FilePathAttr.GetTimeSamples(&TimeSamples) && TimeSamples.size() > 1)
+				{
+					return true;
+				}
+			}
+		}
+	}
 
 	return false;
 }
@@ -1681,6 +1709,10 @@ UUsdAssetImportData* UsdUtils::GetAssetImportData(UObject* Asset)
 	{
 		ImportData = Cast<UUsdAssetImportData>(GroomCache->AssetImportData);
 	}
+	else if (UStreamableSparseVolumeTexture* SparseVolumeTexture = Cast<UStreamableSparseVolumeTexture>(Asset))
+	{
+		ImportData = Cast<UUsdAssetImportData>(SparseVolumeTexture->AssetImportData);
+	}
 
 #endif
 	return ImportData;
@@ -1726,6 +1758,10 @@ void UsdUtils::SetAssetImportData(UObject* Asset, UAssetImportData* ImportData)
 	{
 		GroomCache->AssetImportData = ImportData;
 	}
+	else if (UStreamableSparseVolumeTexture* SparseVolumeTexture = Cast<UStreamableSparseVolumeTexture>(Asset))
+	{
+		SparseVolumeTexture->AssetImportData = ImportData;
+	}
 #endif	  // WITH_EDITOR
 }
 
@@ -3511,8 +3547,7 @@ void UsdUtils::CollectSchemaAnalytics(const UE::FUsdStage& Stage, const FString&
 		"Volume",
 		"VolumeLightAPI",
 		"Xform",
-		"XformCommonAPI"
-	};
+		"XformCommonAPI"};
 
 	FString ConcatenatedSeenSchemas;
 	for (const FString& SchemaName : SeenSchemas)
diff --git a/Engine/Plugins/Importers/USDImporter/Source/USDUtilities/Private/USDLightConversion.cpp b/Engine/Plugins/Importers/USDImporter/Source/USDUtilities/Private/USDLightConversion.cpp
--- a/Engine/Plugins/Importers/USDImporter/Source/USDUtilities/Private/USDLightConversion.cpp
+++ b/Engine/Plugins/Importers/USDImporter/Source/USDUtilities/Private/USDLightConversion.cpp
@@ -246,7 +246,7 @@ bool UsdToUnreal::ConvertDomeLight(
 	// Revert the allocator in case we end up creating a texture on the ansi allocator or something like that
 	FScopedUnrealAllocs UEAllocs;
 
-	const FString ResolvedDomeTexturePath = UsdUtils::GetResolvedTexturePath(DomeLight.GetTextureFileAttr());
+	const FString ResolvedDomeTexturePath = UsdUtils::GetResolvedAssetPath(DomeLight.GetTextureFileAttr());
 	if (ResolvedDomeTexturePath.IsEmpty())
 	{
 		FScopedUsdAllocs Allocs;
diff --git a/Engine/Plugins/Importers/USDImporter/Source/USDUtilities/Private/USDPrimConversion.cpp b/Engine/Plugins/Importers/USDImporter/Source/USDUtilities/Private/USDPrimConversion.cpp
--- a/Engine/Plugins/Importers/USDImporter/Source/USDUtilities/Private/USDPrimConversion.cpp
+++ b/Engine/Plugins/Importers/USDImporter/Source/USDUtilities/Private/USDPrimConversion.cpp
@@ -30,6 +30,7 @@
 #include "CineCameraComponent.h"
 #include "Components/BrushComponent.h"
 #include "Components/DirectionalLightComponent.h"
+#include "Components/HeterogeneousVolumeComponent.h"
 #include "Components/HierarchicalInstancedStaticMeshComponent.h"
 #include "Components/LightComponent.h"
 #include "Components/LocalLightComponent.h"
@@ -57,6 +58,7 @@
 #include "Sections/MovieSceneFloatSection.h"
 #include "Sections/MovieSceneVectorSection.h"
 #include "Sections/MovieSceneVisibilitySection.h"
+#include "SparseVolumeTexture/SparseVolumeTexture.h"
 #include "Tracks/MovieScene3DTransformTrack.h"
 #include "Tracks/MovieSceneBoolTrack.h"
 #include "Tracks/MovieSceneColorTrack.h"
@@ -112,6 +114,8 @@
 #include "pxr/usd/usdSkel/skeleton.h"
 #include "pxr/usd/usdSkel/skeletonQuery.h"
 #include "pxr/usd/usdUtils/stageCache.h"
+#include "pxr/usd/usdVol/openVDBAsset.h"
+#include "pxr/usd/usdVol/volume.h"
 #include "USDIncludesEnd.h"
 
 static bool GConsiderAllPrimsHaveAnimatedBounds = false;
@@ -652,7 +656,8 @@ bool UsdToUnreal::ConvertFloatTimeSamples(
 	const TArray<double>& UsdTimeSamples,
 	const TFunction<float(double)>& ReaderFunc,
 	UMovieSceneFloatTrack& MovieSceneTrack,
-	const FMovieSceneSequenceTransform& SequenceTransform
+	const FMovieSceneSequenceTransform& SequenceTransform,
+	TOptional<ERichCurveInterpMode> InterpolationModeOverride
 )
 {
 	if (!ReaderFunc)
@@ -683,7 +688,8 @@ bool UsdToUnreal::ConvertFloatTimeSamples(
 	const double StageTimeCodesPerSecond = UsdStage->GetTimeCodesPerSecond();
 	const FFrameRate StageFrameRate(StageTimeCodesPerSecond, 1);
 
-	const ERichCurveInterpMode InterpMode = (UsdStage->GetInterpolationType() == pxr::UsdInterpolationTypeLinear)
+	const ERichCurveInterpMode InterpMode = InterpolationModeOverride.IsSet() ? InterpolationModeOverride.GetValue()
+											: (UsdStage->GetInterpolationType() == pxr::UsdInterpolationTypeLinear)
 												? ERichCurveInterpMode::RCIM_Linear
 												: ERichCurveInterpMode::RCIM_Constant;
 
@@ -4011,7 +4017,8 @@ UnrealToUsd::FPropertyTrackWriter UnrealToUsd::CreatePropertyTrackWriter(
 		return Result;
 	}
 
-	pxr::UsdAttribute Attr;
+	TArray<pxr::UsdAttribute> Attrs = {pxr::UsdAttribute{}};
+	pxr::UsdAttribute& Attr = Attrs[0];
 	{
 		pxr::SdfChangeBlock ChangeBlock;
 
@@ -4471,6 +4478,96 @@ UnrealToUsd::FPropertyTrackWriter UnrealToUsd::CreatePropertyTrackWriter(
 				}
 			}
 		}
+		else if (const UHeterogeneousVolumeComponent* VolumeComponent = Cast<const UHeterogeneousVolumeComponent>(&Component))
+		{
+			if (Track.GetPropertyName() == GET_MEMBER_NAME_CHECKED(UHeterogeneousVolumeComponent, Frame))
+			{
+				const TArray<FString>* TimeSamplePaths = nullptr;
+				const TArray<FString>* SourceOpenVDBAssetPrimPaths = nullptr;
+
+				const int32 ElementIndex = 0;
+				if (UMaterialInterface* CurrentMaterial = VolumeComponent->GetMaterial(ElementIndex))
+				{
+					TArray<FMaterialParameterInfo> ParameterInfo;
+					TArray<FGuid> ParameterIds;
+					CurrentMaterial->GetAllSparseVolumeTextureParameterInfo(ParameterInfo, ParameterIds);
+
+					for (const FMaterialParameterInfo& Info : ParameterInfo)
+					{
+						USparseVolumeTexture* SparseVolumeTexture = nullptr;
+						if (CurrentMaterial->GetSparseVolumeTextureParameterValue(Info, SparseVolumeTexture) && SparseVolumeTexture)
+						{
+							if (SparseVolumeTexture->GetNumFrames() > 1)
+							{
+								if (UUsdSparseVolumeTextureAssetUserData* UserData = Cast<UUsdSparseVolumeTextureAssetUserData>(
+										UsdUtils::GetAssetUserData(SparseVolumeTexture)
+									))
+								{
+									SourceOpenVDBAssetPrimPaths = &UserData->SourceOpenVDBAssetPrimPaths;
+									TimeSamplePaths = &UserData->TimeSamplePaths;
+								}
+							}
+						}
+
+						// Only care about animation on first SVT parameter
+						break;
+					}
+				}
+
+				// Collect the Attrs we'll need to write out to.
+				// Realistically this is a single Attr, but there could be more in case multiple OpenVDBAsset prims refer to the
+				// exact same VDB file paths
+				if (SourceOpenVDBAssetPrimPaths && SourceOpenVDBAssetPrimPaths->Num() > 0)
+				{
+					Attrs.Reset(SourceOpenVDBAssetPrimPaths->Num());
+
+					for (const FString& OpenVDBPrimPath : *SourceOpenVDBAssetPrimPaths)
+					{
+						pxr::UsdPrim OpenVDBPrim = UsdStage->GetPrimAtPath(UnrealToUsd::ConvertPath(*OpenVDBPrimPath).Get());
+						if (pxr::UsdVolOpenVDBAsset OpenVDBAsset{OpenVDBPrim})
+						{
+							Attrs.Add(OpenVDBAsset.CreateFilePathAttr());
+						}
+					}
+
+					Attr = Attrs[0];
+				}
+
+				if (TimeSamplePaths)
+				{
+					TArray<pxr::SdfAssetPath> FrameIndexToPath;
+					FrameIndexToPath.Reserve(TimeSamplePaths->Num());
+					for (const FString& TimeSamplePath : *TimeSamplePaths)
+					{
+						FrameIndexToPath.Add(pxr::SdfAssetPath{UnrealToUsd::ConvertString(*TimeSamplePath).Get()});
+					}
+
+					Result.FloatWriter = [Attrs,
+										  FrameIndexToPath = MoveTemp(FrameIndexToPath),
+										  LastPath = TUsdStore<pxr::SdfAssetPath>{}](float UEValue, double UsdTimeCode) mutable
+					{
+						// The UEValue here corresponds to a frame index into the SVT (with constant interpolation).
+						// Regardless of what the change was on the track, we can assume here that our TimeSampleX arrays are up
+						// to date with the generated SVT. This means we essentially just need to author a timeSample at UsdTimeCode
+						// that points at the file path that corresponds to that frame
+
+						const int32 FrameIndex = FMath::FloorToInt32(UEValue);
+						const pxr::SdfAssetPath& TimeSamplePath = FrameIndexToPath[FrameIndex];
+
+						// This check prevents us from writing the same identical path on every single bake tick
+						if (TimeSamplePath != LastPath.Get())
+						{
+							for (const pxr::UsdAttribute& Attr : Attrs)
+							{
+								Attr.Set(TimeSamplePath, UsdTimeCode);
+							}
+						}
+
+						LastPath = TimeSamplePath;
+					};
+				}
+			}
+		}
 		// Bounds component properties
 		else
 		{
@@ -4520,22 +4617,30 @@ UnrealToUsd::FPropertyTrackWriter UnrealToUsd::CreatePropertyTrackWriter(
 		}
 	}
 
-	if (Attr)
+	for (const pxr::UsdAttribute& SomeAttr : Attrs)
 	{
-		std::vector<double> TimeSamples;
-		Attr.GetTimeSamples(&TimeSamples);
-		for (double TimeSample : TimeSamples)
+		if (SomeAttr)
 		{
-			Attr.ClearAtTime(TimeSample);
-		}
+			// Weirdly enough GetTimeSamples() will return time codes with the offset and scale applied, while
+			// ClearAtTime() expects time codes without offset and scale applied, so we must manually undo them here
+			UE::FSdfLayerOffset CombinedOffset = UsdUtils::GetPrimToStageOffset(UE::FUsdPrim{SomeAttr.GetPrim()});
 
-		// Note that we must do this only after the change block is destroyed!
-		// This is important because if we don't have spec for this attribute on the current edit target, we're relying
-		// on the previous code to create it, and we need to let USD emit its internal notices and fully commit the
-		// "attribute creation" spec first. This because NotifyIfOverriddenOpinion will go through the attribute's
-		// spec stack and consider our attribute overriden if it finds a stronger opinion than the one on the edit
-		// target. Well if our own spec hasn't been created yet it will misfire when it runs into any other spec
-		UsdUtils::NotifyIfOverriddenOpinion(Attr);
+			std::vector<double> TimeSamples;
+			SomeAttr.GetTimeSamples(&TimeSamples);
+			for (double TimeSample : TimeSamples)
+			{
+				double LocalTime = (TimeSample - CombinedOffset.Offset) / CombinedOffset.Scale;
+				SomeAttr.ClearAtTime(LocalTime);
+			}
+
+			// Note that we must do this only after the change block is destroyed!
+			// This is important because if we don't have spec for this attribute on the current edit target, we're relying
+			// on the previous code to create it, and we need to let USD emit its internal notices and fully commit the
+			// "attribute creation" spec first. This because NotifyIfOverriddenOpinion will go through the attribute's
+			// spec stack and consider our attribute overriden if it finds a stronger opinion than the one on the edit
+			// target. Well if our own spec hasn't been created yet it will misfire when it runs into any other spec
+			UsdUtils::NotifyIfOverriddenOpinion(SomeAttr);
+		}
 	}
 
 	return Result;
@@ -4940,6 +5045,33 @@ TArray<UE::FUsdAttribute> UnrealToUsd::GetAttributesForProperty(const UE::FUsdPr
 			return {UE::FUsdAttribute(Boundable.GetExtentAttr())};
 		}
 	}
+	else if (PropertyPath == GET_MEMBER_NAME_CHECKED(UHeterogeneousVolumeComponent, Frame))
+	{
+		if (pxr::UsdVolVolume Volume{Prim})
+		{
+			pxr::UsdStageRefPtr Stage = Prim.GetStage();
+
+			TArray<UE::FUsdAttribute> Attrs;
+
+			const std::map<pxr::TfToken, pxr::SdfPath>& FieldMap = Volume.GetFieldPaths();
+			Attrs.Reserve(FieldMap.size());
+
+			for (std::map<pxr::TfToken, pxr::SdfPath>::const_iterator Iter = FieldMap.cbegin(); Iter != FieldMap.cend(); ++Iter)
+			{
+				const pxr::SdfPath& AssetPrimPath = Iter->second;
+
+				if (pxr::UsdVolOpenVDBAsset OpenVDBAsset{Stage->GetPrimAtPath(AssetPrimPath)})
+				{
+					if (pxr::UsdAttribute FilePathAttr = OpenVDBAsset.GetFilePathAttr())
+					{
+						Attrs.Add(UE::FUsdAttribute{FilePathAttr});
+					}
+				}
+			}
+
+			return Attrs;
+		}
+	}
 
 	return {};
 }
@@ -5090,7 +5222,7 @@ bool UnrealToUsd::ConvertDrawModeComponent(const UUsdDrawModeComponent& DrawMode
 					if (FaceTexture->AssetImportData)
 					{
 						FString TextureSourcePath = FaceTexture->AssetImportData->GetFirstFilename();
-						FString ResolvedPath = UsdUtils::GetResolvedTexturePath(Attr);
+						FString ResolvedPath = UsdUtils::GetResolvedAssetPath(Attr);
 
 						// Avoid authoring anything unless they point at different files because in the general case the
 						// asset import data will have an absolute path, while the path on the attribute may be currently relative.
diff --git a/Engine/Plugins/Importers/USDImporter/Source/USDUtilities/Private/USDShadeConversion.cpp b/Engine/Plugins/Importers/USDImporter/Source/USDUtilities/Private/USDShadeConversion.cpp
--- a/Engine/Plugins/Importers/USDImporter/Source/USDUtilities/Private/USDShadeConversion.cpp
+++ b/Engine/Plugins/Importers/USDImporter/Source/USDUtilities/Private/USDShadeConversion.cpp
@@ -130,7 +130,7 @@ namespace UE
 #endif	  // WITH_EDITOR
 
 			// Given an AssetPath, resolve it to an actual file path
-			FString ResolveTexturePath(const pxr::SdfLayerHandle& LayerHandle, const FString& AssetPath)
+			FString ResolveAssetPath(const pxr::SdfLayerHandle& LayerHandle, const FString& AssetPath)
 			{
 				// TODO: Most of this stuff is incompatible with custom resolvers, as these asset paths may be URLs
 				// or just GUIDs or anything like that, where relative vs absolute path make no sense. We will need
@@ -653,7 +653,7 @@ namespace UE
 
 					if (FileInput && FileInput.GetTypeName() == pxr::SdfValueTypeNames->Asset)	  // Check that FileInput is of type Asset
 					{
-						const FString TexturePath = UsdUtils::GetResolvedTexturePath(FileInput.GetAttr());
+						const FString TexturePath = UsdUtils::GetResolvedAssetPath(FileInput.GetAttr());
 						// We will assume the texture is valid, and show a warning if we fail to parse this later.
 						// Note that we don't even check that the file exists: If we have a texture bound to opacity then we
 						// assume this material is meant to be translucent, even if the texture path is invalid (or points inside an USDZ archive)
@@ -1352,7 +1352,7 @@ namespace UE
 						TextureFactory->UdimRegexPattern = MoveTemp(UdimRegexPattern);
 					}
 
-					const FString ResolvedTexturePath = UsdUtils::GetResolvedTexturePath(TextureAssetPathAttr);
+					const FString ResolvedTexturePath = UsdUtils::GetResolvedAssetPath(TextureAssetPathAttr);
 					if (!ResolvedTexturePath.IsEmpty())
 					{
 						EObjectFlags ObjectFlags = RF_Transactional | RF_Transient;
@@ -1420,7 +1420,7 @@ namespace UE
 
 				if (!TexturePath.IsEmpty())
 				{
-					const FString ResolvedTexturePath = UsdUtils::GetResolvedTexturePath(TextureAssetPathAttr);
+					const FString ResolvedTexturePath = UsdUtils::GetResolvedAssetPath(TextureAssetPathAttr);
 					if (!ResolvedTexturePath.IsEmpty())
 					{
 						// Try checking if the texture is inside an USDZ archive first, or else TextureFactory throws an error
@@ -2129,7 +2129,7 @@ namespace UE
 					// we're doing, as it happens with a pure USD python script. (this with USD 21.05 in June 2021)
 					if (ShadeInputValue.IsHolding<pxr::SdfAssetPath>())
 					{
-						FString ResolvedPath = UsdUtils::GetResolvedTexturePath(ShadeInput.GetAttr());
+						FString ResolvedPath = UsdUtils::GetResolvedAssetPath(ShadeInput.GetAttr());
 						InOutHashState.UpdateWithString(*ResolvedPath, ResolvedPath.Len());
 					}
 					else if (ShadeInputValue.IsHolding<pxr::TfToken>())
@@ -2918,14 +2918,19 @@ bool UnrealToUsd::ConvertFlattenMaterial(
 
 #endif	  // WITH_EDITOR
 
-FString UsdUtils::GetResolvedTexturePath(const pxr::UsdAttribute& TextureAssetPathAttr)
+FString UsdUtils::GetResolvedAssetPath(const pxr::UsdAttribute& AssetPathAttr, pxr::UsdTimeCode TimeCode)
 {
+	if (!AssetPathAttr)
+	{
+		return {};
+	}
+
 	FScopedUsdAllocs UsdAllocs;
 
-	pxr::SdfAssetPath TextureAssetPath;
-	TextureAssetPathAttr.Get<pxr::SdfAssetPath>(&TextureAssetPath);
+	pxr::SdfAssetPath AssetPath;
+	AssetPathAttr.Get<pxr::SdfAssetPath>(&AssetPath, TimeCode);
 
-	std::string AssetIdentifier = TextureAssetPath.GetResolvedPath();
+	std::string AssetIdentifier = AssetPath.GetResolvedPath();
 	// Don't normalize an empty path as the result will be "."
 	if (AssetIdentifier.size() > 0)
 	{
@@ -2938,24 +2943,25 @@ FString UsdUtils::GetResolvedTexturePath(const pxr::UsdAttribute& TextureAssetPa
 
 	if (ResolvedTexturePath.IsEmpty())
 	{
-		FString TexturePath = UsdToUnreal::ConvertString(TextureAssetPath.GetAssetPath());
+		FString TexturePath = UsdToUnreal::ConvertString(AssetPath.GetAssetPath());
 		FPaths::NormalizeFilename(TexturePath);
 
 		if (!TexturePath.IsEmpty())
 		{
-			pxr::SdfLayerRefPtr TextureLayer = UsdUtils::FindLayerForAttribute(TextureAssetPathAttr, pxr::UsdTimeCode::EarliestTime().GetValue());
-			ResolvedTexturePath = UsdShadeConversionImpl::ResolveTexturePath(TextureLayer, TexturePath);
+			pxr::SdfLayerRefPtr TextureLayer = UsdUtils::FindLayerForAttribute(AssetPathAttr, TimeCode.GetValue());
+			ResolvedTexturePath = UsdShadeConversionImpl::ResolveAssetPath(TextureLayer, TexturePath);
 		}
 	}
 
-	if (ResolvedTexturePath.IsEmpty())
-	{
-		UE_LOG(LogUsd, Warning, TEXT("Failed to resolve texture path on attribute '%s'"), *UsdToUnreal::ConvertPath(TextureAssetPathAttr.GetPath()));
-	}
-
 	return ResolvedTexturePath;
 }
 
+// Deprecated
+FString UsdUtils::GetResolvedTexturePath(const pxr::UsdAttribute& TextureAssetPathAttr)
+{
+	return UsdUtils::GetResolvedAssetPath(TextureAssetPathAttr);
+}
+
 FString UsdUtils::GetTextureHash(
 	const FString& ResolvedTexturePath,
 	bool bSRGB,
diff --git a/Engine/Plugins/Importers/USDImporter/Source/USDUtilities/Public/USDPrimConversion.h b/Engine/Plugins/Importers/USDImporter/Source/USDUtilities/Public/USDPrimConversion.h
--- a/Engine/Plugins/Importers/USDImporter/Source/USDUtilities/Public/USDPrimConversion.h
+++ b/Engine/Plugins/Importers/USDImporter/Source/USDUtilities/Public/USDPrimConversion.h
@@ -47,6 +47,7 @@ class USceneComponent;
 class USkeletalMeshComponent;
 class UUsdAssetCache2;
 class UUsdDrawModeComponent;
+enum ERichCurveInterpMode : int;
 struct FFrameRate;
 struct FMovieSceneSequenceTransform;
 struct FUsdCombinedPrimMetadata;
@@ -141,7 +142,8 @@ namespace UsdToUnreal
 		const TArray<double>& UsdTimeSamples,
 		const TFunction<float(double)>& ReaderFunc,
 		UMovieSceneFloatTrack& MovieSceneTrack,
-		const FMovieSceneSequenceTransform& SequenceTransform
+		const FMovieSceneSequenceTransform& SequenceTransform,
+		TOptional<ERichCurveInterpMode> InterpolationModeOverride = {}
 	);
 	USDUTILITIES_API bool ConvertColorTimeSamples(
 		const UE::FUsdStage& UsdStage,
diff --git a/Engine/Plugins/Importers/USDImporter/Source/USDUtilities/Public/USDShadeConversion.h b/Engine/Plugins/Importers/USDImporter/Source/USDUtilities/Public/USDShadeConversion.h
--- a/Engine/Plugins/Importers/USDImporter/Source/USDUtilities/Public/USDShadeConversion.h
+++ b/Engine/Plugins/Importers/USDImporter/Source/USDUtilities/Public/USDShadeConversion.h
@@ -222,7 +222,10 @@ namespace UsdUtils
 	USDUTILITIES_API FSHAHash
 	HashShadeMaterial(const pxr::UsdShadeMaterial& UsdShadeMaterial, const pxr::TfToken& RenderContext = pxr::UsdShadeTokens->universalRenderContext);
 
-	/** Returns the resolved path from a pxr::SdfAssetPath attribute. For UDIMs path, returns the path to the 1001 tile. */
+	/** Returns the resolved path from an pxr::SdfAssetPath attribute. For UDIMs path, returns the path to the 1001 tile. */
+	USDUTILITIES_API FString GetResolvedAssetPath(const pxr::UsdAttribute& AssetPathAttr, pxr::UsdTimeCode TimeCode = pxr::UsdTimeCode::Default());
+
+	UE_DEPRECATED(5.4, "This function has been renamed to 'GetResolvedAssetPath', as it should work for any asset type")
 	USDUTILITIES_API FString GetResolvedTexturePath(const pxr::UsdAttribute& TextureAssetPathAttr);
 
 	/**
diff --git a/Engine/Plugins/Importers/USDImporter/Source/UnrealUSDWrapper/Private/UnrealUSDWrapper.cpp b/Engine/Plugins/Importers/USDImporter/Source/UnrealUSDWrapper/Private/UnrealUSDWrapper.cpp
--- a/Engine/Plugins/Importers/USDImporter/Source/UnrealUSDWrapper/Private/UnrealUSDWrapper.cpp
+++ b/Engine/Plugins/Importers/USDImporter/Source/UnrealUSDWrapper/Private/UnrealUSDWrapper.cpp
@@ -110,8 +110,8 @@ namespace UnrealIdentifiers
 	const TfToken UnrealSVTMappedMaterialParameters = TfToken("unreal:SVT:mappedMaterialParameters");
 	const TfToken UnrealSVTMappedGridComponents = TfToken("unreal:SVT:mappedGridComponents");
 	const TfToken UnrealSVTMappedAttributeChannels = TfToken("unreal:SVT:mappedAttributeChannels");
-	const TfToken UnrealSVTAttributesADataType = TfToken("unreal:SVT:attributesAdataType");
-	const TfToken UnrealSVTAttributesBDataType = TfToken("unreal:SVT:attributesBdataType");
+	const TfToken UnrealSVTAttributesADataType = TfToken("unreal:SVT:attributesADataType");
+	const TfToken UnrealSVTAttributesBDataType = TfToken("unreal:SVT:attributesBDataType");
 
 	const TfToken DiffuseColor = TfToken("diffuseColor");
 	const TfToken EmissiveColor = TfToken("emissiveColor");
