USD: Fix some open/compose file picker filters not working on the Mac.

diff --git a/Engine/Plugins/Importers/USDImporter/Source/USDStage/Private/USDStageActor.cpp b/Engine/Plugins/Importers/USDImporter/Source/USDStage/Private/USDStageActor.cpp
--- a/Engine/Plugins/Importers/USDImporter/Source/USDStage/Private/USDStageActor.cpp
+++ b/Engine/Plugins/Importers/USDImporter/Source/USDStage/Private/USDStageActor.cpp
@@ -1077,10 +1077,13 @@ AUsdStageActor::AUsdStageActor()
 				TArray<FString> SupportedExtensions = UnrealUSDWrapper::GetNativeFileFormats();
 				if (SupportedExtensions.Num() > 0)
 				{
-					FString JoinedExtensions = FString::Join(SupportedExtensions, TEXT("; *."));	// Combine "usd" and "usda" into "usd; *.usda"
+					// Note: Cannot have space after semicolon or else the parsing breaks on the Mac...
+					FString JoinedWithSemicolon = FString::Join(SupportedExtensions, TEXT(";*."));
+					FString JoinedWithComma = FString::Join(SupportedExtensions, TEXT(", *."));
+
 					Property->SetMetaData(
 						TEXT("FilePathFilter"),
-						FString::Printf(TEXT("Universal Scene Description files|*.%s"), *JoinedExtensions, *JoinedExtensions)
+						FString::Printf(TEXT("Universal Scene Description files (*.%s)|*.%s"), *JoinedWithComma, *JoinedWithSemicolon)
 					);
 				}
 				break;
diff --git a/Engine/Plugins/Importers/USDImporter/Source/UnrealUSDWrapper/Private/UnrealUSDWrapper.cpp b/Engine/Plugins/Importers/USDImporter/Source/UnrealUSDWrapper/Private/UnrealUSDWrapper.cpp
--- a/Engine/Plugins/Importers/USDImporter/Source/UnrealUSDWrapper/Private/UnrealUSDWrapper.cpp
+++ b/Engine/Plugins/Importers/USDImporter/Source/UnrealUSDWrapper/Private/UnrealUSDWrapper.cpp
@@ -1219,10 +1219,13 @@ public:
 				TArray<FString> SupportedExtensions = UnrealUSDWrapper::GetAllSupportedFileFormats();
 				if (SupportedExtensions.Num() > 0)
 				{
-					FString JoinedExtensions = FString::Join(SupportedExtensions, TEXT("; *."));	// Combine "usd" and "usda" into "usd; *.usda"
+					// Note: Cannot have space after semicolon or else the parsing breaks on the Mac...
+					FString JoinedWithSemicolon = FString::Join(SupportedExtensions, TEXT(";*."));
+					FString JoinedWithComma = FString::Join(SupportedExtensions, TEXT(", *."));
+
 					Property->SetMetaData(
 						TEXT("FilePathFilter"),
-						FString::Printf(TEXT("Universal Scene Description files|*.%s"), *JoinedExtensions, *JoinedExtensions)
+						FString::Printf(TEXT("Universal Scene Description files (*.%s)|*.%s"), *JoinedWithComma, *JoinedWithSemicolon)
 					);
 				}
 				break;
