Fix error C4702: unreachable code
#rnx

diff --git a/Engine/Plugins/Importers/USDImporter/Source/USDTests/Private/USDTestsBlueprintLibrary.cpp b/Engine/Plugins/Importers/USDImporter/Source/USDTests/Private/USDTestsBlueprintLibrary.cpp
--- a/Engine/Plugins/Importers/USDImporter/Source/USDTests/Private/USDTestsBlueprintLibrary.cpp
+++ b/Engine/Plugins/Importers/USDImporter/Source/USDTests/Private/USDTestsBlueprintLibrary.cpp
@@ -55,9 +55,9 @@ bool USDTestsBlueprintLibrary::RecompileBlueprintStageActor(AUsdStageActor* Blue
 		UE_LOG(LogUsd, Error, TEXT("Blueprint is in an unexpected state after compiling (%s)"), *BP->GetName());
 		return false;
 	}
-#endif	  // WITH_EDITOR
-
+#else
 	return false;
+#endif	  // WITH_EDITOR
 }
 
 void USDTestsBlueprintLibrary::DirtyStageActorBlueprint(AUsdStageActor* BlueprintDerivedStageActor)
