USD: Fix incorrect skeletal mesh sections whenever the skeleton had too many bones.

diff --git a/Engine/Plugins/Importers/USDImporter/Source/USDSchemas/Private/USDSkelSkeletonTranslator.cpp b/Engine/Plugins/Importers/USDImporter/Source/USDSchemas/Private/USDSkelSkeletonTranslator.cpp
--- a/Engine/Plugins/Importers/USDImporter/Source/USDSchemas/Private/USDSkelSkeletonTranslator.cpp
+++ b/Engine/Plugins/Importers/USDImporter/Source/USDSchemas/Private/USDSkelSkeletonTranslator.cpp
@@ -36,6 +36,7 @@
 #include "MaterialDomain.h"
 #include "Materials/Material.h"
 #include "Rendering/SkeletalMeshLODImporterData.h"
+#include "Rendering/SkeletalMeshModel.h"
 
 #if WITH_EDITOR
 #include "AnimGraphNode_LiveLinkPose.h"
@@ -116,6 +117,14 @@ namespace UsdSkelSkeletonTranslatorImpl
 
 		bool bMaterialsHaveChanged = false;
 
+		const FSkeletalMeshModel* ImportedResource = SkeletalMesh->GetImportedModel();
+		if (!ImportedResource)
+		{
+			return false;
+		}
+
+		const TIndirectArray<FSkeletalMeshLODModel>& LODModels = ImportedResource->LODModels;
+
 		uint32 SkeletalMeshSlotIndex = 0;
 		for (int32 LODIndex = 0; LODIndex < LODIndexToMaterialInfo.Num(); ++LODIndex)
 		{
@@ -134,8 +143,14 @@ namespace UsdSkelSkeletonTranslatorImpl
 				);
 				continue;
 			}
-			TArray<int32>& LODMaterialMap = LODInfo->LODMaterialMap;
-			LODMaterialMap.Reserve(LODSlots.Num());
+
+			if (!LODModels.IsValidIndex(LODIndex))
+			{
+				return false;
+			}
+			const FSkeletalMeshLODModel& LODModel = LODModels[LODIndex];
+
+			TMap<int32, int32> LODIndexToMeshIndex;
 
 			for (int32 LODSlotIndex = 0; LODSlotIndex < LODSlots.Num(); ++LODSlotIndex, ++SkeletalMeshSlotIndex)
 			{
@@ -198,15 +213,32 @@ namespace UsdSkelSkeletonTranslatorImpl
 					bMaterialsHaveChanged = true;
 				}
 
-				// Already have a material at that LOD remap slot, need to reassign
-				if (LODMaterialMap.IsValidIndex(LODSlotIndex))
-				{
-					LODMaterialMap[LODSlotIndex] = SkeletalMeshSlotIndex;
-				}
-				// Add new material slot remap
-				else
+				LODIndexToMeshIndex.Add(LODSlotIndex, SkeletalMeshSlotIndex);
+			}
+
+			// Our LOD slots from USD want to use LODSlotIndex (above) as a material index, but the SkeletalMesh
+			// actual material slot order may be different as we just append all material assignments,
+			// so we need to fill in LODMaterialMap which is internally used to do that mapping.
+			//
+			// Note that LODMaterialMap needs to match the actual list of sections on the skeletal mesh, and
+			// we may end up with more (or less?) sections than we expect (e.g. if our skeleton is too large
+			// the build process may create new "chunked" sections that also point at the same material slots).
+			// Here we step through all sections for this LOD and add LODMaterialMap entries for the
+			// relevant ones.
+			TArray<int32>& LODMaterialMap = LODInfo->LODMaterialMap;
+			LODMaterialMap.SetNumUninitialized(LODModel.Sections.Num());
+			for (int32& Mapping : LODMaterialMap)
+			{
+				Mapping = INDEX_NONE;	 // Initialize map with INDEX_NONE (means no remapping for that index)
+			}
+
+			for (int32 SectionIndex = 0; SectionIndex < LODModel.Sections.Num(); ++SectionIndex)
+			{
+				const FSkelMeshSection& Section = LODModel.Sections[SectionIndex];
+
+				if (int32* FoundMeshSlotIndex = LODIndexToMeshIndex.Find(Section.MaterialIndex))
 				{
-					LODMaterialMap.Add(SkeletalMeshSlotIndex);
+					LODMaterialMap[SectionIndex] = *FoundMeshSlotIndex;
 				}
 			}
 		}
