USD: Use the default SparseVolumeMaterial instead of our custom default material.

#rnx

diff --git a/Engine/Plugins/Importers/USDImporter/Source/USDClasses/Public/USDProjectSettings.h b/Engine/Plugins/Importers/USDImporter/Source/USDClasses/Public/USDProjectSettings.h
--- a/Engine/Plugins/Importers/USDImporter/Source/USDClasses/Public/USDProjectSettings.h
+++ b/Engine/Plugins/Importers/USDImporter/Source/USDClasses/Public/USDProjectSettings.h
@@ -117,7 +117,7 @@ public:
 	 * will prioritize trying to use those as the volumetric materials for the Sparse Volume Textures instead.
 	 */
 	UPROPERTY(config, EditAnywhere, Category = "USD|Reference Materials", meta = (AllowedClasses = "/Script/Engine.MaterialInterface"))
-	FSoftObjectPath ReferenceDefaultSVTMaterial = FSoftObjectPath{TEXT("/USDImporter/Materials/ReferenceDefaultSVTMaterial.ReferenceDefaultSVTMaterial")};
+	FSoftObjectPath ReferenceDefaultSVTMaterial = FSoftObjectPath{TEXT("/Engine/EngineMaterials/SparseVolumeMaterial.SparseVolumeMaterial")};
 
 	/**
 	 * What material to use for UUsdDrawModeComponents with "Cards" draw mode and provided textures (corresponding to
