USD: Fix minor issues related to saving/exporting USD layers;

Fix exporting all layers not working for New stages;
When saving/exporting a USD layer to disk, don't show the format option that lists all supported extensions at the same time;

diff --git a/Engine/Plugins/Importers/USDImporter/Source/USDStageEditor/Private/SUSDStage.cpp b/Engine/Plugins/Importers/USDImporter/Source/USDStageEditor/Private/SUSDStage.cpp
--- a/Engine/Plugins/Importers/USDImporter/Source/USDStageEditor/Private/SUSDStage.cpp
+++ b/Engine/Plugins/Importers/USDImporter/Source/USDStageEditor/Private/SUSDStage.cpp
@@ -10,6 +10,7 @@
 #include "USDClassesModule.h"
 #include "USDConversionUtils.h"
 #include "USDLayerUtils.h"
+#include "USDLog.h"
 #include "USDProjectSettings.h"
 #include "USDSchemasModule.h"
 #include "USDSchemaTranslator.h"
@@ -2252,11 +2253,11 @@ void SUsdStage::FileExportAllLayers(const FString& OutputDirectory)
 	// We will need to actually load these as opposed to just copy-paste the files as we may need to update references/paths
 	SUSDStageImpl::AppendAllExternalReferences(LayerStack);
 
-	// Discard session layers
+	// Discard session layers, but keep the root layer if it's anonymous (temp stage)
 	for (int32 Index = LayerStack.Num() - 1; Index >= 0; --Index)
 	{
 		UE::FSdfLayer& Layer = LayerStack[Index];
-		if (Layer.IsAnonymous())
+		if (Layer.IsAnonymous() && UsdStage.GetRootLayer() != Layer)
 		{
 			const int32 Count = 1;
 			LayerStack.RemoveAt(Index, Count, EAllowShrinking::No);
@@ -2378,7 +2379,11 @@ void SUsdStage::FileExportFlattenedStage(const FString& OutputLayer)
 
 	double StartTime = FPlatformTime::Cycles64();
 
-	UsdStage.Export(*OutputLayerCopy);
+	const bool bResult = UsdStage.Export(*OutputLayerCopy);
+	if (!bResult)
+	{
+		UE_LOG(LogUsd, Warning, TEXT("Failed to export flattened USD Stage to path '%s'!"), *OutputLayerCopy);
+	}
 
 	// Send analytics
 	if (FEngineAnalytics::IsAvailable())
diff --git a/Engine/Plugins/Importers/USDImporter/Source/USDUtilities/Private/USDLayerUtils.cpp b/Engine/Plugins/Importers/USDImporter/Source/USDUtilities/Private/USDLayerUtils.cpp
--- a/Engine/Plugins/Importers/USDImporter/Source/USDUtilities/Private/USDLayerUtils.cpp
+++ b/Engine/Plugins/Importers/USDImporter/Source/USDUtilities/Private/USDLayerUtils.cpp
@@ -233,8 +233,14 @@ TOptional<FString> UsdUtils::BrowseUsdFile(EBrowseFileMode Mode)
 		SupportedExtensions.Remove(TEXT("usdz"));
 	}
 
-	FString JoinedExtensions = FString::Join(SupportedExtensions, TEXT(";*."));	   // Combine "usd" and "usda" into "usd; *.usda"
-	FString FileTypes = FString::Printf(TEXT("Universal Scene Description files (*.%s)|*.%s|"), *JoinedExtensions, *JoinedExtensions);
+	// Show an option for "all supported extensions" at the same time, but only if not saving: For saving the user should have to pick one directly
+	FString FileTypes;
+	if (Mode != EBrowseFileMode::Save)
+	{
+		FString JoinedExtensions = FString::Join(SupportedExtensions, TEXT(";*."));	   // Combine "usd" and "usda" into "usd; *.usda"
+		FileTypes = FString::Printf(TEXT("Universal Scene Description files (*.%s)|*.%s|"), *JoinedExtensions, *JoinedExtensions);
+	}
+
 	for (const FString& SupportedExtension : SupportedExtensions)
 	{
 		const bool bIsTextNative = NativeTextExtensionsSet.Contains(SupportedExtension);
