USD: Fix old textures being used on draw mode components in some cases;

Fix missing up axis conversion when setting new textures back onto the USD stage;

diff --git a/Engine/Plugins/Importers/USDImporter/Source/USDSchemas/Private/USDGeomXformableTranslator.cpp b/Engine/Plugins/Importers/USDImporter/Source/USDSchemas/Private/USDGeomXformableTranslator.cpp
--- a/Engine/Plugins/Importers/USDImporter/Source/USDSchemas/Private/USDGeomXformableTranslator.cpp
+++ b/Engine/Plugins/Importers/USDImporter/Source/USDSchemas/Private/USDGeomXformableTranslator.cpp
@@ -22,6 +22,7 @@
 #include "Components/SceneComponent.h"
 #include "Components/SkeletalMeshComponent.h"
 #include "Components/StaticMeshComponent.h"
+#include "EditorFramework/AssetImportData.h"
 #include "Engine/Level.h"
 #include "Engine/StaticMesh.h"
 #include "Engine/Texture2D.h"
@@ -764,7 +765,12 @@ void FUsdGeomXformableTranslator::UpdateComponents(USceneComponent* SceneCompone
 
 namespace UE::UsdXformableTranslatorImpl::Private
 {
-	void AssignDrawModeComponentTextures(UE::FUsdPrim Prim, UUsdDrawModeComponent* DrawModeComponent, UUsdAssetCache2& AssetCache, FUsdInfoCache& InfoCache)
+	void AssignDrawModeComponentTextures(
+		UE::FUsdPrim Prim,
+		UUsdDrawModeComponent* DrawModeComponent,
+		UUsdAssetCache2& AssetCache,
+		FUsdInfoCache& InfoCache
+	)
 	{
 		if (!Prim)
 		{
@@ -784,9 +790,10 @@ namespace UE::UsdXformableTranslatorImpl::Private
 		}
 
 		// Collect textures from the info cache (they will all be linked to this Prim but will have within their AssetUserData the attribute
-		// that they originated from)
+		// that they originated from).
+		// The info cache may have old textures in case we're changing editing the stage, so we track multiple textures per attribute.
 		TArray<UTexture2D*> Textures = InfoCache.GetAssetsForPrim<UTexture2D>(Prim.GetPrimPath());
-		std::unordered_map<pxr::SdfPath, UTexture2D*, pxr::SdfPath::Hash> AttrPathToTextures;
+		std::unordered_map<pxr::SdfPath, TArray<UTexture2D*>, pxr::SdfPath::Hash> AttrPathToTextures;
 		for (UTexture2D* Texture : Textures)
 		{
 			if (UUsdAssetUserData* AssetUserData = Texture->GetAssetUserData<UUsdAssetUserData>())
@@ -801,7 +808,8 @@ namespace UE::UsdXformableTranslatorImpl::Private
 					pxr::SdfPath SdfPath = UnrealToUsd::ConvertPath(*TexturePath).Get();
 					if (SdfPath.IsPropertyPath())
 					{
-						AttrPathToTextures.insert({SdfPath, Texture});
+						TArray<UTexture2D*>& TexturesForAttr = AttrPathToTextures[SdfPath];
+						TexturesForAttr.AddUnique(Texture);
 					}
 				}
 			}
@@ -842,12 +850,26 @@ namespace UE::UsdXformableTranslatorImpl::Private
 
 				if (TextureSetter)
 				{
-					std::unordered_map<pxr::SdfPath, UTexture2D*, pxr::SdfPath::Hash>::iterator iter = AttrPathToTextures.find(Attr.GetPath());
+					std::unordered_map<pxr::SdfPath, TArray<UTexture2D*>, pxr::SdfPath::Hash>::iterator iter = AttrPathToTextures.find(Attr.GetPath());
 					if (iter != AttrPathToTextures.end())
 					{
-						if (UTexture2D* Texture = iter->second)
+						const FString TexturePath = UsdUtils::GetResolvedAssetPath(Attr);
+
+						const TArray<UTexture2D*>& Textures = iter->second;
+						for (UTexture2D* Texture : Textures)
 						{
-							(DrawModeComponent->*TextureSetter)(Texture);
+							if (!Texture)
+							{
+								continue;
+							}
+
+							if (UAssetImportData* ImportData = Texture->AssetImportData)
+							{
+								if (FPaths::IsSamePath(TexturePath, ImportData->GetFirstFilename()))
+								{
+									(DrawModeComponent->*TextureSetter)(Texture);
+								}
+							}
 						}
 					}
 				}
diff --git a/Engine/Plugins/Importers/USDImporter/Source/USDUtilities/Private/USDPrimConversion.cpp b/Engine/Plugins/Importers/USDImporter/Source/USDUtilities/Private/USDPrimConversion.cpp
--- a/Engine/Plugins/Importers/USDImporter/Source/USDUtilities/Private/USDPrimConversion.cpp
+++ b/Engine/Plugins/Importers/USDImporter/Source/USDUtilities/Private/USDPrimConversion.cpp
@@ -1668,7 +1668,12 @@ UsdToUnreal::FPropertyTrackReader UsdToUnreal::CreatePropertyTrackReader(
 	return Reader;
 }
 
-bool UsdToUnreal::ConvertDrawMode(const pxr::UsdPrim& Prim, UUsdDrawModeComponent* DrawModeComponent, double EvalTime, pxr::UsdGeomBBoxCache* BBoxCache)
+bool UsdToUnreal::ConvertDrawMode(
+	const pxr::UsdPrim& Prim,
+	UUsdDrawModeComponent* DrawModeComponent,
+	double EvalTime,
+	pxr::UsdGeomBBoxCache* BBoxCache
+)
 {
 	// We're not going to check if Prim actually has the "bounds" draw mode or if it has "applyDrawMode" set
 	// to true, as that can be expensive and this can get called from UpdateComponents, which can get called
@@ -5076,7 +5081,12 @@ TArray<UE::FUsdAttribute> UnrealToUsd::GetAttributesForProperty(const UE::FUsdPr
 	return {};
 }
 
-bool UnrealToUsd::ConvertDrawModeComponent(const UUsdDrawModeComponent& DrawModeComponent, pxr::UsdPrim& UsdPrim, bool bWriteExtents, double UsdTimeCode)
+bool UnrealToUsd::ConvertDrawModeComponent(
+	const UUsdDrawModeComponent& DrawModeComponent,
+	pxr::UsdPrim& UsdPrim,
+	bool bWriteExtents,
+	double UsdTimeCode
+)
 {
 	if (!UsdPrim)
 	{
@@ -5158,11 +5168,11 @@ bool UnrealToUsd::ConvertDrawModeComponent(const UUsdDrawModeComponent& DrawMode
 	if (pxr::UsdAttribute Attr = GeomModelAPI.CreateModelDrawModeAttr())
 	{
 		Attr.Set(
-			DrawModeComponent.DrawMode == EUsdDrawMode::Origin	  ? pxr::UsdGeomTokens->origin
-			: DrawModeComponent.DrawMode == EUsdDrawMode::Bounds	  ? pxr::UsdGeomTokens->bounds
-			: DrawModeComponent.DrawMode == EUsdDrawMode::Cards	  ? pxr::UsdGeomTokens->cards
+			DrawModeComponent.DrawMode == EUsdDrawMode::Origin		? pxr::UsdGeomTokens->origin
+			: DrawModeComponent.DrawMode == EUsdDrawMode::Bounds	? pxr::UsdGeomTokens->bounds
+			: DrawModeComponent.DrawMode == EUsdDrawMode::Cards		? pxr::UsdGeomTokens->cards
 			: DrawModeComponent.DrawMode == EUsdDrawMode::Inherited ? pxr::UsdGeomTokens->inherited
-																  : pxr::UsdGeomTokens->default_,
+																	: pxr::UsdGeomTokens->default_,
 			UsdTimeCode
 		);
 	}
@@ -5172,7 +5182,7 @@ bool UnrealToUsd::ConvertDrawModeComponent(const UUsdDrawModeComponent& DrawMode
 		Attr.Set(
 			DrawModeComponent.CardGeometry == EUsdModelCardGeometry::Cross ? pxr::UsdGeomTokens->cross
 			: DrawModeComponent.CardGeometry == EUsdModelCardGeometry::Box ? pxr::UsdGeomTokens->box
-																		 : pxr::UsdGeomTokens->fromTexture,
+																		   : pxr::UsdGeomTokens->fromTexture,
 			UsdTimeCode
 		);
 	}
@@ -5205,7 +5215,10 @@ bool UnrealToUsd::ConvertDrawModeComponent(const UUsdDrawModeComponent& DrawMode
 	using GetAttrFuncType = decltype(&pxr::UsdGeomModelAPI::GetModelCardTextureXPosAttr);
 
 	TFunction<void(EUsdModelCardFace, GetAttrFuncType, CreateAttrFuncType)> ExportCardFace =
-		[&DrawModeComponent, AuthoredFaces, UsdTimeCode, &GeomModelAPI](EUsdModelCardFace Face, GetAttrFuncType GetAttr, CreateAttrFuncType CreateAttr)
+		[&DrawModeComponent,
+		 AuthoredFaces,
+		 UsdTimeCode,
+		 &GeomModelAPI](EUsdModelCardFace Face, GetAttrFuncType GetAttr, CreateAttrFuncType CreateAttr)
 	{
 		UTexture2D* FaceTexture = DrawModeComponent.GetTextureForFace(Face);
 
@@ -5283,36 +5296,43 @@ bool UnrealToUsd::ConvertDrawModeComponent(const UUsdDrawModeComponent& DrawMode
 			}
 		}
 	};
-	ExportCardFace(
-		EUsdModelCardFace::XPos,
-		&pxr::UsdGeomModelAPI::GetModelCardTextureXPosAttr,
-		&pxr::UsdGeomModelAPI::CreateModelCardTextureXPosAttr
-	);
-	ExportCardFace(
-		EUsdModelCardFace::YPos,
-		&pxr::UsdGeomModelAPI::GetModelCardTextureYPosAttr,
-		&pxr::UsdGeomModelAPI::CreateModelCardTextureYPosAttr
-	);
-	ExportCardFace(
-		EUsdModelCardFace::ZPos,
-		&pxr::UsdGeomModelAPI::GetModelCardTextureZPosAttr,
-		&pxr::UsdGeomModelAPI::CreateModelCardTextureZPosAttr
-	);
-	ExportCardFace(
-		EUsdModelCardFace::XNeg,
-		&pxr::UsdGeomModelAPI::GetModelCardTextureXNegAttr,
-		&pxr::UsdGeomModelAPI::CreateModelCardTextureXNegAttr
-	);
-	ExportCardFace(
-		EUsdModelCardFace::YNeg,
-		&pxr::UsdGeomModelAPI::GetModelCardTextureYNegAttr,
-		&pxr::UsdGeomModelAPI::CreateModelCardTextureYNegAttr
-	);
-	ExportCardFace(
-		EUsdModelCardFace::ZNeg,
-		&pxr::UsdGeomModelAPI::GetModelCardTextureZNegAttr,
-		&pxr::UsdGeomModelAPI::CreateModelCardTextureZNegAttr
-	);
+
+	GetAttrFuncType XPosGet = &pxr::UsdGeomModelAPI::GetModelCardTextureXPosAttr;
+	CreateAttrFuncType XPosCreate = &pxr::UsdGeomModelAPI::CreateModelCardTextureXPosAttr;
+	GetAttrFuncType YPosGet = &pxr::UsdGeomModelAPI::GetModelCardTextureYPosAttr;
+	CreateAttrFuncType YPosCreate = &pxr::UsdGeomModelAPI::CreateModelCardTextureYPosAttr;
+	GetAttrFuncType ZPosGet = &pxr::UsdGeomModelAPI::GetModelCardTextureZPosAttr;
+	CreateAttrFuncType ZPosCreate = &pxr::UsdGeomModelAPI::CreateModelCardTextureZPosAttr;
+	GetAttrFuncType XNegGet = &pxr::UsdGeomModelAPI::GetModelCardTextureXNegAttr;
+	CreateAttrFuncType XNegCreate = &pxr::UsdGeomModelAPI::CreateModelCardTextureXNegAttr;
+	GetAttrFuncType YNegGet = &pxr::UsdGeomModelAPI::GetModelCardTextureYNegAttr;
+	CreateAttrFuncType YNegCreate = &pxr::UsdGeomModelAPI::CreateModelCardTextureYNegAttr;
+	GetAttrFuncType ZNegGet = &pxr::UsdGeomModelAPI::GetModelCardTextureZNegAttr;
+	CreateAttrFuncType ZNegCreate = &pxr::UsdGeomModelAPI::CreateModelCardTextureZNegAttr;
+
+	// We swap these when importing so that they look right in UE (e.g. ZPos is always pointing at UE +Z axis),
+	// so when writing back we need to swap back too
+	FUsdStageInfo StageInfo{UsdPrim.GetStage()};
+	if (StageInfo.UpAxis == EUsdUpAxis::ZAxis)
+	{
+		Swap(YPosGet, YNegGet);
+		Swap(YPosCreate, YNegCreate);
+	}
+	else
+	{
+		Swap(YPosGet, ZPosGet);
+		Swap(YPosCreate, ZPosCreate);
+
+		Swap(YNegGet, ZNegGet);
+		Swap(YNegCreate, ZNegCreate);
+	}
+
+	ExportCardFace(EUsdModelCardFace::XPos, XPosGet, XPosCreate);
+	ExportCardFace(EUsdModelCardFace::YPos, YPosGet, YPosCreate);
+	ExportCardFace(EUsdModelCardFace::ZPos, ZPosGet, ZPosCreate);
+	ExportCardFace(EUsdModelCardFace::XNeg, XNegGet, XNegCreate);
+	ExportCardFace(EUsdModelCardFace::YNeg, YNegGet, YNegCreate);
+	ExportCardFace(EUsdModelCardFace::ZNeg, ZNegGet, ZNegCreate);
 
 	return true;
 }
