USD: Remove fallback of generating universal render context materials when importing with the 'unreal' render context.

diff --git a/Engine/Plugins/Importers/USDImporter/Source/USDSchemas/Private/USDShadeMaterialTranslator.cpp b/Engine/Plugins/Importers/USDImporter/Source/USDSchemas/Private/USDShadeMaterialTranslator.cpp
--- a/Engine/Plugins/Importers/USDImporter/Source/USDSchemas/Private/USDShadeMaterialTranslator.cpp
+++ b/Engine/Plugins/Importers/USDImporter/Source/USDSchemas/Private/USDShadeMaterialTranslator.cpp
@@ -259,6 +259,28 @@ void FUsdShadeMaterialTranslator::CreateAssets()
 		TOptional<FString> UnrealMaterial = UsdUtils::GetUnrealSurfaceOutput(ShadeMaterial.GetPrim());
 		if (UnrealMaterial.IsSet())
 		{
+			UE_LOG(
+				LogUsd,
+				Log,
+				TEXT(
+					"Skipping generation of assets for material prim '%s' as all prims that bind this material will use its referenced Unreal material '%s' instead."
+				),
+				*PrimPath.GetString(),
+				*UnrealMaterial.GetValue()
+			);
+
+			UObject* Object = FSoftObjectPath(UnrealMaterial.GetValue()).TryLoad();
+			if (!Object)
+			{
+				UE_LOG(
+					LogUsd,
+					Warning,
+					TEXT("Failed to find the Unreal material '%s' referenced by material prim '%s'."),
+					*UnrealMaterial.GetValue(),
+					*PrimPath.GetString()
+				);
+			}
+
 			return;
 		}
 	}
@@ -305,6 +327,7 @@ void FUsdShadeMaterialTranslator::CreateAssets()
 				);
 				if (!bSuccess)
 				{
+					NewMaterial->MarkAsGarbage();
 					return;
 				}
 
diff --git a/Engine/Plugins/Importers/USDImporter/Source/USDUtilities/Private/USDShadeConversion.cpp b/Engine/Plugins/Importers/USDImporter/Source/USDUtilities/Private/USDShadeConversion.cpp
--- a/Engine/Plugins/Importers/USDImporter/Source/USDUtilities/Private/USDShadeConversion.cpp
+++ b/Engine/Plugins/Importers/USDImporter/Source/USDUtilities/Private/USDShadeConversion.cpp
@@ -2165,13 +2165,7 @@ bool UsdToUnreal::ConvertMaterial(
 	if (RenderContext)
 	{
 		pxr::TfToken ProvidedRenderContextToken = UnrealToUsd::ConvertToken(RenderContext).Get();
-
-		// We won't ever create an asset for the Unreal render context material prim.
-		// Check the universal context here in case this material should also generate an asset for that context too
-		if (ProvidedRenderContextToken != UnrealIdentifiers::Unreal)
-		{
-			RenderContextToken = ProvidedRenderContextToken;
-		}
+		RenderContextToken = ProvidedRenderContextToken;
 	}
 
 	pxr::UsdShadeShader SurfaceShader = UsdShadeMaterial.ComputeSurfaceSource(RenderContextToken);
