USD: Fix Sequencer tracks not receiving keys if they refer to a non-local layer.

diff --git a/Engine/Plugins/Importers/USDImporter/Source/USDStage/Private/USDLevelSequenceHelper.cpp b/Engine/Plugins/Importers/USDImporter/Source/USDStage/Private/USDLevelSequenceHelper.cpp
--- a/Engine/Plugins/Importers/USDImporter/Source/USDStage/Private/USDLevelSequenceHelper.cpp
+++ b/Engine/Plugins/Importers/USDImporter/Source/USDStage/Private/USDLevelSequenceHelper.cpp
@@ -111,6 +111,63 @@ namespace UsdLevelSequenceHelperImpl
 		return FFrameNumber(static_cast<int32>(FMath::RoundToDouble(TimeAsFrame)));
 	}
 
+	/**
+	 * We always want to mark the LevelSequences we spawn for non-local layers as read-only. This because our current
+	 * approach is that only local layers can be written to, meaning there is no point in allowing the user to edit these
+	 * sequences as those changes won't be written out to USD.
+	 * We use this struct to let us temporarily set a MovieScene to ReadOnly == false while we're adding keyframes to it.
+	 */
+	struct FScopedReadOnlyDisable
+	{
+		FScopedReadOnlyDisable(UMovieScene* InMovieScene, const UE::FSdfLayer& InLayer, const UE::FUsdStage& InOwnerStage)
+			: MovieScene(InMovieScene)
+			, Layer(InLayer)
+			, OwnerStage(InOwnerStage)
+		{
+#if WITH_EDITOR
+			if (MovieScene)
+			{
+				// Keep track of movie scenes that were already read-only too: Maybe the user or some
+				// other mechanism made them that way, so we'll want to put those back later
+				bWasReadOnly = MovieScene->IsReadOnly();
+			}
+
+			MovieScene->SetReadOnly(false);
+#endif	  // WITH_EDITOR
+		}
+
+		~FScopedReadOnlyDisable()
+		{
+#if WITH_EDITOR
+			bool bRestoreReadOnly = bWasReadOnly;
+
+			// If the sequence originally was ReadOnly for any reason, we know we need to put it back to ReadOnly.
+			// Otherwise, we want to set it as ReadOnly only if Layer is not part of the stage's local layer stack.
+			if (!bWasReadOnly && OwnerStage && Layer)
+			{
+				bRestoreReadOnly = !OwnerStage.HasLocalLayer(Layer);
+			}
+
+			if (bRestoreReadOnly)
+			{
+				MovieScene->SetReadOnly(true);
+			}
+#endif	  // WITH_EDITOR
+		}
+
+	private:
+		FScopedReadOnlyDisable(const FScopedReadOnlyDisable& Other) = delete;
+		FScopedReadOnlyDisable(FScopedReadOnlyDisable&& Other) = delete;
+		FScopedReadOnlyDisable& operator=(const FScopedReadOnlyDisable& Other) = delete;
+		FScopedReadOnlyDisable& operator=(FScopedReadOnlyDisable&& Other) = delete;
+
+	private:
+		bool bWasReadOnly = false;
+		UMovieScene* MovieScene = nullptr;
+		const UE::FSdfLayer Layer;
+		const UE::FUsdStage OwnerStage;
+	};
+
 	// Like UMovieScene::FindTrack, except that if we require class T it will return a track of type T or any type that derives from T
 	template<typename TrackType>
 	TrackType* FindTrackTypeOrDerived(const UMovieScene* MovieScene, const FGuid& Guid, const FName& TrackName = NAME_None)
@@ -599,7 +656,7 @@ public:
 
 private:
 	ULevelSequence* FindSequenceForAttribute(const UE::FUsdAttribute& Attribute);
-	ULevelSequence* FindOrAddSequenceForAttribute(const UE::FUsdAttribute& Attribute);
+	ULevelSequence* FindOrAddSequenceForAttribute(const UE::FUsdAttribute& Attribute, UE::FSdfLayer* OutSequenceLayer = nullptr);
 	ULevelSequence* FindSequenceForIdentifier(const FString& SequenceIdentitifer);
 	ULevelSequence* FindOrAddSequenceForLayer(const UE::FSdfLayer& Layer, const FString& SequenceIdentifier, const FString& SequenceDisplayName);
 	UE::FSdfLayer FindEditTargetForSubsequence(ULevelSequence* Sequence);
@@ -1086,19 +1143,13 @@ ULevelSequence* FUsdLevelSequenceHelperImpl::FindSequenceForAttribute(const UE::
 	}
 
 	UE::FSdfLayer AttributeLayer = UsdUtils::FindLayerForAttribute(Attribute, 0.0);
-
 	if (!AttributeLayer)
 	{
 		return nullptr;
 	}
 
-	FString AttributeLayerIdentifier = AttributeLayer.GetIdentifier();
-
 	UE::FUsdPrim Prim = Attribute.GetPrim();
 
-	UE::FSdfLayer PrimLayer = UsdUtils::FindLayerForPrim(Prim);
-	FString PrimLayerIdentifier = PrimLayer.GetIdentifier();
-
 	ULevelSequence* Sequence = nullptr;
 
 	// If the attribute is on the Root or a SubLayer, return the Sequence associated with that layer
@@ -1115,7 +1166,7 @@ ULevelSequence* FUsdLevelSequenceHelperImpl::FindSequenceForAttribute(const UE::
 	return Sequence;
 }
 
-ULevelSequence* FUsdLevelSequenceHelperImpl::FindOrAddSequenceForAttribute(const UE::FUsdAttribute& Attribute)
+ULevelSequence* FUsdLevelSequenceHelperImpl::FindOrAddSequenceForAttribute(const UE::FUsdAttribute& Attribute, UE::FSdfLayer* OutSequenceLayer)
 {
 	if (!Attribute || !Attribute.GetPrim())
 	{
@@ -1129,15 +1180,10 @@ ULevelSequence* FUsdLevelSequenceHelperImpl::FindOrAddSequenceForAttribute(const
 		{
 			const FString SequenceIdentifier = AttributeLayer.GetIdentifier();
 			Sequence = FindOrAddSequenceForLayer(AttributeLayer, SequenceIdentifier, SequenceIdentifier);
-
-#if WITH_EDITORONLY_DATA
-			// Make level sequences for non-local layers read-only: Our current approach is that only
-			// local layers can be written to, so there is no point in doing otherwise
-			if (Sequence && UsdStage && !UsdStage.HasLocalLayer(AttributeLayer))
+			if (OutSequenceLayer)
 			{
-				Sequence->GetMovieScene()->SetReadOnly(true);
+				*OutSequenceLayer = AttributeLayer;
 			}
-#endif
 		}
 	}
 
@@ -1187,6 +1233,10 @@ ULevelSequence* FUsdLevelSequenceHelperImpl::FindOrAddSequenceForLayer(
 		LevelSequencesByIdentifier.Add(SequenceIdentifier, Sequence);
 		IdentifierByLevelSequence.Add(Sequence, SequenceIdentifier);
 
+		// Here we abuse the FScopedReadOnlyDisable so that we can use the code in its destructor to set
+		// Sequence to ReadOnly if Layer doesn't belong to UsdStage's local layer stack
+		UsdLevelSequenceHelperImpl::FScopedReadOnlyDisable SetToReadOnly{MovieScene, Layer, UsdStage};
+
 		const FLayerTimeInfo LayerTimeInfo = FindOrAddLayerTimeInfo(Layer);
 
 		UpdateMovieSceneTimeRanges(*MovieScene, LayerTimeInfo);
@@ -1656,6 +1706,8 @@ void FUsdLevelSequenceHelperImpl::AddCommonTracks(const UUsdPrimTwin& PrimTwin,
 
 			if (UMovieScene* MovieScene = AttributeSequence->GetMovieScene())
 			{
+				UsdLevelSequenceHelperImpl::FScopedReadOnlyDisable ReadOnlyGuard{MovieScene, Layer, UsdStage};
+
 				if (bNeedTrackToCompensateResetXformOp)
 				{
 					TimeSampleUnion.Append(AncestorTimeSamples);
@@ -1753,7 +1805,8 @@ void FUsdLevelSequenceHelperImpl::AddCommonTracks(const UUsdPrimTwin& PrimTwin,
 
 			if (AttributeForSequence && TotalVisibilityTimeSamples.Num() > 0)
 			{
-				if (ULevelSequence* AttributeSequence = FindOrAddSequenceForAttribute(AttributeForSequence))
+				UE::FSdfLayer SequenceLayer;
+				if (ULevelSequence* AttributeSequence = FindOrAddSequenceForAttribute(AttributeForSequence, &SequenceLayer))
 				{
 					const bool bIsMuted = UsdUtils::IsAttributeMuted(AttributeForSequence, UsdStage);
 
@@ -1766,6 +1819,8 @@ void FUsdLevelSequenceHelperImpl::AddCommonTracks(const UUsdPrimTwin& PrimTwin,
 
 					if (UMovieScene* MovieScene = AttributeSequence->GetMovieScene())
 					{
+						UsdLevelSequenceHelperImpl::FScopedReadOnlyDisable ReadOnlyGuard{MovieScene, SequenceLayer, UsdStage};
+
 						if (UMovieSceneVisibilityTrack* VisibilityTrack = AddTrack<UMovieSceneVisibilityTrack>(
 								UnrealIdentifiers::HiddenInGamePropertyName,
 								PrimTwin,
@@ -1835,6 +1890,7 @@ void FUsdLevelSequenceHelperImpl::AddBoundsTracks(const UUsdPrimTwin& PrimTwin,
 
 	// Find the Sequence where we'll author the tracks
 	ULevelSequence* TargetSequence = nullptr;
+	UE::FSdfLayer TargetLayer;
 	bool bIsMuted = false;
 	if (Prim.IsA(TEXT("Boundable")))
 	{
@@ -1842,7 +1898,7 @@ void FUsdLevelSequenceHelperImpl::AddBoundsTracks(const UUsdPrimTwin& PrimTwin,
 		if (ExtentAttr && ExtentAttr.HasAuthoredValue())
 		{
 			bIsMuted = UsdUtils::IsAttributeMuted(ExtentAttr, Prim.GetStage());
-			TargetSequence = FindOrAddSequenceForAttribute(ExtentAttr);
+			TargetSequence = FindOrAddSequenceForAttribute(ExtentAttr, &TargetLayer);
 		}
 	}
 	if (!TargetSequence)
@@ -1853,7 +1909,7 @@ void FUsdLevelSequenceHelperImpl::AddBoundsTracks(const UUsdPrimTwin& PrimTwin,
 			if (ExtentsHintAttr && ExtentsHintAttr.HasAuthoredValue())
 			{
 				bIsMuted = UsdUtils::IsAttributeMuted(ExtentsHintAttr, Prim.GetStage());
-				TargetSequence = FindOrAddSequenceForAttribute(ExtentsHintAttr);
+				TargetSequence = FindOrAddSequenceForAttribute(ExtentsHintAttr, &TargetLayer);
 			}
 		}
 	}
@@ -1865,6 +1921,7 @@ void FUsdLevelSequenceHelperImpl::AddBoundsTracks(const UUsdPrimTwin& PrimTwin,
 		// If the user manually modifies these, we'll author these as `extent` or `extentsHint` depending on the prim, but only on-demand.
 		UE::FSdfLayer PrimLayer = UsdUtils::FindLayerForPrim(Prim);
 		TargetSequence = FindSequenceForIdentifier(PrimLayer.GetIdentifier());
+		TargetLayer = PrimLayer;
 	}
 	if (!TargetSequence)
 	{
@@ -1884,6 +1941,8 @@ void FUsdLevelSequenceHelperImpl::AddBoundsTracks(const UUsdPrimTwin& PrimTwin,
 		return;
 	}
 
+	UsdLevelSequenceHelperImpl::FScopedReadOnlyDisable ReadOnlyGuard{MovieScene, TargetLayer, UsdStage};
+
 	UMovieSceneDoubleVectorTrack* MinTrack = AddTrack<UMovieSceneDoubleVectorTrack>(
 		GET_MEMBER_NAME_CHECKED(UUsdDrawModeComponent, BoundsMin),
 		PrimTwin,
@@ -1936,9 +1995,6 @@ void FUsdLevelSequenceHelperImpl::AddCameraTracks(const UUsdPrimTwin& PrimTwin,
 		UnrealIdentifiers::SensorWidthPropertyName,
 		UnrealIdentifiers::SensorHeightPropertyName};
 
-	UE::FSdfLayer PrimLayer = UsdUtils::FindLayerForPrim(Prim);
-	ULevelSequence* PrimSequence = FindSequenceForIdentifier(PrimLayer.GetIdentifier());
-
 	// For ACineCameraActor the camera component is not the actual root component, so we need to fetch it manually here
 	ACineCameraActor* CameraActor = Cast<ACineCameraActor>(PrimTwin.GetSceneComponent()->GetOwner());
 	if (!CameraActor)
@@ -1967,7 +2023,8 @@ void FUsdLevelSequenceHelperImpl::AddCameraTracks(const UUsdPrimTwin& PrimTwin,
 		}
 
 		// Find out the sequence where this attribute should be written to
-		if (ULevelSequence* AttributeSequence = FindOrAddSequenceForAttribute(Attr))
+		UE::FSdfLayer SequenceLayer;
+		if (ULevelSequence* AttributeSequence = FindOrAddSequenceForAttribute(Attr, &SequenceLayer))
 		{
 			const bool bIsMuted = UsdUtils::IsAttributeMuted(Attr, UsdStage);
 
@@ -1984,6 +2041,8 @@ void FUsdLevelSequenceHelperImpl::AddCameraTracks(const UUsdPrimTwin& PrimTwin,
 				continue;
 			}
 
+			UsdLevelSequenceHelperImpl::FScopedReadOnlyDisable ReadOnlyGuard{MovieScene, SequenceLayer, UsdStage};
+
 			TArray<double> TimeSamples;
 			if (!Attr.GetTimeSamples(TimeSamples))
 			{
@@ -2117,6 +2176,8 @@ void FUsdLevelSequenceHelperImpl::AddLightTracks(const UUsdPrimTwin& PrimTwin, c
 			continue;
 		}
 
+		UsdLevelSequenceHelperImpl::FScopedReadOnlyDisable ReadOnlyGuard{MovieScene, PrimLayer, UsdStage};
+
 		UsdToUnreal::FPropertyTrackReader Reader = UsdToUnreal::CreatePropertyTrackReader(Prim, PropertyPath);
 
 		switch (TrackType)
@@ -2222,6 +2283,8 @@ void FUsdLevelSequenceHelperImpl::AddSkeletalTracks(const UUsdPrimTwin& PrimTwin
 		return;
 	}
 
+	UsdLevelSequenceHelperImpl::FScopedReadOnlyDisable ReadOnlyGuard{MovieScene, SkelAnimationLayer, UsdStage};
+
 	// We will mute all SkelAnimation attributes if we mute, so here let's only consider something muted
 	// if it has all attributes muted as well.
 	// We know at least one of these attributes ones is valid and animated because we have an UAnimSequence
@@ -2295,6 +2358,8 @@ void FUsdLevelSequenceHelperImpl::AddGeometryCacheTracks(const UUsdPrimTwin& Pri
 		return;
 	}
 
+	UsdLevelSequenceHelperImpl::FScopedReadOnlyDisable ReadOnlyGuard{MovieScene, GeometryCacheLayer, UsdStage};
+
 	const bool bIsMuted = false;
 	if (UMovieSceneGeometryCacheTrack* GeometryCacheTrack = AddTrack<
 			UMovieSceneGeometryCacheTrack>(Prim.GetName(), PrimTwin, *ComponentToBind, *GeometryCacheSequence, bIsMuted))
@@ -2368,6 +2433,8 @@ void FUsdLevelSequenceHelperImpl::AddGroomTracks(const UUsdPrimTwin& PrimTwin, c
 		return;
 	}
 
+	UsdLevelSequenceHelperImpl::FScopedReadOnlyDisable ReadOnlyGuard{MovieScene, GroomLayer, UsdStage};
+
 	const bool bIsMuted = false;
 	if (UMovieSceneGroomCacheTrack* GroomCacheTrack = AddTrack<
 			UMovieSceneGroomCacheTrack>(Prim.GetName(), PrimTwin, *ComponentToBind, *GroomAnimationSequence, bIsMuted))
@@ -2405,6 +2472,8 @@ void FUsdLevelSequenceHelperImpl::AddVolumeTracks(const UUsdPrimTwin& PrimTwin,
 		return;
 	}
 
+	UsdLevelSequenceHelperImpl::FScopedReadOnlyDisable ReadOnlyGuard{MovieScene, PrimLayer, UsdStage};
+
 	const FName PropertyPath = GET_MEMBER_NAME_CHECKED(UHeterogeneousVolumeComponent, Frame);
 
 	// Here we'll just get *any* of the filePath attrs from this Volume prim to check for the
