USD: Fix not being able to isolate sublayers children of session layers;

Copy over session layer when saving a temp stage into a persistent stage;
Fix root layer being marked as dirty after saving, in some contexts;
Fix UnrealUSDWrapper::OpenStage not using the provided session layer in some situations.

diff --git a/Engine/Plugins/Importers/USDImporter/Source/USDStage/Private/USDStageActor.cpp b/Engine/Plugins/Importers/USDImporter/Source/USDStage/Private/USDStageActor.cpp
--- a/Engine/Plugins/Importers/USDImporter/Source/USDStage/Private/USDStageActor.cpp
+++ b/Engine/Plugins/Importers/USDImporter/Source/USDStage/Private/USDStageActor.cpp
@@ -1339,21 +1339,15 @@ void AUsdStageActor::IsolateLayer(const UE::FSdfLayer& Layer)
 		// We should only be allowed to isolate a layer belonging to UsdStage's local layer stack, but checking for that
 		// is not trivial given that layers can be muted.
 
-		TArray<UE::FSdfLayer> ValidLayers;
+		const bool bIncludeSessionLayers = true;
+		TSet<UE::FSdfLayer> ValidLayers{UsdStage.GetLayerStack(bIncludeSessionLayers)};
+
 		UE::FUsdStage FreshCurrentStage;
-		if (Layer.IsAnonymous() && UsdStage.GetRootLayer() != Layer)
+		if (!Layer.IsAnonymous() && !ValidLayers.Contains(Layer))
 		{
-			// If we're an anonymous layer that is *not* the stage's root layer, it means that we're
-			// probably a session layer. We can't use the trick below as opening a new stage would
-			// get a brand new session layer, so for a proper check we must use the layer stack of the
-			// currently opened stage
-			const bool bIncludeSessionLayers = true;
-			ValidLayers = UsdStage.GetLayerStack(bIncludeSessionLayers);
-		}
-		else
-		{
-			// Check if the layer we're trying to isolate is part of the stage. To do that we need to reopen a fresh copy
-			// of the current stage in case we have any layers muted, as muted layers don't show up on the layer stack.
+			// If the layer has a file on disk but ValidLayers does not contain it, there's still a chance that this
+			// is in fact part of the usual layer stack of the stage but is currently muted. To check for that we need
+			// to reopen a fresh copy of the stage, as muted layers don't usually show up on the layer stack.
 			// Note that we can't just check the list of muted layers either, as it's possible to mute *any* layer for
 			// a given stage, not only the layers that are currently used by it.
 			// We'll use an empty population mask though (which should prevent prim composition) and just use the layers
@@ -1361,8 +1355,7 @@ void AUsdStageActor::IsolateLayer(const UE::FSdfLayer& Layer)
 			FreshCurrentStage = UnrealUSDWrapper::OpenMaskedStage(*UsdStage.GetRootLayer().GetIdentifier(), EUsdInitialLoadSet::LoadNone, {});
 			ensure(FreshCurrentStage);
 
-			const bool bIncludeSessionLayers = true;
-			ValidLayers = FreshCurrentStage.GetLayerStack(bIncludeSessionLayers);
+			ValidLayers.Append(FreshCurrentStage.GetLayerStack(bIncludeSessionLayers));
 		}
 
 		if (!ValidLayers.Contains(Layer))
@@ -2337,7 +2330,14 @@ void AUsdStageActor::SetUsdStage(const UE::FUsdStage& NewStage)
 	UnloadUsdStage();
 	CloseUsdStage();
 
-	RootLayer.FilePath = NewStage.GetRootLayer().GetIdentifier();
+	FString RelativeFilePath = NewStage.GetRootLayer().GetIdentifier();
+#if USE_USD_SDK
+	if (!RelativeFilePath.IsEmpty() && !FPaths::IsRelative(RelativeFilePath) && !RelativeFilePath.StartsWith(UnrealIdentifiers::IdentifierPrefix))
+	{
+		RelativeFilePath = UsdUtils::MakePathRelativeToProjectDir(RelativeFilePath);
+	}
+#endif	  // USE_USD_SDK
+	RootLayer.FilePath = RelativeFilePath;
 
 	UsdStage = NewStage;
 	IsolatedStage = UE::FUsdStage{};
diff --git a/Engine/Plugins/Importers/USDImporter/Source/USDStageEditorViewModels/Private/USDStageViewModel.cpp b/Engine/Plugins/Importers/USDImporter/Source/USDStageEditorViewModels/Private/USDStageViewModel.cpp
--- a/Engine/Plugins/Importers/USDImporter/Source/USDStageEditorViewModels/Private/USDStageViewModel.cpp
+++ b/Engine/Plugins/Importers/USDImporter/Source/USDStageEditorViewModels/Private/USDStageViewModel.cpp
@@ -266,10 +266,23 @@ void FUsdStageViewModel::SaveStageAs(const TCHAR* FilePath)
 					FString OldComment = RootLayer.GetComment();
 					RootLayer.SetComment(UnrealIdentifiers::LayerSavedComment);
 
-					OpenStage(FilePath);
+					// Open our stage right here so that we can provide it with a session layer. Note that we'll be using
+					// the same session layer as the original stage, so that the user doesn't lose their session
+					UE::FSdfLayer NewRootLayer = UE::FSdfLayer::FindOrOpen(FilePath);
+					UE::FUsdStage NewStage = UnrealUSDWrapper::OpenStage(NewRootLayer, UsdStage.GetSessionLayer(), StageActorPtr->InitialLoadSet);
+					StageActorPtr->SetUsdStage(NewStage);
 
 					RootLayer.SetComment(*OldComment);
 
+					// Annoyingly, if we're in a workflow where NewRootLayer == RootLayer (i.e. delete root layer
+					// on disk -> press Save -> pick the same filepath as the original root layer), adding/removing that
+					// comment above will cause the layer to be marked as dirty... So here we save that layer alone again
+					// to remove the dirtiness state
+					if (NewRootLayer && RootLayer == NewRootLayer && NewRootLayer.IsDirty())
+					{
+						NewRootLayer.Save();
+					}
+
 					UsdViewModelImpl::SaveUEStateLayer(UsdStage);
 				}
 			}
diff --git a/Engine/Plugins/Importers/USDImporter/Source/UnrealUSDWrapper/Private/UnrealUSDWrapper.cpp b/Engine/Plugins/Importers/USDImporter/Source/UnrealUSDWrapper/Private/UnrealUSDWrapper.cpp
--- a/Engine/Plugins/Importers/USDImporter/Source/UnrealUSDWrapper/Private/UnrealUSDWrapper.cpp
+++ b/Engine/Plugins/Importers/USDImporter/Source/UnrealUSDWrapper/Private/UnrealUSDWrapper.cpp
@@ -818,53 +818,39 @@ namespace UE::UnrealUSDWrapper::Private
 		static_assert((int)pxr::UsdStage::InitialLoadSet::LoadNone == (int)EUsdInitialLoadSet::LoadNone);
 		pxr::UsdStage::InitialLoadSet LoadSet = static_cast<pxr::UsdStage::InitialLoadSet>(InitialLoadSet);
 
-		FString IdentifierStr = FString(RootIdentifier);
-		if (FPaths::FileExists(IdentifierStr))
+		FString RootIdentifierStr = FString{RootIdentifier};
+		FString SessionIdentifierStr = FString{SessionIdentifier};
+
+		RootIdentifierStr.RemoveFromStart(UnrealIdentifiers::IdentifierPrefix);
+		SessionIdentifierStr.RemoveFromStart(UnrealIdentifiers::IdentifierPrefix);
+
+		pxr::SdfLayerRefPtr RootLayer = pxr::SdfLayer::FindOrOpen(TCHAR_TO_ANSI(*RootIdentifierStr));
+		pxr::SdfLayerRefPtr SessionLayer = pxr::SdfLayer::FindOrOpen(TCHAR_TO_ANSI(*SessionIdentifierStr));
+		if (RootLayer)
 		{
 			if (PopulationMask)
 			{
-				Stage = pxr::UsdStage::OpenMasked(TCHAR_TO_ANSI(*IdentifierStr), Mask, LoadSet);
+				// We use this additional check so we don't have to replicate USD's "_CreateAnonymousSessionLayer"
+				// Basically we can't pass an invalid session layer pointer here the stage will actually end up
+				// with no session layer at all
+				if (SessionLayer)
+				{
+					Stage = pxr::UsdStage::OpenMasked(RootLayer, SessionLayer, Mask, LoadSet);
+				}
+				else
+				{
+					Stage = pxr::UsdStage::OpenMasked(RootLayer, Mask, LoadSet);
+				}
 			}
 			else
 			{
-				Stage = pxr::UsdStage::Open(TCHAR_TO_ANSI(*IdentifierStr), LoadSet);
-			}
-		}
-		else
-		{
-			FString SessionIdentifierStr = FString{SessionIdentifier};
-
-			IdentifierStr.RemoveFromStart(UnrealIdentifiers::IdentifierPrefix);
-			SessionIdentifierStr.RemoveFromStart(UnrealIdentifiers::IdentifierPrefix);
-
-			pxr::SdfLayerRefPtr RootLayer = pxr::SdfLayer::Find(TCHAR_TO_ANSI(*IdentifierStr));
-			pxr::SdfLayerRefPtr SessionLayer = pxr::SdfLayer::Find(TCHAR_TO_ANSI(*SessionIdentifierStr));
-			if (RootLayer)
-			{
-				if (PopulationMask)
+				if (SessionLayer)
 				{
-					// We use this additional check so we don't have to replicate USD's "_CreateAnonymousSessionLayer"
-					// Basically we can't pass an invalid session layer pointer here the stage will actually end up
-					// with no session layer at all
-					if (SessionLayer)
-					{
-						Stage = pxr::UsdStage::OpenMasked(RootLayer, SessionLayer, Mask, LoadSet);
-					}
-					else
-					{
-						Stage = pxr::UsdStage::OpenMasked(RootLayer, Mask, LoadSet);
-					}
+					Stage = pxr::UsdStage::Open(RootLayer, SessionLayer, LoadSet);
 				}
 				else
 				{
-					if (SessionLayer)
-					{
-						Stage = pxr::UsdStage::Open(RootLayer, SessionLayer, LoadSet);
-					}
-					else
-					{
-						Stage = pxr::UsdStage::Open(RootLayer, LoadSet);
-					}
+					Stage = pxr::UsdStage::Open(RootLayer, LoadSet);
 				}
 			}
 		}
