USD: Simplify right-click menu item for isolating a layer.

diff --git a/Engine/Plugins/Importers/USDImporter/Source/USDStageEditor/Private/SUSDLayersTreeView.cpp b/Engine/Plugins/Importers/USDImporter/Source/USDStageEditor/Private/SUSDLayersTreeView.cpp
--- a/Engine/Plugins/Importers/USDImporter/Source/USDStageEditor/Private/SUSDLayersTreeView.cpp
+++ b/Engine/Plugins/Importers/USDImporter/Source/USDStageEditor/Private/SUSDLayersTreeView.cpp
@@ -1081,10 +1081,10 @@ TSharedPtr<SWidget> SUsdLayersTreeView::ConstructLayerContextMenu()
 					if (IsolatedStage)
 					{
 						TArray<FUsdLayerViewModelRef> MySelectedItems = GetSelectedItems();
-						if (MySelectedItems.Num() == 1 && IsolatedStage)
+						if (MySelectedItems.Num() == 1)
 						{
 							const UE::FSdfLayerWeak& Layer = MySelectedItems[0]->GetLayer();
-							if (Layer == IsolatedStage.GetRootLayer() || Layer == UsdStage.GetRootLayer())
+							if (Layer == IsolatedStage.GetRootLayer())
 							{
 								return LOCTEXT("StopIsolatingStage_Text", "Stop isolating");
 							}
@@ -1110,13 +1110,6 @@ TSharedPtr<SWidget> SUsdLayersTreeView::ConstructLayerContextMenu()
 									"Stops isolating this layer and go back to showing the full composed stage"
 								);
 							}
-							else if (IsolatedStage && Layer == UsdStage.GetRootLayer())
-							{
-								return LOCTEXT(
-									"LeaveIsolatedMode_ToolTip",
-									"Stops isolating on any layer and go back to showing the full composed stage"
-								);
-							}
 						}
 					}
 
@@ -1246,13 +1239,7 @@ bool SUsdLayersTreeView::CanIsolateSelectedLayer() const
 		const UE::FSdfLayerWeak& Layer = SelectedItem->GetLayer();
 		const bool bLayerIsStageRoot = Layer == UsdStage.GetRootLayer();
 
-		return SelectedItem->IsValid()
-			   && (
-				   // If we're right clicking a sublayer its always OK
-				   !bLayerIsStageRoot
-				   // If we're isolating and selecting either a sublayer or the root layer (i.e. not the session layer)
-				   || (IsolatedStage && (SelectedItem->ParentItem || bLayerIsStageRoot))
-			   );
+		return SelectedItem->IsValid() && !bLayerIsStageRoot;
 	}
 
 	return false;
