USD: Fix generated materials not updating after changing the reference material paths on project settings;

diff --git a/Engine/Plugins/Importers/USDImporter/Source/USDClasses/Private/USDClassesModule.cpp b/Engine/Plugins/Importers/USDImporter/Source/USDClasses/Private/USDClassesModule.cpp
--- a/Engine/Plugins/Importers/USDImporter/Source/USDClasses/Private/USDClassesModule.cpp
+++ b/Engine/Plugins/Importers/USDImporter/Source/USDClasses/Private/USDClassesModule.cpp
@@ -474,7 +474,7 @@ TOptional<IUsdClassesModule::FDisplayColorMaterial> IUsdClassesModule::FDisplayC
 	return Result;
 }
 
-UMaterialInstanceDynamic* IUsdClassesModule::CreateDisplayColorMaterialInstanceDynamic(const FDisplayColorMaterial& DisplayColorDescription)
+const FSoftObjectPath* IUsdClassesModule::GetReferenceMaterialPath(const FDisplayColorMaterial& DisplayColorDescription)
 {
 	const UUsdProjectSettings* Settings = GetDefault<UUsdProjectSettings>();
 	if (!Settings)
@@ -482,29 +482,33 @@ UMaterialInstanceDynamic* IUsdClassesModule::CreateDisplayColorMaterialInstanceD
 		return nullptr;
 	}
 
-	const FSoftObjectPath* ParentPathPtr = nullptr;
 	if (DisplayColorDescription.bHasOpacity)
 	{
 		if (DisplayColorDescription.bIsDoubleSided)
 		{
-			ParentPathPtr = &Settings->ReferenceDisplayColorAndOpacityTwoSidedMaterial;
+			return &Settings->ReferenceDisplayColorAndOpacityTwoSidedMaterial;
 		}
 		else
 		{
-			ParentPathPtr = &Settings->ReferenceDisplayColorAndOpacityMaterial;
+			return &Settings->ReferenceDisplayColorAndOpacityMaterial;
 		}
 	}
 	else
 	{
 		if (DisplayColorDescription.bIsDoubleSided)
 		{
-			ParentPathPtr = &Settings->ReferenceDisplayColorTwoSidedMaterial;
+			return &Settings->ReferenceDisplayColorTwoSidedMaterial;
 		}
 		else
 		{
-			ParentPathPtr = &Settings->ReferenceDisplayColorMaterial;
+			return &Settings->ReferenceDisplayColorMaterial;
 		}
 	}
+}
+
+UMaterialInstanceDynamic* IUsdClassesModule::CreateDisplayColorMaterialInstanceDynamic(const FDisplayColorMaterial& DisplayColorDescription)
+{
+	const FSoftObjectPath* ParentPathPtr = GetReferenceMaterialPath(DisplayColorDescription);
 	if (!ParentPathPtr)
 	{
 		return nullptr;
@@ -534,36 +538,12 @@ UMaterialInstanceDynamic* IUsdClassesModule::CreateDisplayColorMaterialInstanceD
 UMaterialInstanceConstant* IUsdClassesModule::CreateDisplayColorMaterialInstanceConstant(const FDisplayColorMaterial& DisplayColorDescription)
 {
 #if WITH_EDITOR
-	const UUsdProjectSettings* Settings = GetDefault<UUsdProjectSettings>();
-	if (!Settings)
+	const FSoftObjectPath* ParentPathPtr = GetReferenceMaterialPath(DisplayColorDescription);
+	if (!ParentPathPtr)
 	{
 		return nullptr;
 	}
 
-	const FSoftObjectPath* ParentPathPtr = nullptr;
-	if (DisplayColorDescription.bHasOpacity)
-	{
-		if (DisplayColorDescription.bIsDoubleSided)
-		{
-			ParentPathPtr = &Settings->ReferenceDisplayColorAndOpacityTwoSidedMaterial;
-		}
-		else
-		{
-			ParentPathPtr = &Settings->ReferenceDisplayColorAndOpacityMaterial;
-		}
-	}
-	else
-	{
-		if (DisplayColorDescription.bIsDoubleSided)
-		{
-			ParentPathPtr = &Settings->ReferenceDisplayColorTwoSidedMaterial;
-		}
-		else
-		{
-			ParentPathPtr = &Settings->ReferenceDisplayColorMaterial;
-		}
-	}
-
 	if (UMaterialInterface* ParentMaterial = Cast<UMaterialInterface>(ParentPathPtr->TryLoad()))
 	{
 		FName AssetName = MakeUniqueObjectName(
diff --git a/Engine/Plugins/Importers/USDImporter/Source/USDClasses/Public/USDClassesModule.h b/Engine/Plugins/Importers/USDImporter/Source/USDClasses/Public/USDClassesModule.h
--- a/Engine/Plugins/Importers/USDImporter/Source/USDClasses/Public/USDClassesModule.h
+++ b/Engine/Plugins/Importers/USDImporter/Source/USDClasses/Public/USDClassesModule.h
@@ -10,6 +10,15 @@ class UMaterialInstanceConstant;
 class UMaterialInstanceDynamic;
 struct FAnalyticsEventAttribute;
 
+enum class EUsdReferenceMaterialProperties : uint8
+{
+	None = 0,
+	Translucent = 1,
+	VT = 2,
+	TwoSided = 4
+};
+ENUM_CLASS_FLAGS(EUsdReferenceMaterialProperties)
+
 class IUsdClassesModule : public IModuleInterface
 {
 public:
@@ -64,6 +73,8 @@ public:
 		static TOptional<FDisplayColorMaterial> FromString(const FString& DisplayColorString);
 	};
 
+	USDCLASSES_API static const FSoftObjectPath* GetReferenceMaterialPath(const FDisplayColorMaterial& DisplayColorDescription);
+
 	USDCLASSES_API static UMaterialInstanceDynamic* CreateDisplayColorMaterialInstanceDynamic(const FDisplayColorMaterial& DisplayColorDescription);
 	USDCLASSES_API static UMaterialInstanceConstant* CreateDisplayColorMaterialInstanceConstant(const FDisplayColorMaterial& DisplayColorDescription);
 };
diff --git a/Engine/Plugins/Importers/USDImporter/Source/USDSchemas/Private/MeshTranslationImpl.cpp b/Engine/Plugins/Importers/USDImporter/Source/USDSchemas/Private/MeshTranslationImpl.cpp
--- a/Engine/Plugins/Importers/USDImporter/Source/USDSchemas/Private/MeshTranslationImpl.cpp
+++ b/Engine/Plugins/Importers/USDImporter/Source/USDSchemas/Private/MeshTranslationImpl.cpp
@@ -458,19 +458,39 @@ TMap<const UsdUtils::FUsdPrimMaterialSlot*, UMaterialInterface*> MeshTranslation
 			{
 				case UsdUtils::EPrimAssignmentType::DisplayColor:
 				{
-					const FString PrefixedHash = UsdUtils::GetAssetHashPrefix(UsdPrim, bReuseIdenticalAssets) + Slot.MaterialSource;
+					TOptional<IUsdClassesModule::FDisplayColorMaterial> DisplayColorDesc = IUsdClassesModule::FDisplayColorMaterial::FromString(
+						Slot.MaterialSource
+					);
 
-					// Try reusing an already created DisplayColor material
-					if (UMaterialInterface* ExistingMaterial = Cast<UMaterialInterface>(AssetCache.GetCachedAsset(PrefixedHash)))
+					if (DisplayColorDesc.IsSet())
 					{
-						Material = ExistingMaterial;
-					}
+						FString DisplayColorHash;
+						{
+							FSHAHash Hash;
+							FSHA1 SHA1;
+							SHA1.UpdateWithString(*Slot.MaterialSource, Slot.MaterialSource.Len());
 
-					// Need to create a new DisplayColor material
-					if (Material == nullptr)
-					{
-						if (TOptional<IUsdClassesModule::FDisplayColorMaterial>
-								DisplayColorDesc = IUsdClassesModule::FDisplayColorMaterial::FromString(Slot.MaterialSource))
+							const FSoftObjectPath* ReferencePath = IUsdClassesModule::GetReferenceMaterialPath(DisplayColorDesc.GetValue());
+							if (ReferencePath)
+							{
+								FString ReferencePathString = ReferencePath->ToString();
+								SHA1.UpdateWithString(*ReferencePathString, ReferencePathString.Len());
+							}
+
+							SHA1.Final();
+							SHA1.GetHash(&Hash.Hash[0]);
+							DisplayColorHash = Hash.ToString();
+						}
+						const FString PrefixedHash = UsdUtils::GetAssetHashPrefix(UsdPrim, bReuseIdenticalAssets) + DisplayColorHash;
+
+						// Try reusing an already created DisplayColor material
+						if (UMaterialInterface* ExistingMaterial = Cast<UMaterialInterface>(AssetCache.GetCachedAsset(PrefixedHash)))
+						{
+							Material = ExistingMaterial;
+						}
+
+						// Need to create a new DisplayColor material
+						if (Material == nullptr)
 						{
 							UMaterialInstance* MaterialInstance = nullptr;
 
diff --git a/Engine/Plugins/Importers/USDImporter/Source/USDSchemas/Private/USDShadeMaterialTranslator.cpp b/Engine/Plugins/Importers/USDImporter/Source/USDSchemas/Private/USDShadeMaterialTranslator.cpp
--- a/Engine/Plugins/Importers/USDImporter/Source/USDSchemas/Private/USDShadeMaterialTranslator.cpp
+++ b/Engine/Plugins/Importers/USDImporter/Source/USDSchemas/Private/USDShadeMaterialTranslator.cpp
@@ -7,6 +7,7 @@
 #include "USDClassesModule.h"
 #include "USDLog.h"
 #include "USDPrimConversion.h"
+#include "USDProjectSettings.h"
 #include "USDShadeConversion.h"
 #include "USDTypesConversion.h"
 
@@ -221,6 +222,43 @@ namespace UE::UsdShadeTranslator::Private
 			}
 		}
 	}
+
+	// We need to hash the reference material that we'll use, so that if this is changed we regenerate a new instance.
+	// However, unlike for displayColor materials, we can't really know *which* reference material we'll end up using
+	// until after we've already created it (which doesn't sound like it makes any sense but it's part of why we have those
+	// VT and double-sided "upgrade" mechanisms).
+	//
+	// If all we want is a hash, the solution can be simple though: Hash them all. Yea we may end up unnecessarily regenerating
+	// materials sometimes but changing the reference materials on the project settings should be rare.
+	void HashPreviewSurfaceReferences(FSHA1& InOutHash)
+	{
+		const UUsdProjectSettings* Settings = GetDefault<UUsdProjectSettings>();
+		if (!Settings)
+		{
+			return;
+		}
+
+		TArray<const FSoftObjectPath*> ReferenceMaterials = {
+			&Settings->ReferencePreviewSurfaceMaterial,
+			&Settings->ReferencePreviewSurfaceTranslucentMaterial,
+			&Settings->ReferencePreviewSurfaceTwoSidedMaterial,
+			&Settings->ReferencePreviewSurfaceTranslucentTwoSidedMaterial,
+			&Settings->ReferencePreviewSurfaceVTMaterial,
+			&Settings->ReferencePreviewSurfaceTranslucentVTMaterial,
+			&Settings->ReferencePreviewSurfaceTwoSidedVTMaterial,
+			&Settings->ReferencePreviewSurfaceTranslucentTwoSidedVTMaterial};
+
+		for (const FSoftObjectPath* ReferencePath : ReferenceMaterials)
+		{
+			if (!ReferencePath)
+			{
+				continue;
+			}
+
+			FString ReferencePathStr = ReferencePath->ToString();
+			InOutHash.UpdateWithString(*ReferencePathStr, ReferencePathStr.Len());
+		}
+	}
 }	 // namespace UE::UsdShadeTranslator::Private
 
 void FUsdShadeMaterialTranslator::CreateAssets()
@@ -285,8 +323,20 @@ void FUsdShadeMaterialTranslator::CreateAssets()
 		}
 	}
 
-	const FString PrefixedMaterialHash = UsdUtils::GetAssetHashPrefix(GetPrim(), Context->bReuseIdenticalAssets)
-										 + UsdUtils::HashShadeMaterial(ShadeMaterial, RenderContextToken).ToString();
+	FString MaterialHash;
+	{
+		FSHAHash OutHash;
+
+		FSHA1 SHA1;
+
+		UsdUtils::HashShadeMaterial(ShadeMaterial, SHA1, RenderContextToken);
+		UE::UsdShadeTranslator::Private::HashPreviewSurfaceReferences(SHA1);
+
+		SHA1.Final();
+		SHA1.GetHash(&OutHash.Hash[0]);
+		MaterialHash = OutHash.ToString();
+	}
+	const FString PrefixedMaterialHash = UsdUtils::GetAssetHashPrefix(GetPrim(), Context->bReuseIdenticalAssets) + MaterialHash;
 
 	UMaterialInterface* ConvertedMaterial = nullptr;
 
diff --git a/Engine/Plugins/Importers/USDImporter/Source/USDSchemas/Private/USDVolVolumeTranslator.cpp b/Engine/Plugins/Importers/USDImporter/Source/USDSchemas/Private/USDVolVolumeTranslator.cpp
--- a/Engine/Plugins/Importers/USDImporter/Source/USDSchemas/Private/USDVolVolumeTranslator.cpp
+++ b/Engine/Plugins/Importers/USDImporter/Source/USDSchemas/Private/USDVolVolumeTranslator.cpp
@@ -1205,9 +1205,10 @@ void FUsdVolVolumeTranslator::CreateAssets()
 		}
 	}
 	// Fall back to the default SVT material instead
+	const UUsdProjectSettings* ProjectSettings = nullptr;
 	if (!ReferenceMaterial)
 	{
-		const UUsdProjectSettings* ProjectSettings = GetDefault<UUsdProjectSettings>();
+		ProjectSettings = GetDefault<UUsdProjectSettings>();
 		if (!ProjectSettings)
 		{
 			return;
@@ -1249,6 +1250,11 @@ void FUsdVolVolumeTranslator::CreateAssets()
 	{
 		FSHA1 SHA1;
 		HashForVolumetricMaterial(ReferenceMaterial, MaterialParameterToTexture, SHA1);
+		if (ProjectSettings)
+		{
+			FString ReferencePathString = ProjectSettings->ReferenceDefaultSVTMaterial.ToString();
+			SHA1.UpdateWithString(*ReferencePathString, ReferencePathString.Len());
+		}
 		SHA1.Final();
 		SHA1.GetHash(&MaterialHash.Hash[0]);
 	}
diff --git a/Engine/Plugins/Importers/USDImporter/Source/USDUtilities/Private/USDShadeConversion.cpp b/Engine/Plugins/Importers/USDImporter/Source/USDUtilities/Private/USDShadeConversion.cpp
--- a/Engine/Plugins/Importers/USDImporter/Source/USDUtilities/Private/USDShadeConversion.cpp
+++ b/Engine/Plugins/Importers/USDImporter/Source/USDUtilities/Private/USDShadeConversion.cpp
@@ -3502,20 +3502,30 @@ bool UsdUtils::IsMaterialTranslucent(const pxr::UsdShadeMaterial& UsdShadeMateri
 }
 
 FSHAHash UsdUtils::HashShadeMaterial(const pxr::UsdShadeMaterial& UsdShadeMaterial, const pxr::TfToken& RenderContext)
+{
+	FSHAHash OutHash;
+
+	FSHA1 SHA1;
+	HashShadeMaterial(UsdShadeMaterial, SHA1, RenderContext);
+	SHA1.Final();
+	SHA1.GetHash(&OutHash.Hash[0]);
+
+	return OutHash;
+}
+
+void UsdUtils::HashShadeMaterial(const pxr::UsdShadeMaterial& UsdShadeMaterial, FSHA1& InOutHash, const pxr::TfToken& RenderContext)
 {
 	FScopedUsdAllocs UsdAllocs;
 
 	pxr::UsdShadeShader SurfaceShader = UsdShadeMaterial.ComputeSurfaceSource({RenderContext});
 	if (!SurfaceShader)
 	{
-		return {};
+		return;
 	}
 
-	FSHA1 HashState;
-
 	for (const pxr::UsdShadeInput& ShadeInput : SurfaceShader.GetInputs())
 	{
-		UsdShadeConversionImpl::HashShadeInput(ShadeInput, HashState);
+		UsdShadeConversionImpl::HashShadeInput(ShadeInput, InOutHash);
 	}
 
 	bool bValue = false;
@@ -3523,13 +3533,7 @@ FSHAHash UsdUtils::HashShadeMaterial(const pxr::UsdShadeMaterial& UsdShadeMateri
 	{
 		Attr.Get<bool>(&bValue);
 	}
-	HashState.Update(reinterpret_cast<uint8*>(&bValue), sizeof(bValue));
-
-	FSHAHash OutHash;
-	HashState.Final();
-	HashState.GetHash(&OutHash.Hash[0]);
-
-	return OutHash;
+	InOutHash.Update(reinterpret_cast<uint8*>(&bValue), sizeof(bValue));
 }
 
 #undef LOCTEXT_NAMESPACE
diff --git a/Engine/Plugins/Importers/USDImporter/Source/USDUtilities/Public/USDShadeConversion.h b/Engine/Plugins/Importers/USDImporter/Source/USDUtilities/Public/USDShadeConversion.h
--- a/Engine/Plugins/Importers/USDImporter/Source/USDUtilities/Public/USDShadeConversion.h
+++ b/Engine/Plugins/Importers/USDImporter/Source/USDUtilities/Public/USDShadeConversion.h
@@ -219,8 +219,15 @@ namespace UsdUtils
 	 */
 	USDUTILITIES_API bool IsMaterialTranslucent(const pxr::UsdShadeMaterial& UsdShadeMaterial);
 
-	USDUTILITIES_API FSHAHash
-	HashShadeMaterial(const pxr::UsdShadeMaterial& UsdShadeMaterial, const pxr::TfToken& RenderContext = pxr::UsdShadeTokens->universalRenderContext);
+	USDUTILITIES_API FSHAHash HashShadeMaterial(	//
+		const pxr::UsdShadeMaterial& UsdShadeMaterial,
+		const pxr::TfToken& RenderContext = pxr::UsdShadeTokens->universalRenderContext
+	);
+	USDUTILITIES_API void HashShadeMaterial(	//
+		const pxr::UsdShadeMaterial& UsdShadeMaterial,
+		FSHA1& InOutHash,
+		const pxr::TfToken& RenderContext = pxr::UsdShadeTokens->universalRenderContext
+	);
 
 	/** Returns the resolved path from an pxr::SdfAssetPath attribute. For UDIMs path, returns the path to the 1001 tile. */
 	USDUTILITIES_API FString GetResolvedAssetPath(const pxr::UsdAttribute& AssetPathAttr, pxr::UsdTimeCode TimeCode = pxr::UsdTimeCode::Default());
