Unreachable code fix.

diff --git a/Engine/Plugins/Importers/USDImporter/Source/USDTests/Private/USDTestsBlueprintLibrary.cpp b/Engine/Plugins/Importers/USDImporter/Source/USDTests/Private/USDTestsBlueprintLibrary.cpp
--- a/Engine/Plugins/Importers/USDImporter/Source/USDTests/Private/USDTestsBlueprintLibrary.cpp
+++ b/Engine/Plugins/Importers/USDImporter/Source/USDTests/Private/USDTestsBlueprintLibrary.cpp
@@ -53,7 +53,6 @@ bool USDTestsBlueprintLibrary::RecompileBlueprintStageActor(AUsdStageActor* Blue
 	else
 	{
 		UE_LOG(LogUsd, Error, TEXT("Blueprint is in an unexpected state after compiling (%s)"), *BP->GetName());
-		return false;
 	}
 #else
 	return false;
