Alternate fix for unreachable code (CIS).

diff --git a/Engine/Plugins/Importers/USDImporter/Source/USDTests/Private/USDTestsBlueprintLibrary.cpp b/Engine/Plugins/Importers/USDImporter/Source/USDTests/Private/USDTestsBlueprintLibrary.cpp
--- a/Engine/Plugins/Importers/USDImporter/Source/USDTests/Private/USDTestsBlueprintLibrary.cpp
+++ b/Engine/Plugins/Importers/USDImporter/Source/USDTests/Private/USDTestsBlueprintLibrary.cpp
@@ -50,13 +50,10 @@ bool USDTestsBlueprintLibrary::RecompileBlueprintStageActor(AUsdStageActor* Blue
 		UE_LOG(LogUsd, Error, TEXT("Blueprint failed to compile (%s)"), *BP->GetName());
 		return false;
 	}
-	else
-	{
-		UE_LOG(LogUsd, Error, TEXT("Blueprint is in an unexpected state after compiling (%s)"), *BP->GetName());
-	}
-#else
-	return false;
+	
+	UE_LOG(LogUsd, Error, TEXT("Blueprint is in an unexpected state after compiling (%s)"), *BP->GetName());
 #endif	  // WITH_EDITOR
+	return false;
 }
 
 void USDTestsBlueprintLibrary::DirtyStageActorBlueprint(AUsdStageActor* BlueprintDerivedStageActor)
