USD: Fixed crash when importing a static mesh with LiveLinkAPI applied to it.

diff --git a/Engine/Plugins/Importers/USDImporter/Source/USDSchemas/Private/USDGeomXformableTranslator.cpp b/Engine/Plugins/Importers/USDImporter/Source/USDSchemas/Private/USDGeomXformableTranslator.cpp
--- a/Engine/Plugins/Importers/USDImporter/Source/USDSchemas/Private/USDGeomXformableTranslator.cpp
+++ b/Engine/Plugins/Importers/USDImporter/Source/USDSchemas/Private/USDGeomXformableTranslator.cpp
@@ -152,11 +152,16 @@ namespace UE::UsdXformableTranslatorImpl::Private
 					SubjectRepresentation.Subject = FName{*UsdToUnreal::ConvertString(SubjectName)};
 				}
 			}
-			Controller->SetSubjectRepresentation(SubjectRepresentation);
 
-			// This should be done after setting the subject representation to ensure that the LiveLink component's ControllerMap has a transform
-			// controller
-			Controller->SetControlledComponent(ULiveLinkTransformRole::StaticClass(), Component);
+			{
+				FScopedUnrealAllocs UEAllocs;
+
+				Controller->SetSubjectRepresentation(SubjectRepresentation);
+
+				// This should be done after setting the subject representation to ensure that the LiveLink component's ControllerMap has a transform
+				// controller
+				Controller->SetControlledComponent(ULiveLinkTransformRole::StaticClass(), Component);
+			}
 
 			if (pxr::UsdAttribute Attr = Prim.GetAttribute(UnrealIdentifiers::UnrealLiveLinkEnabled))
 			{
