USD: Fix no-editor build breakage caused by cl 31285978.

diff --git a/Engine/Plugins/Importers/USDImporter/Source/USDSchemas/Private/USDGeomXformableTranslator.cpp b/Engine/Plugins/Importers/USDImporter/Source/USDSchemas/Private/USDGeomXformableTranslator.cpp
--- a/Engine/Plugins/Importers/USDImporter/Source/USDSchemas/Private/USDGeomXformableTranslator.cpp
+++ b/Engine/Plugins/Importers/USDImporter/Source/USDSchemas/Private/USDGeomXformableTranslator.cpp
@@ -868,13 +868,21 @@ namespace UE::UsdXformableTranslatorImpl::Private
 								continue;
 							}
 
+#if WITH_EDITOR
+							// Check that the texture is the desired one for the attribute by checking
+							// its source path. Unfortunately we can only do this check in the editor for now,
+							// but this should only be relevant during workflows with card editing, which shouldn't
+							// happen at runtime anyway (see UE-200918)
 							if (UAssetImportData* ImportData = Texture->AssetImportData)
 							{
-								if (FPaths::IsSamePath(TexturePath, ImportData->GetFirstFilename()))
+								if (!FPaths::IsSamePath(TexturePath, ImportData->GetFirstFilename()))
 								{
-									(DrawModeComponent->*TextureSetter)(Texture);
+									continue;
 								}
 							}
+#endif	  // WITH_EDITOR
+
+							(DrawModeComponent->*TextureSetter)(Texture);
 						}
 					}
 				}
