USD: Allow setting material overrides on HISM components from PointInstancer prims.

diff --git a/Engine/Plugins/Importers/USDImporter/Source/USDStage/Private/USDStageActor.cpp b/Engine/Plugins/Importers/USDImporter/Source/USDStage/Private/USDStageActor.cpp
--- a/Engine/Plugins/Importers/USDImporter/Source/USDStage/Private/USDStageActor.cpp
+++ b/Engine/Plugins/Importers/USDImporter/Source/USDStage/Private/USDStageActor.cpp
@@ -32,6 +32,7 @@
 #include "UsdWrappers/SdfLayer.h"
 #include "UsdWrappers/UsdGeomBBoxCache.h"
 #include "UsdWrappers/UsdGeomXformable.h"
+#include "UsdWrappers/UsdRelationship.h"
 #include "UsdWrappers/UsdStage.h"
 
 #include "Async/Async.h"
@@ -39,6 +40,7 @@
 #include "CineCameraActor.h"
 #include "CineCameraComponent.h"
 #include "Components/DirectionalLightComponent.h"
+#include "Components/HierarchicalInstancedStaticMeshComponent.h"
 #include "Components/LightComponent.h"
 #include "Components/LightComponentBase.h"
 #include "Components/PointLightComponent.h"
@@ -86,9 +88,14 @@
 #include "Subsystems/AssetEditorSubsystem.h"
 #include "UnrealEdGlobals.h"
 #include "USDClassesEditorModule.h"
-
 #endif	  // WITH_EDITOR
 
+#if USE_USD_SDK
+#include "USDIncludesStart.h"
+#include "pxr/usd/usdGeom/tokens.h"
+#include "USDIncludesEnd.h"
+#endif	  // USE_USD_SDK
+
 #define LOCTEXT_NAMESPACE "USDStageActor"
 
 static bool GRegenerateSkeletalAssetsOnControlRigBake = true;
@@ -949,6 +956,29 @@ struct FUsdStageActorImpl
 		}
 #endif	  // USE_USD_SDK
 	}
+
+	static TSet<FString> GetPointInstancerPrototypes(const UE::FUsdPrim& Prim)
+	{
+		TSet<FString> PrototypePaths;
+
+#if USE_USD_SDK
+		static FString PrototypesStr = UsdToUnreal::ConvertToken(pxr::UsdGeomTokens->prototypes);
+		if (UE::FUsdRelationship Relationship = Prim.GetRelationship(*PrototypesStr))
+		{
+			TArray<UE::FSdfPath> Targets;
+			if (Relationship.GetTargets(Targets))
+			{
+				PrototypePaths.Reserve(Targets.Num());
+				for (const UE::FSdfPath& Path : Targets)
+				{
+					PrototypePaths.Add(Path.GetString());
+				}
+			}
+		}
+#endif	  // USE_USD_SDK
+
+		return PrototypePaths;
+	}
 };
 
 /**
@@ -2088,7 +2118,9 @@ UUsdPrimTwin* AUsdStageActor::ExpandPrim(const UE::FUsdPrim& Prim, bool bResync,
 		}
 		else
 		{
-			ObjectsToWatch.Remove(UsdPrimTwin->SceneComponent.Get());
+			USceneComponent* TwinSceneComponent = UsdPrimTwin->SceneComponent.Get();
+
+			ObjectsToWatch.Remove(TwinSceneComponent);
 			if (Prim.IsA(TEXT("Camera")))
 			{
 				if (ACineCameraActor* CameraActor = Cast<ACineCameraActor>(SceneComponent->GetOwner()))
@@ -2096,7 +2128,17 @@ UUsdPrimTwin* AUsdStageActor::ExpandPrim(const UE::FUsdPrim& Prim, bool bResync,
 					ObjectsToWatch.Remove(CameraActor->GetCineCameraComponent());
 				}
 			}
-			SchemaTranslator->UpdateComponents(UsdPrimTwin->SceneComponent.Get());
+			else if (Prim.IsA(TEXT("PointInstancer")))
+			{
+				TSet<FString> PrototypePaths = FUsdStageActorImpl::GetPointInstancerPrototypes(Prim);
+
+				for (const TObjectPtr<USceneComponent>& Child : TwinSceneComponent->GetAttachChildren())
+				{
+					UHierarchicalInstancedStaticMeshComponent* HISMComponent = Cast<UHierarchicalInstancedStaticMeshComponent>(Child.Get());
+					ObjectsToWatch.Remove(HISMComponent);
+				}
+			}
+			SchemaTranslator->UpdateComponents(TwinSceneComponent);
 		}
 
 		bExpandChildren = !SchemaTranslator->CollapsesChildren(ECollapsingType::Components);
@@ -2135,6 +2177,7 @@ UUsdPrimTwin* AUsdStageActor::ExpandPrim(const UE::FUsdPrim& Prim, bool bResync,
 		}
 
 		ObjectsToWatch.Add(TwinSceneComponent, UsdPrimTwin->PrimPath);
+
 		// Make sure we monitor direct changes to camera properties on the component as well as the actor
 		if (Prim.IsA(TEXT("Camera")))
 		{
@@ -2143,6 +2186,41 @@ UUsdPrimTwin* AUsdStageActor::ExpandPrim(const UE::FUsdPrim& Prim, bool bResync,
 				ObjectsToWatch.Add(CameraActor->GetCineCameraComponent(), UsdPrimTwin->PrimPath);
 			}
 		}
+		else if (Prim.IsA(TEXT("PointInstancer")))
+		{
+			// Collect all the known prototype paths for this PointInstancer
+			TSet<FString> PrototypePaths = FUsdStageActorImpl::GetPointInstancerPrototypes(Prim);
+
+			const TArray<TObjectPtr<USceneComponent>>& ChildComponents = TwinSceneComponent->GetAttachChildren();
+			for (const TObjectPtr<USceneComponent>& Child : ChildComponents)
+			{
+				UHierarchicalInstancedStaticMeshComponent* HISMComponent = Cast<UHierarchicalInstancedStaticMeshComponent>(Child.Get());
+				if (!HISMComponent)
+				{
+					continue;
+				}
+
+				UStaticMesh* HISMMesh = HISMComponent->GetStaticMesh();
+				if (!HISMMesh)
+				{
+					continue;
+				}
+
+				UUsdAssetUserData* UserData = UsdUtils::GetAssetUserData(HISMMesh);
+				if (!UserData)
+				{
+					continue;
+				}
+
+				for (const FString& Path : UserData->PrimPaths)
+				{
+					if (PrototypePaths.Contains(Path))
+					{
+						ObjectsToWatch.Add(HISMComponent, Path);
+					}
+				}
+			}
+		}
 	}
 
 	// Check if the prim should have Sequencer tracks or not
@@ -4508,129 +4586,131 @@ void AUsdStageActor::OnObjectPropertyChanged(UObject* ObjectBeingModified, FProp
 
 	FString PrimPath = ObjectsToWatch[ComponentBeingModified];
 
-	if (UUsdPrimTwin* UsdPrimTwin = GetRootPrimTwin()->Find(PrimPath))
+	// Not all of our spawned components will have prim twins (e.g. HISM components for PointInstancers)
+	USceneComponent* PrimSceneComponent = Cast<USceneComponent>(ComponentBeingModified);
+	UUsdPrimTwin* UsdPrimTwin = GetRootPrimTwin()->Find(PrimPath);
+	if (UsdPrimTwin)
 	{
-		// Update prim from UE
-		USceneComponent* PrimSceneComponent = UsdPrimTwin->SceneComponent.Get();
-		if (PrimSceneComponent)
-		{
-			if (CurrentStage)
-			{
-				// This block is important, as it not only prevents us from getting into infinite loops with the USD notices,
-				// but it also guarantees that if we have an object property change, the corresponding stage notice is not also
-				// independently saved to the transaction via the UUsdTransactor, which would be duplication
-				FScopedBlockNoticeListening BlockNotices(this);
+		PrimSceneComponent = UsdPrimTwin->SceneComponent.Get();
+	}
 
-				UE::FUsdPrim UsdPrim = CurrentStage.GetPrimAtPath(UE::FSdfPath(*PrimPath));
+	// Update prim from UE
+	if (PrimSceneComponent && CurrentStage)
+	{
+		// This block is important, as it not only prevents us from getting into infinite loops with the USD notices,
+		// but it also guarantees that if we have an object property change, the corresponding stage notice is not also
+		// independently saved to the transaction via the UUsdTransactor, which would be duplication
+		FScopedBlockNoticeListening BlockNotices(this);
 
-				// We want to keep component visibilities in sync with USD, which uses inherited visibilities
-				// To accomplish that while blocking notices we must always propagate component visibility changes manually.
-				// This part is effectively the same as calling pxr::UsdGeomImageable::MakeVisible/Invisible.
-				if (PropertyChangedEvent.GetPropertyName() == TEXT("bHiddenInGame"))
-				{
-					PrimSceneComponent->Modify();
+		UE::FUsdPrim UsdPrim = CurrentStage.GetPrimAtPath(UE::FSdfPath(*PrimPath));
 
-					if (PrimSceneComponent->bHiddenInGame)
-					{
-						FUsdStageActorImpl::MakeInvisible(*UsdPrimTwin);
-					}
-					else
-					{
-						FUsdStageActorImpl::MakeVisible(*UsdPrimTwin, CurrentStage);
-					}
-				}
+		// We want to keep component visibilities in sync with USD, which uses inherited visibilities
+		// To accomplish that while blocking notices we must always propagate component visibility changes manually.
+		// This part is effectively the same as calling pxr::UsdGeomImageable::MakeVisible/Invisible.
+		// TODO: Allow writing out visibility without needing a prim twin
+		if (UsdPrimTwin && PropertyChangedEvent.GetPropertyName() == TEXT("bHiddenInGame"))
+		{
+			PrimSceneComponent->Modify();
+
+			if (PrimSceneComponent->bHiddenInGame)
+			{
+				FUsdStageActorImpl::MakeInvisible(*UsdPrimTwin);
+			}
+			else
+			{
+				FUsdStageActorImpl::MakeVisible(*UsdPrimTwin, CurrentStage);
+			}
+		}
 
 #if USE_USD_SDK
 
-				UnrealToUsd::ConvertLiveLinkProperties(Controller ? Cast<UActorComponent>(Controller) : PrimSceneComponent, UsdPrim);
+		UnrealToUsd::ConvertLiveLinkProperties(Controller ? Cast<UActorComponent>(Controller) : PrimSceneComponent, UsdPrim);
 
-				UnrealToUsd::ConvertSceneComponent(CurrentStage, PrimSceneComponent, UsdPrim);
+		UnrealToUsd::ConvertSceneComponent(CurrentStage, PrimSceneComponent, UsdPrim);
 
-				// When we parse a Gprim like a Cube or a Cylinder, we'll always generate some "default" meshes (e.g. Cylinder with
-				// height always equal 1), and combine the Xform and the effect of the prim's attributes (e.g. height/width) into
-				// a SINGLE transform, and put that on the component (this approach allows attribute animation purely with Sequencer tracks).
-				// When we modify any property and want to write back out to USD however, we'll write that combined transform as the prim's
-				// transform. This means we must also "reset" the (e.g. height/width) attributes, so that the combined transform stays
-				// consistent
-				const bool bDefaultValues = true;
-				const bool bTimeSampleValues = false;
-				UsdUtils::AuthorIdentityTransformGprimAttributes(UsdPrim, bDefaultValues, bTimeSampleValues);
+		// When we parse a Gprim like a Cube or a Cylinder, we'll always generate some "default" meshes (e.g. Cylinder with
+		// height always equal 1), and combine the Xform and the effect of the prim's attributes (e.g. height/width) into
+		// a SINGLE transform, and put that on the component (this approach allows attribute animation purely with Sequencer tracks).
+		// When we modify any property and want to write back out to USD however, we'll write that combined transform as the prim's
+		// transform. This means we must also "reset" the (e.g. height/width) attributes, so that the combined transform stays
+		// consistent
+		const bool bDefaultValues = true;
+		const bool bTimeSampleValues = false;
+		UsdUtils::AuthorIdentityTransformGprimAttributes(UsdPrim, bDefaultValues, bTimeSampleValues);
 
-				if (UMeshComponent* MeshComponent = Cast<UMeshComponent>(PrimSceneComponent))
+		if (UMeshComponent* MeshComponent = Cast<UMeshComponent>(PrimSceneComponent))
+		{
+			UnrealToUsd::ConvertMeshComponent(CurrentStage, MeshComponent, UsdPrim);
+		}
+		else if (UUsdDrawModeComponent* DrawModeComponent = Cast<UUsdDrawModeComponent>(PrimSceneComponent))
+		{
+			const static TSet<FName> BoundsProperties = {
+				GET_MEMBER_NAME_CHECKED(UUsdDrawModeComponent, BoundsMin),
+				GET_MEMBER_NAME_CHECKED(UUsdDrawModeComponent, BoundsMax),
+			};
+
+			// If we just manually tweaked the extents, also author those back out to USD as extents opinions
+			const bool bWriteExtents = BoundsProperties.Contains(PropertyChangedEvent.GetMemberPropertyName());
+			const double UsdTimeCode = UsdUtils::GetDefaultTimeCode();
+			UnrealToUsd::ConvertDrawModeComponent(*DrawModeComponent, UsdPrim, bWriteExtents, UsdTimeCode);
+		}
+		else if (UsdPrim && UsdPrim.IsA(TEXT("Camera")))
+		{
+			// Our component may be pointing directly at a camera component in case we recreated an exported
+			// ACineCameraActor (see UE-120826)
+			if (UCineCameraComponent* RecreatedCameraComponent = Cast<UCineCameraComponent>(PrimSceneComponent))
+			{
+				UnrealToUsd::ConvertCameraComponent(*RecreatedCameraComponent, UsdPrim);
+			}
+			// Or it could have been just a generic Camera prim, at which case we'll have spawned an entire new
+			// ACineCameraActor for it. In this scenario our prim twin is pointing at the root component, so we need
+			// to dig to the actual UCineCameraComponent to write out the camera data.
+			// We should only do this when the Prim actually corresponds to the Camera though, or else we'll also catch
+			// the prim/component pair that corresponds to the root scene component in case we recreated an exported
+			// ACineCameraActor.
+			else if (ACineCameraActor* CameraActor = Cast<ACineCameraActor>(PrimSceneComponent->GetOwner()))
+			{
+				if (UCineCameraComponent* CameraComponent = CameraActor->GetCineCameraComponent())
 				{
-					UnrealToUsd::ConvertMeshComponent(CurrentStage, MeshComponent, UsdPrim);
+					UnrealToUsd::ConvertCameraComponent(*CameraComponent, UsdPrim);
 				}
-				else if (UUsdDrawModeComponent* DrawModeComponent = Cast<UUsdDrawModeComponent>(PrimSceneComponent))
+			}
+		}
+		else if (ALight* LightActor = Cast<ALight>(PrimSceneComponent->GetOwner()))
+		{
+			if (ULightComponent* LightComponent = LightActor->GetLightComponent())
+			{
+				UnrealToUsd::ConvertLightComponent(*LightComponent, UsdPrim, UsdUtils::GetDefaultTimeCode());
+
+				if (UDirectionalLightComponent* DirectionalLight = Cast<UDirectionalLightComponent>(LightComponent))
 				{
-					const static TSet<FName> BoundsProperties = {
-						GET_MEMBER_NAME_CHECKED(UUsdDrawModeComponent, BoundsMin),
-						GET_MEMBER_NAME_CHECKED(UUsdDrawModeComponent, BoundsMax),
-					};
-
-					// If we just manually tweaked the extents, also author those back out to USD as extents opinions
-					const bool bWriteExtents = BoundsProperties.Contains(PropertyChangedEvent.GetMemberPropertyName());
-					const double UsdTimeCode = UsdUtils::GetDefaultTimeCode();
-					UnrealToUsd::ConvertDrawModeComponent(*DrawModeComponent, UsdPrim, bWriteExtents, UsdTimeCode);
+					UnrealToUsd::ConvertDirectionalLightComponent(*DirectionalLight, UsdPrim, UsdUtils::GetDefaultTimeCode());
 				}
-				else if (UsdPrim && UsdPrim.IsA(TEXT("Camera")))
+				else if (URectLightComponent* RectLight = Cast<URectLightComponent>(LightComponent))
 				{
-					// Our component may be pointing directly at a camera component in case we recreated an exported
-					// ACineCameraActor (see UE-120826)
-					if (UCineCameraComponent* RecreatedCameraComponent = Cast<UCineCameraComponent>(PrimSceneComponent))
-					{
-						UnrealToUsd::ConvertCameraComponent(*RecreatedCameraComponent, UsdPrim);
-					}
-					// Or it could have been just a generic Camera prim, at which case we'll have spawned an entire new
-					// ACineCameraActor for it. In this scenario our prim twin is pointing at the root component, so we need
-					// to dig to the actual UCineCameraComponent to write out the camera data.
-					// We should only do this when the Prim actually corresponds to the Camera though, or else we'll also catch
-					// the prim/component pair that corresponds to the root scene component in case we recreated an exported
-					// ACineCameraActor.
-					else if (ACineCameraActor* CameraActor = Cast<ACineCameraActor>(PrimSceneComponent->GetOwner()))
-					{
-						if (UCineCameraComponent* CameraComponent = CameraActor->GetCineCameraComponent())
-						{
-							UnrealToUsd::ConvertCameraComponent(*CameraComponent, UsdPrim);
-						}
-					}
+					UnrealToUsd::ConvertRectLightComponent(*RectLight, UsdPrim, UsdUtils::GetDefaultTimeCode());
 				}
-				else if (ALight* LightActor = Cast<ALight>(PrimSceneComponent->GetOwner()))
+				else if (UPointLightComponent* PointLight = Cast<UPointLightComponent>(LightComponent))
 				{
-					if (ULightComponent* LightComponent = LightActor->GetLightComponent())
-					{
-						UnrealToUsd::ConvertLightComponent(*LightComponent, UsdPrim, UsdUtils::GetDefaultTimeCode());
-
-						if (UDirectionalLightComponent* DirectionalLight = Cast<UDirectionalLightComponent>(LightComponent))
-						{
-							UnrealToUsd::ConvertDirectionalLightComponent(*DirectionalLight, UsdPrim, UsdUtils::GetDefaultTimeCode());
-						}
-						else if (URectLightComponent* RectLight = Cast<URectLightComponent>(LightComponent))
-						{
-							UnrealToUsd::ConvertRectLightComponent(*RectLight, UsdPrim, UsdUtils::GetDefaultTimeCode());
-						}
-						else if (UPointLightComponent* PointLight = Cast<UPointLightComponent>(LightComponent))
-						{
-							UnrealToUsd::ConvertPointLightComponent(*PointLight, UsdPrim, UsdUtils::GetDefaultTimeCode());
+					UnrealToUsd::ConvertPointLightComponent(*PointLight, UsdPrim, UsdUtils::GetDefaultTimeCode());
 
-							if (USpotLightComponent* SpotLight = Cast<USpotLightComponent>(LightComponent))
-							{
-								UnrealToUsd::ConvertSpotLightComponent(*SpotLight, UsdPrim, UsdUtils::GetDefaultTimeCode());
-							}
-						}
+					if (USpotLightComponent* SpotLight = Cast<USpotLightComponent>(LightComponent))
+					{
+						UnrealToUsd::ConvertSpotLightComponent(*SpotLight, UsdPrim, UsdUtils::GetDefaultTimeCode());
 					}
 				}
-				// In contrast to the other light types, the USkyLightComponent is the root component of the ASkyLight
-				else if (USkyLightComponent* SkyLightComponent = Cast<USkyLightComponent>(PrimSceneComponent))
-				{
-					UnrealToUsd::ConvertLightComponent(*SkyLightComponent, UsdPrim, UsdUtils::GetDefaultTimeCode());
-					UnrealToUsd::ConvertSkyLightComponent(*SkyLightComponent, UsdPrim, UsdUtils::GetDefaultTimeCode());
-				}
-#endif	  // #if USE_USD_SDK
-
-				// Update stage window in case any of our component changes trigger USD stage changes
-				this->OnPrimChanged.Broadcast(PrimPath, false);
 			}
 		}
+		// In contrast to the other light types, the USkyLightComponent is the root component of the ASkyLight
+		else if (USkyLightComponent* SkyLightComponent = Cast<USkyLightComponent>(PrimSceneComponent))
+		{
+			UnrealToUsd::ConvertLightComponent(*SkyLightComponent, UsdPrim, UsdUtils::GetDefaultTimeCode());
+			UnrealToUsd::ConvertSkyLightComponent(*SkyLightComponent, UsdPrim, UsdUtils::GetDefaultTimeCode());
+		}
+#endif	  // #if USE_USD_SDK
+
+		// Update stage window in case any of our component changes trigger USD stage changes
+		this->OnPrimChanged.Broadcast(PrimPath, false);
 	}
 }
 
