USD: Handle joint mirroring with uniform scales to prevent joins from snapping and rotating strangely.

diff --git a/Engine/Plugins/Importers/USDImporter/Source/USDUtilities/Private/USDSkeletalDataConversion.cpp b/Engine/Plugins/Importers/USDImporter/Source/USDUtilities/Private/USDSkeletalDataConversion.cpp
--- a/Engine/Plugins/Importers/USDImporter/Source/USDUtilities/Private/USDSkeletalDataConversion.cpp
+++ b/Engine/Plugins/Importers/USDImporter/Source/USDUtilities/Private/USDSkeletalDataConversion.cpp
@@ -1288,7 +1288,25 @@ bool UsdToUnreal::ConvertSkeleton(const pxr::UsdSkelSkeletonQuery& SkeletonQuery
 		for (uint32 Index = 0; Index < JointLocalRestTransforms.size(); ++Index)
 		{
 			const GfMatrix4d& UsdMatrix = JointLocalRestTransforms[Index];
-			FTransform BoneTransform = UsdToUnreal::ConvertMatrix(StageInfo, UsdMatrix);
+
+			// Here we use DecomposeWithUniformReflection instead of the previous UsdToUnreal::ConvertMatrix(StageInfo, UsdMatrix)
+			// call, because internally that would have done the matrix decomposition via FTransform::SetFromMatrix.
+			//
+			// The only difference between the two being that if we detect any negative scaling, DecomposeWithUniformReflection will
+			// flip *all* axes instead of only one, which will keep the scaling uniform. Otherwise, we may get weird joint flipping
+			// effects and the joint rotation axes being inverted (see UE-193643). Those are likely consequences of decomposed
+			// transforms not being easily invertible (some code at some point will silently assume uniform scaling, and things would
+			// break).
+			//
+			// Note that FBX secretly does this as well, because the FBX SDK's Matrix.GetT(), Matrix.GetQ() and Matrix.GetS()
+			// (used within UnFbx::FFbxImporter::ImportBones) seem to behave the same way and flip all axes when a reflection is
+			// detected.
+			FTransform BoneTransform;
+			{
+				FMatrix Matrix = UsdToUnreal::ConvertMatrix(UsdMatrix);
+				BoneTransform = UsdUtils::DecomposeWithUniformReflection(Matrix);
+				BoneTransform = UsdUtils::ConvertTransformToUESpace(StageInfo, BoneTransform);
+			}
 			BoneTransforms.Add(BoneTransform);
 		}
 	}
diff --git a/Engine/Plugins/Importers/USDImporter/Source/USDUtilities/Private/USDTypesConversion.cpp b/Engine/Plugins/Importers/USDImporter/Source/USDUtilities/Private/USDTypesConversion.cpp
--- a/Engine/Plugins/Importers/USDImporter/Source/USDUtilities/Private/USDTypesConversion.cpp
+++ b/Engine/Plugins/Importers/USDImporter/Source/USDUtilities/Private/USDTypesConversion.cpp
@@ -65,6 +65,35 @@ namespace UsdUtils
 		return FTransform(Rotation, Translation, Scale);
 	}
 
+	FTransform DecomposeWithUniformReflection(const FMatrix& InMatrix)
+	{
+		// Reference: FTransform::SetFromMatrix
+
+		FTransform Result;
+
+		FMatrix M = InMatrix;
+
+		FVector Scale = M.ExtractScaling();
+
+		// If there is negative scaling going on, we handle that here
+		if (InMatrix.Determinant() < 0.f)
+		{
+			// Note: Here is where we flip on all three axes simultaneously. Flipping on any odd number of
+			// axes would have led to the same result (as long as the rotation matrix M is also correspondingly
+			// flipped), but by flipping 3 of them we get to keep the scaling uniform.
+			Scale *= -1.0f;
+			M.SetAxis(0, -M.GetScaledAxis(EAxis::X));
+			M.SetAxis(1, -M.GetScaledAxis(EAxis::Y));
+			M.SetAxis(2, -M.GetScaledAxis(EAxis::Z));
+		}
+
+		Result.SetScale3D(Scale);
+		Result.SetRotation(FQuat{M}.GetNormalized());
+		Result.SetTranslation(InMatrix.GetOrigin());
+
+		return Result;
+	}
+
 	FTransform ConvertTransformToUsdSpace(const FUsdStageInfo& StageInfo, const FTransform& TransformInUESpace)
 	{
 		FTransform TransformInUsdSpace = UsdUtils::ConvertAxes(StageInfo.UpAxis == EUsdUpAxis::ZAxis, TransformInUESpace);
@@ -77,6 +106,19 @@ namespace UsdUtils
 
 		return TransformInUsdSpace;
 	}
+
+	FTransform ConvertTransformToUESpace(const FUsdStageInfo& StageInfo, const FTransform& TransformInUsdSpace)
+	{
+		FTransform TransformInUESpace = UsdUtils::ConvertAxes(StageInfo.UpAxis == EUsdUpAxis::ZAxis, TransformInUsdSpace);
+
+		const float UEMetersPerUnit = 0.01f;
+		if (!FMath::IsNearlyEqual(StageInfo.MetersPerUnit, UEMetersPerUnit))
+		{
+			TransformInUESpace.ScaleTranslation(StageInfo.MetersPerUnit / UEMetersPerUnit);
+		}
+
+		return TransformInUESpace;
+	}
 }	 // namespace UsdUtils
 
 namespace UsdToUnreal
@@ -497,8 +539,7 @@ namespace UnrealToUsd
 			static_cast<pxr::GfHalf>(InValue.W),
 			static_cast<pxr::GfHalf>(InValue.X),
 			static_cast<pxr::GfHalf>(InValue.Y),
-			static_cast<pxr::GfHalf>(InValue.Z)
-		};
+			static_cast<pxr::GfHalf>(InValue.Z)};
 	}
 
 	pxr::GfQuatf ConvertQuatFloat(const FQuat& InValue)
@@ -507,8 +548,7 @@ namespace UnrealToUsd
 			static_cast<float>(InValue.W),
 			static_cast<float>(InValue.X),
 			static_cast<float>(InValue.Y),
-			static_cast<float>(InValue.Z)
-		};
+			static_cast<float>(InValue.Z)};
 	}
 
 	pxr::GfQuatd ConvertQuatDouble(const FQuat& InValue)
diff --git a/Engine/Plugins/Importers/USDImporter/Source/USDUtilities/Public/USDTypesConversion.h b/Engine/Plugins/Importers/USDImporter/Source/USDUtilities/Public/USDTypesConversion.h
--- a/Engine/Plugins/Importers/USDImporter/Source/USDUtilities/Public/USDTypesConversion.h
+++ b/Engine/Plugins/Importers/USDImporter/Source/USDUtilities/Public/USDTypesConversion.h
@@ -158,7 +158,15 @@ namespace UnrealToUsd
 
 namespace UsdUtils
 {
+	/**
+	 * Decomposes the FMatrix into an FTransform.
+	 * Identical to FTransform::SetFromMatrix, except that if a reflection is detected all axes are flipped instead
+	 * of only the X axis, which keeps a uniform scaling.
+	 */
+	USDUTILITIES_API FTransform DecomposeWithUniformReflection(const FMatrix& InMatrix);
+
 	USDUTILITIES_API FTransform ConvertTransformToUsdSpace(const FUsdStageInfo& StageInfo, const FTransform& TransformInUESpace);
+	USDUTILITIES_API FTransform ConvertTransformToUESpace(const FUsdStageInfo& StageInfo, const FTransform& TransformInUsdSpace);
 
 	FTransform ConvertAxes(const bool bZUp, const FTransform Transform);
 }
