USD: Suppress automatic dirty layer saving when switching isolated layer.

diff --git a/Engine/Plugins/Importers/USDImporter/Source/USDStage/Private/USDStageActor.cpp b/Engine/Plugins/Importers/USDImporter/Source/USDStage/Private/USDStageActor.cpp
--- a/Engine/Plugins/Importers/USDImporter/Source/USDStage/Private/USDStageActor.cpp
+++ b/Engine/Plugins/Importers/USDImporter/Source/USDStage/Private/USDStageActor.cpp
@@ -1351,6 +1351,17 @@ void AUsdStageActor::IsolateLayer(const UE::FSdfLayer& Layer)
 		return;
 	}
 
+	// The USD Stage Editor listens to OnPreStageChanged and will use UE::UsdStageEditorModule::Private::SaveStageActorLayersForWorld
+	// to show the "Save dirty layers" dialog as a response, if we have any dirty/memory-only layers. We're never really going to
+	// discard unsaved changes by isolating/stopping isolation though, so we don't actually need to save anything in this case...
+	// Let's temporarily tweak the project settings to disable automatic saving of dirty layers while we swap our isolated layer.
+	UUsdProjectSettings* Settings = GetMutableDefault<UUsdProjectSettings>();
+	if (!Settings)
+	{
+		return;
+	}
+	TGuardValue<EUsdSaveDialogBehavior> DisableDialogGuard{Settings->ShowSaveLayersDialogWhenClosing, EUsdSaveDialogBehavior::NeverSave};
+
 	OnPreStageChanged.Broadcast();
 
 	// Stop isolating
