USD: Fixed possible crash due to USD GeometryCache with bad data.

diff --git a/Engine/Plugins/Importers/USDImporter/Source/GeometryCacheUSD/Private/GeometryCacheUSDStream.cpp b/Engine/Plugins/Importers/USDImporter/Source/GeometryCacheUSD/Private/GeometryCacheUSDStream.cpp
--- a/Engine/Plugins/Importers/USDImporter/Source/GeometryCacheUSD/Private/GeometryCacheUSDStream.cpp
+++ b/Engine/Plugins/Importers/USDImporter/Source/GeometryCacheUSD/Private/GeometryCacheUSDStream.cpp
@@ -25,7 +25,7 @@ static int32 kUsdReadConcurrency = 10;
 // differences, etc.) replace the version GUID below with a new one.
 // In case of merge conflicts with DDC versions, you MUST generate a new GUID
 // and set this new GUID as the version.
-#define USDSTREAM_DERIVED_DATA_VERSION TEXT("BB2B7CC023C54AEBBCC5ABDC1B0BFFD2")
+#define USDSTREAM_DERIVED_DATA_VERSION TEXT("7F54F423DF3247B7B77B9E2B356FBE83")
 
 class FUsdStreamDDCUtils
 {
@@ -64,6 +64,7 @@ public:
 			// modify MeshConversionOptions directly as it's owned by the UsdTrack and we may be called from a thread)
 			MD5.Update((uint8*)&MeshConversionOptions.SubdivisionLevel, sizeof(MeshConversionOptions.SubdivisionLevel));
 			MD5.Update((uint8*)&MeshConversionOptions.PurposesToLoad, sizeof(MeshConversionOptions.PurposesToLoad));
+			MD5.Update((uint8*)&MeshConversionOptions.bMergeIdenticalMaterialSlots, sizeof(MeshConversionOptions.bMergeIdenticalMaterialSlots));
 
 			uint8 Digest[16];
 			MD5.Final(Digest);
diff --git a/Engine/Plugins/Importers/USDImporter/Source/USDUtilities/Private/USDGeomMeshConversion.cpp b/Engine/Plugins/Importers/USDImporter/Source/USDUtilities/Private/USDGeomMeshConversion.cpp
--- a/Engine/Plugins/Importers/USDImporter/Source/USDUtilities/Private/USDGeomMeshConversion.cpp
+++ b/Engine/Plugins/Importers/USDImporter/Source/USDUtilities/Private/USDGeomMeshConversion.cpp
@@ -5086,6 +5086,11 @@ void UsdUtils::HashGeomMeshPrim(const UE::FUsdStage& Stage, const FString& PrimP
 
 		HashArrayPrimvar<GfVec2f>(InOutHashState, PrimvarsByUVIndex[UVChannelIndex].Get(), TimeCode);
 	}
+
+	// The number of geomsubsets will give the upper limit of the number of sections in the mesh
+	std::vector<pxr::UsdGeomSubset> GeomSubsets = pxr::UsdShadeMaterialBindingAPI(UsdPrim).GetMaterialBindSubsets();
+	const int32 NumGeomSubsets = static_cast<int32>(GeomSubsets.size());
+	InOutHashState.Update((uint8*)&NumGeomSubsets, sizeof(NumGeomSubsets));
 }
 
 bool UsdUtils::GetPointInstancerTransforms(
diff --git a/Engine/Plugins/Runtime/GeometryCache/Source/GeometryCache/Private/GeometryCacheSceneProxy.cpp b/Engine/Plugins/Runtime/GeometryCache/Source/GeometryCache/Private/GeometryCacheSceneProxy.cpp
--- a/Engine/Plugins/Runtime/GeometryCache/Source/GeometryCache/Private/GeometryCacheSceneProxy.cpp
+++ b/Engine/Plugins/Runtime/GeometryCache/Source/GeometryCache/Private/GeometryCacheSceneProxy.cpp
@@ -475,7 +475,8 @@ void FGeometryCacheSceneProxy::GetDynamicMeshElements(const TArray<const FSceneV
 #endif
 
 						// Apply view mode material overrides
-						const int32 MaterialIndex = TrackProxy->Materials.IsValidIndex(BatchInfo.MaterialIndex) ? BatchInfo.MaterialIndex : BatchIndex; 
+						const int32 MaterialIndex = TrackProxy->Materials.IsValidIndex(BatchInfo.MaterialIndex) ? BatchInfo.MaterialIndex :
+													TrackProxy->Materials.IsValidIndex(BatchIndex) ? BatchIndex : 0; // extra precaution in case of bad data
 						FMaterialRenderProxy* MaterialProxy = bWireframe ? WireframeMaterialInstance : TrackProxy->Materials[MaterialIndex]->GetRenderProxy();
 						MeshBatch.bWireframe = bWireframe;
 						MeshBatch.MaterialRenderProxy = MaterialProxy;
