USD: Fixed missing Start/End frame values for imported GeometryCache, which prevented them from exporting correctly.

diff --git a/Engine/Plugins/Importers/USDImporter/Source/USDSchemas/Private/USDGeometryCacheTranslator.cpp b/Engine/Plugins/Importers/USDImporter/Source/USDSchemas/Private/USDGeometryCacheTranslator.cpp
--- a/Engine/Plugins/Importers/USDImporter/Source/USDSchemas/Private/USDGeometryCacheTranslator.cpp
+++ b/Engine/Plugins/Importers/USDImporter/Source/USDSchemas/Private/USDGeometryCacheTranslator.cpp
@@ -448,13 +448,13 @@ namespace UsdGeometryCacheTranslatorImpl
 				Context->ObjectFlags | EObjectFlags::RF_Public | RF_Transient
 			);
 
-			TOptional<FReadMeshDataArgs> Args;
+			FReadMeshDataArgs Args(GetReadMeshDataArgs(Context, RootPrimPath));
 			if (!Context->bIsImporting)
 			{
-				Args = GetReadMeshDataArgs(Context, RootPrimPath);
 				// StartOffsetTime is the offset applied to the GeometryCache section on the sequencer track, so not relevant when importing
-				StartOffsetTime = static_cast<float>(Args->StartFrame) / Args->FramesPerSecond;
+				StartOffsetTime = static_cast<float>(Args.StartFrame) / Args.FramesPerSecond;
 			}
+			GeometryCache->SetFrameStartEnd(Args.StartFrame, Args.EndFrame);
 
 			// Create a track for each mesh to be processed and add it to the GeometryCache
 			for (int32 Index = 0; Index < MeshPaths.Num(); ++Index)
@@ -463,7 +463,7 @@ namespace UsdGeometryCacheTranslatorImpl
 				UGeometryCacheTrack* Track = nullptr;
 				if (!Context->bIsImporting)
 				{
-					Track = CreateUsdStreamTrack(GeometryCache, *Args, PrimPath, MaterialOffsets[Index]);
+					Track = CreateUsdStreamTrack(GeometryCache, Args, PrimPath, MaterialOffsets[Index]);
 				}
 				else
 				{
diff --git a/Engine/Plugins/Importers/USDImporter/Source/USDUtilities/Private/USDGeomMeshConversion.cpp b/Engine/Plugins/Importers/USDImporter/Source/USDUtilities/Private/USDGeomMeshConversion.cpp
--- a/Engine/Plugins/Importers/USDImporter/Source/USDUtilities/Private/USDGeomMeshConversion.cpp
+++ b/Engine/Plugins/Importers/USDImporter/Source/USDUtilities/Private/USDGeomMeshConversion.cpp
@@ -3944,12 +3944,15 @@ namespace UE::UsdGeometryCacheConversion::Private
 	{
 		FGeometryCacheExportContext(const UGeometryCache& GeometryCache)
 			: SlotNames(GeometryCache.MaterialSlotNames)
-			, FrameRate((GeometryCache.GetEndFrame() - GeometryCache.GetStartFrame()) / GeometryCache.CalculateDuration())
 		{
+			// The GeometryCache's EndFrame is exclusive since it's there to allow frame interpolation past the real last frame
+			InclusiveEndFrame = FMath::Max(GeometryCache.GetEndFrame() - 1, GeometryCache.GetStartFrame() + 1);
+			FrameRate = (InclusiveEndFrame - GeometryCache.GetStartFrame()) / GeometryCache.CalculateDuration();
 		}
 
 		const TArray<FName>& SlotNames;
-		const float FrameRate;
+		int32 InclusiveEndFrame;
+		float FrameRate;
 
 		// Cached values of the last written attribute values
 		// Since int cannot be interpolated, the missing timesampled attribute values will be the "held"
@@ -4293,11 +4296,11 @@ bool UnrealToUsd::ConvertGeometryCache(const UGeometryCache* GeometryCache, pxr:
 	pxr::UsdGeomMesh TargetMesh{UsdPrim};
 	pxr::UsdPrim MaterialPrim = MaterialStage->OverridePrim(UsdPrim.GetPath());
 
+	UsdGeometryCacheImpl::FGeometryCacheExportContext ExportContext(*GeometryCache);
 	const int32 StartFrame = GeometryCache->GetStartFrame();
-	const int32 EndFrame = GeometryCache->GetEndFrame();
+	const int32 EndFrame = ExportContext.InclusiveEndFrame;
 	int32 ActualStartFrame = -1;
 
-	UsdGeometryCacheImpl::FGeometryCacheExportContext ExportContext(*GeometryCache);
 	for (int32 FrameIndex = StartFrame; FrameIndex <= EndFrame; ++FrameIndex)
 	{
 		FGeometryCacheMeshData MeshData = UsdGeometryCacheImpl::GetFlattenedGeometryCacheMeshData(GeometryCache, FrameIndex - StartFrame);
