USD: Fixed detection of physics collision attribute change to properly rebuild collision shape.

#rnx

diff --git a/Engine/Plugins/Importers/USDImporter/Source/USDStage/Private/USDStageActor.cpp b/Engine/Plugins/Importers/USDImporter/Source/USDStage/Private/USDStageActor.cpp
--- a/Engine/Plugins/Importers/USDImporter/Source/USDStage/Private/USDStageActor.cpp
+++ b/Engine/Plugins/Importers/USDImporter/Source/USDStage/Private/USDStageActor.cpp
@@ -93,6 +93,7 @@
 #if USE_USD_SDK
 #include "USDIncludesStart.h"
 #include "pxr/usd/usdGeom/tokens.h"
+#include "pxr/usd/usdPhysics/tokens.h"
 #include "USDIncludesEnd.h"
 #endif	  // USE_USD_SDK
 
@@ -1591,7 +1592,10 @@ void AUsdStageActor::OnUsdObjectsChanged(const UsdUtils::FObjectChangesByPath& I
 					// When we change these UsdGeomModelAPI attributes we may need to create a new component type
 					// for the prim (as it may now need/stop needing an alternate draw mode component)
 					UnrealIdentifiers::ModelDrawMode,
-					UnrealIdentifiers::ModelApplyDrawMode};
+					UnrealIdentifiers::ModelApplyDrawMode,
+					// Physics collision attribute change needs to rebuild the collision shape
+					*UsdToUnreal::ConvertToken(pxr::UsdPhysicsTokens->physicsCollisionEnabled),
+					*UsdToUnreal::ConvertToken(pxr::UsdPhysicsTokens->physicsApproximation)};
 
 				// Some stage info should trigger some resyncs because they should trigger reparsing of geometry
 				if ((PrimPath.IsAbsoluteRootPath() && StageResyncProperties.Contains(AttributeChange.PropertyName))
