USD: Fixed unnecessary time-index conversions in USD GeometryCache, which can cause frame number errors due to float imprecision.

diff --git a/Engine/Plugins/Importers/USDImporter/Source/GeometryCacheUSD/Private/GeometryCacheTrackUSD.cpp b/Engine/Plugins/Importers/USDImporter/Source/GeometryCacheUSD/Private/GeometryCacheTrackUSD.cpp
--- a/Engine/Plugins/Importers/USDImporter/Source/GeometryCacheUSD/Private/GeometryCacheTrackUSD.cpp
+++ b/Engine/Plugins/Importers/USDImporter/Source/GeometryCacheUSD/Private/GeometryCacheTrackUSD.cpp
@@ -82,7 +82,7 @@ const bool UGeometryCacheTrackUsd::UpdateBoundsData(
 {
 	const int32 SampleIndex = FindSampleIndexFromTime(Time, bLooping);
 
-	const FGeometryCacheTrackSampleInfo& SampledInfo = GetSampleInfo(Time, bLooping);
+	const FGeometryCacheTrackSampleInfo& SampledInfo = GetSampleInfo(SampleIndex);
 	if (InOutBoundsSampleIndex != SampleIndex)
 	{
 		OutBounds = SampledInfo.BoundingBox;
@@ -131,6 +131,17 @@ void UGeometryCacheTrackUsd::GetFractionalFrameIndexFromTime(const float Time, c
 
 const FGeometryCacheTrackSampleInfo& UGeometryCacheTrackUsd::GetSampleInfo(float Time, bool bLooping)
 {
+	const int32 SampleIndex = FindSampleIndexFromTime(Time, bLooping);
+	return GetSampleInfo(SampleIndex);
+}
+
+const FGeometryCacheTrackSampleInfo& UGeometryCacheTrackUsd::GetSampleInfo(int32 SampleIndex)
+{
+	if (SampleIndex < 0)
+	{
+		return FGeometryCacheTrackSampleInfo::EmptySampleInfo;
+	}
+
 	if (SampleInfos.Num() == 0)
 	{
 		if (Duration > 0.f)
@@ -145,7 +156,6 @@ const FGeometryCacheTrackSampleInfo& UGeometryCacheTrackUsd::GetSampleInfo(float
 	}
 
 	// The sample info index must start from 0, while the sample index is between the range of the animation
-	const int32 SampleIndex = FindSampleIndexFromTime(Time, bLooping);
 	const int32 SampleInfoIndex = SampleIndex - StartFrameIndex;
 
 	FGeometryCacheTrackSampleInfo& CurrentSampleInfo = SampleInfos[SampleInfoIndex];
@@ -155,6 +165,7 @@ const FGeometryCacheTrackSampleInfo& UGeometryCacheTrackUsd::GetSampleInfo(float
 		FGeometryCacheMeshData TempMeshData;
 		if (GetMeshData(SampleIndex, TempMeshData))
 		{
+			float Time = GetTimeFromSampleIndex(SampleIndex);
 			CurrentSampleInfo = FGeometryCacheTrackSampleInfo(
 				Time,
 				(FBox)TempMeshData.BoundingBox,
diff --git a/Engine/Plugins/Importers/USDImporter/Source/GeometryCacheUSD/Private/GeometryCacheUSDSceneProxy.cpp b/Engine/Plugins/Importers/USDImporter/Source/GeometryCacheUSD/Private/GeometryCacheUSDSceneProxy.cpp
--- a/Engine/Plugins/Importers/USDImporter/Source/GeometryCacheUSD/Private/GeometryCacheUSDSceneProxy.cpp
+++ b/Engine/Plugins/Importers/USDImporter/Source/GeometryCacheUSD/Private/GeometryCacheUSDSceneProxy.cpp
@@ -50,12 +50,8 @@ bool FGeomCacheTrackUsdProxy::IsTopologyCompatible(int32 SampleIndexA, int32 Sam
 
 	if (UGeometryCacheTrackUsd* UsdTrack = Cast<UGeometryCacheTrackUsd>(Track))
 	{
-		// The boolean argument is not actually used
-		const float TimeA = UsdTrack->GetTimeFromSampleIndex(SampleIndexA);
-		const int32 NumVerticesA = UsdTrack->GetSampleInfo(TimeA, false).NumVertices;
-
-		const float TimeB = UsdTrack->GetTimeFromSampleIndex(SampleIndexB);
-		const int32 NumVerticesB = UsdTrack->GetSampleInfo(TimeB, false).NumVertices;
+		const int32 NumVerticesA = UsdTrack->GetSampleInfo(SampleIndexA).NumVertices;
+		const int32 NumVerticesB = UsdTrack->GetSampleInfo(SampleIndexB).NumVertices;
 
 		return NumVerticesA == NumVerticesB;
 	}
diff --git a/Engine/Plugins/Importers/USDImporter/Source/GeometryCacheUSD/Public/GeometryCacheTrackUSD.h b/Engine/Plugins/Importers/USDImporter/Source/GeometryCacheUSD/Public/GeometryCacheTrackUSD.h
--- a/Engine/Plugins/Importers/USDImporter/Source/GeometryCacheUSD/Public/GeometryCacheTrackUSD.h
+++ b/Engine/Plugins/Importers/USDImporter/Source/GeometryCacheUSD/Public/GeometryCacheTrackUSD.h
@@ -67,6 +67,8 @@ public:
 	const int32 FindSampleIndexFromTime(const float Time, const bool bLooping) const;
 	float GetTimeFromSampleIndex(int32 SampleIndex) const;
 	void GetFractionalFrameIndexFromTime(const float Time, const bool bLooping, int& OutFrameIndex, float& OutFraction) const;
+	// GetSampleInfo version that avoids converting time to index. Prefer this version when the index is already available
+	const FGeometryCacheTrackSampleInfo& GetSampleInfo(int32 SampleIndex);
 
 	bool GetMeshData(int32 SampleIndex, FGeometryCacheMeshData& OutMeshData);
 
