USD: Fix nullptr deref when failing to create a texture.

We were assuming that UsdUtils::CreateTexture can always create a texture for some reason. Here it is given the path to something that is not a texture, so it fails and returns nullptr.

diff --git a/Engine/Plugins/Importers/USDImporter/Source/USDSchemas/Private/USDGeomXformableTranslator.cpp b/Engine/Plugins/Importers/USDImporter/Source/USDSchemas/Private/USDGeomXformableTranslator.cpp
--- a/Engine/Plugins/Importers/USDImporter/Source/USDSchemas/Private/USDGeomXformableTranslator.cpp
+++ b/Engine/Plugins/Importers/USDImporter/Source/USDSchemas/Private/USDGeomXformableTranslator.cpp
@@ -996,22 +996,28 @@ void FUsdGeomXformableTranslator::CreateAlternativeDrawModeAssets(EUsdDrawMode D
 						Context->AssetCache.Get()
 					));
 
-					Context->AssetCache->CacheAsset(PrefixedTextureHash, Texture);
+					if (Texture)
+					{
+						Context->AssetCache->CacheAsset(PrefixedTextureHash, Texture);
+					}
 				}
 
-				// We link the textures to the prim, so that if the prim is reloaded the AUsdStageActor knows to potentially
-				// drop the textures. However we put the full attribute path on AssetUserData, so that when we're filling in
-				// our UUsdDrawModeComponent later, we know which texture came from which attribute
-
-				UUsdAssetUserData* TextureUserData = Texture->GetAssetUserData<UUsdAssetUserData>();
-				if (!TextureUserData)
+				if (Texture)
 				{
-					TextureUserData = NewObject<UUsdAssetUserData>(Texture, TEXT("USDAssetUserData"));
-					Texture->AddAssetUserData(TextureUserData);
-				}
-				TextureUserData->PrimPaths.AddUnique(UsdToUnreal::ConvertPath(Attr.GetPath()));
+					// We link the textures to the prim, so that if the prim is reloaded the AUsdStageActor knows to potentially
+					// drop the textures. However we put the full attribute path on AssetUserData, so that when we're filling in
+					// our UUsdDrawModeComponent later, we know which texture came from which attribute
 
-				Context->InfoCache->LinkAssetToPrim(PrimPath, Texture);
+					UUsdAssetUserData* TextureUserData = Texture->GetAssetUserData<UUsdAssetUserData>();
+					if (!TextureUserData)
+					{
+						TextureUserData = NewObject<UUsdAssetUserData>(Texture, TEXT("USDAssetUserData"));
+						Texture->AddAssetUserData(TextureUserData);
+					}
+					TextureUserData->PrimPaths.AddUnique(UsdToUnreal::ConvertPath(Attr.GetPath()));
+
+					Context->InfoCache->LinkAssetToPrim(PrimPath, Texture);
+				}
 			}
 		}
 	};
