USD: Changed the non-native schema analytic event to return only the count instead of the schema names.

diff --git a/Engine/Plugins/Importers/USDImporter/Source/USDUtilities/Private/USDConversionUtils.cpp b/Engine/Plugins/Importers/USDImporter/Source/USDUtilities/Private/USDConversionUtils.cpp
--- a/Engine/Plugins/Importers/USDImporter/Source/USDUtilities/Private/USDConversionUtils.cpp
+++ b/Engine/Plugins/Importers/USDImporter/Source/USDUtilities/Private/USDConversionUtils.cpp
@@ -3549,22 +3549,21 @@ void UsdUtils::CollectSchemaAnalytics(const UE::FUsdStage& Stage, const FString&
 		"Xform",
 		"XformCommonAPI"};
 
-	FString ConcatenatedSeenSchemas;
+	int32 CustomSchemaCount = 0;
 	for (const FString& SchemaName : SeenSchemas)
 	{
 		// We only care about non-native schemas
 		if (!NativeSchemaNames.Contains(SchemaName))
 		{
-			ConcatenatedSeenSchemas += SchemaName + TEXT(", ");
+			++CustomSchemaCount;
 		}
 	}
-	ConcatenatedSeenSchemas.RemoveFromEnd(TEXT(", "));
 
-	if (ConcatenatedSeenSchemas.Len() > 0)
+	if (CustomSchemaCount > 0)
 	{
 		TArray<FAnalyticsEventAttribute> EventAttributes;
-		EventAttributes.Emplace(TEXT("NonNativeSchemas"), ConcatenatedSeenSchemas);
-		IUsdClassesModule::SendAnalytics(MoveTemp(EventAttributes), FString::Printf(TEXT("%s.NonNativeSchemaCounts"), *EventName));
+		EventAttributes.Emplace(TEXT("CustomSchemas"), CustomSchemaCount);
+		IUsdClassesModule::SendAnalytics(MoveTemp(EventAttributes), FString::Printf(TEXT("%s.CustomSchemaCount"), *EventName));
 	}
 #endif	  // USE_USD_SDK
 }
