USDImporter: remove build settings forcing USD modules back to C++17

With the patch now in place to address previous issues with MSVC's
conforming preprocessor, it should no longer be necessary to force
USD modules to continue building with the older C++17 language
standard.

diff --git a/Engine/Plugins/Importers/USDImporter/Source/USDSchemas/USDSchemas.Build.cs b/Engine/Plugins/Importers/USDImporter/Source/USDSchemas/USDSchemas.Build.cs
--- a/Engine/Plugins/Importers/USDImporter/Source/USDSchemas/USDSchemas.Build.cs
+++ b/Engine/Plugins/Importers/USDImporter/Source/USDSchemas/USDSchemas.Build.cs
@@ -9,14 +9,6 @@ namespace UnrealBuildTool.Rules
 	{
 		public USDSchemas(ReadOnlyTargetRules Target) : base(Target)
 		{
-			// Does not compile with C++20:
-			// error C4002: too many arguments for function-like macro invocation 'TF_PP_CAT_IMPL'
-			// warning C5103: pasting '"TF_LOG_STACK_TRACE_ON_ERROR"' and '"TF_LOG_STACK_TRACE_ON_WARNING"' does not result in a valid preprocessing token
-			CppStandard = CppStandardVersion.Cpp17;
-
-			// Replace with PCHUsageMode.UseExplicitOrSharedPCHs when this plugin can compile with cpp20
-			PCHUsage = PCHUsageMode.NoPCHs;
-
 			bUseRTTI = true;
 
 			PrivateDependencyModuleNames.AddRange(
diff --git a/Engine/Plugins/Importers/USDImporter/Source/USDStageEditorViewModels/USDStageEditorViewModels.Build.cs b/Engine/Plugins/Importers/USDImporter/Source/USDStageEditorViewModels/USDStageEditorViewModels.Build.cs
--- a/Engine/Plugins/Importers/USDImporter/Source/USDStageEditorViewModels/USDStageEditorViewModels.Build.cs
+++ b/Engine/Plugins/Importers/USDImporter/Source/USDStageEditorViewModels/USDStageEditorViewModels.Build.cs
@@ -9,14 +9,6 @@ namespace UnrealBuildTool.Rules
 	{
 		public USDStageEditorViewModels(ReadOnlyTargetRules Target) : base(Target)
 		{
-			// Does not compile with C++20:
-			// error C4002: too many arguments for function-like macro invocation 'TF_PP_CAT_IMPL'
-			// warning C5103: pasting '"TF_LOG_STACK_TRACE_ON_ERROR"' and '"TF_LOG_STACK_TRACE_ON_WARNING"' does not result in a valid preprocessing token
-			CppStandard = CppStandardVersion.Cpp17;
-
-			// Replace with PCHUsageMode.UseExplicitOrSharedPCHs when this plugin can compile with cpp20
-			PCHUsage = PCHUsageMode.NoPCHs;
-
 			bUseRTTI = true;
 
 			PrivateDependencyModuleNames.AddRange(
diff --git a/Engine/Plugins/Importers/USDImporter/Source/USDUtilities/USDUtilities.Build.cs b/Engine/Plugins/Importers/USDImporter/Source/USDUtilities/USDUtilities.Build.cs
--- a/Engine/Plugins/Importers/USDImporter/Source/USDUtilities/USDUtilities.Build.cs
+++ b/Engine/Plugins/Importers/USDImporter/Source/USDUtilities/USDUtilities.Build.cs
@@ -9,14 +9,6 @@ namespace UnrealBuildTool.Rules
 	{
 		public USDUtilities(ReadOnlyTargetRules Target) : base(Target)
 		{
-			// Does not compile with C++20:
-			// error C4002: too many arguments for function-like macro invocation 'TF_PP_CAT_IMPL'
-			// warning C5103: pasting '"TF_LOG_STACK_TRACE_ON_ERROR"' and '"TF_LOG_STACK_TRACE_ON_WARNING"' does not result in a valid preprocessing token
-			CppStandard = CppStandardVersion.Cpp17;
-
-			// Replace with PCHUsageMode.UseExplicitOrSharedPCHs when this plugin can compile with cpp20
-			PCHUsage = PCHUsageMode.NoPCHs;
-
 			bUseRTTI = true;
 
 			PublicDependencyModuleNames.AddRange(
diff --git a/Engine/Plugins/Importers/USDImporter/Source/UnrealUSDWrapper/UnrealUSDWrapper.Build.cs b/Engine/Plugins/Importers/USDImporter/Source/UnrealUSDWrapper/UnrealUSDWrapper.Build.cs
--- a/Engine/Plugins/Importers/USDImporter/Source/UnrealUSDWrapper/UnrealUSDWrapper.Build.cs
+++ b/Engine/Plugins/Importers/USDImporter/Source/UnrealUSDWrapper/UnrealUSDWrapper.Build.cs
@@ -10,14 +10,6 @@ namespace UnrealBuildTool.Rules
 	{
 		public UnrealUSDWrapper(ReadOnlyTargetRules Target) : base(Target)
 		{
-			// Does not compile with C++20:
-			// error C2666: 'UE::FDummyWeakPtrType::operator ==': overloaded functions have similar conversions
-			// error C2666: 'UE::FDummyRefPtrType::operator ==': overloaded functions have similar conversions
-			CppStandard = CppStandardVersion.Cpp17;
-
-			// Replace with PCHUsageMode.UseExplicitOrSharedPCHs when this plugin can compile with cpp20
-			PCHUsage = PCHUsageMode.NoPCHs;
-
 			bUseRTTI = true;
 
 			PublicDependencyModuleNames.AddRange(
