USD: Serialize the LevelSequenceHelper internal maps into UE transactions.

diff --git a/Engine/Plugins/Importers/USDImporter/Source/USDStage/Private/USDLevelSequenceHelper.cpp b/Engine/Plugins/Importers/USDImporter/Source/USDStage/Private/USDLevelSequenceHelper.cpp
--- a/Engine/Plugins/Importers/USDImporter/Source/USDStage/Private/USDLevelSequenceHelper.cpp
+++ b/Engine/Plugins/Importers/USDImporter/Source/USDStage/Private/USDLevelSequenceHelper.cpp
@@ -589,6 +589,7 @@ public:
 	~FUsdLevelSequenceHelperImpl();
 
 	ULevelSequence* Init(const UE::FUsdStage& InUsdStage);
+	bool Serialize(FArchive& Ar);
 	void SetInfoCache(TSharedPtr<FUsdInfoCache> InfoCache);
 	void SetBBoxCache(TSharedPtr<UE::FUsdGeomBBoxCache> InBBoxCache);
 	bool HasData() const;
@@ -732,8 +733,14 @@ private:
 
 		// For now we support one binding per component type (mostly so we can fit a binding to a scene component and
 		// camera component for a Camera prim twin)
-		TMap<const UClass*, FGuid> ObjectClassToBindingGuid;
+		TMap<TWeakObjectPtr<const UClass>, FGuid> ObjectClassToBindingGuid;
 	};
+	friend FArchive& operator<<(FArchive& Ar, FUsdLevelSequenceHelperImpl::FPrimTwinBindings& Bindings)
+	{
+		Ar << Bindings.Sequence;
+		Ar << Bindings.ObjectClassToBindingGuid;
+		return Ar;
+	}
 
 	TMap<TWeakObjectPtr<const UUsdPrimTwin>, FPrimTwinBindings> PrimTwinToBindings;
 
@@ -875,6 +882,28 @@ ULevelSequence* FUsdLevelSequenceHelperImpl::Init(const UE::FUsdStage& InUsdStag
 	return MainLevelSequence;
 }
 
+bool FUsdLevelSequenceHelperImpl::Serialize(FArchive& Ar)
+{
+	Ar << MainLevelSequence;
+	Ar << LevelSequencesByIdentifier;
+	Ar << IdentifierByLevelSequence;
+	Ar << LocalLayersSequences;
+	Ar << SequencesID;
+	Ar << LayerIdentifierByLevelSequenceName;
+	Ar << PrimPathByLevelSequenceName;
+	Ar << PrimTwinToBindings;
+	Ar << RootMotionHandling;
+	Ar << StageActorBinding;
+
+	// Always keep SequenceHierarchyCache up-to-date given that it can't be serialized itself
+	if (Ar.IsLoading() && MainLevelSequence)
+	{
+		UMovieSceneCompiledDataManager::CompileHierarchy(MainLevelSequence, &SequenceHierarchyCache, EMovieSceneServerClientMask::All);
+	}
+
+	return true;
+}
+
 void FUsdLevelSequenceHelperImpl::SetInfoCache(TSharedPtr<FUsdInfoCache> InInfoCache)
 {
 	InfoCache = InInfoCache;
@@ -3043,7 +3072,7 @@ void FUsdLevelSequenceHelperImpl::RemovePossessable(const UUsdPrimTwin& PrimTwin
 	// ones don't modify the Sequence and change properties, so we must modify them here
 	Bindings->Sequence->Modify();
 
-	for (const TPair<const UClass*, FGuid>& Pair : Bindings->ObjectClassToBindingGuid)
+	for (const TPair<TWeakObjectPtr<const UClass>, FGuid>& Pair : Bindings->ObjectClassToBindingGuid)
 	{
 		const FGuid& ComponentPossessableGuid = Pair.Value;
 
@@ -3646,7 +3675,8 @@ void FUsdLevelSequenceHelperImpl::HandleMovieSceneChange(UMovieScene& MovieScene
 			continue;
 		}
 
-		for (TMap<const UClass*, FGuid>::TIterator BindingIt = Bindings.ObjectClassToBindingGuid.CreateIterator(); BindingIt; ++BindingIt)
+		for (TMap<TWeakObjectPtr<const UClass>, FGuid>::TIterator BindingIt = Bindings.ObjectClassToBindingGuid.CreateIterator(); BindingIt;
+			 ++BindingIt)
 		{
 			const FGuid& Guid = BindingIt->Value;
 
@@ -4311,6 +4341,10 @@ public:
 	{
 		return nullptr;
 	}
+	bool Serialize(FArchive& Ar)
+	{
+		return false;
+	}
 	void SetInfoCache(TSharedPtr<FUsdInfoCache> InfoCache){};
 	void SetBBoxCache(TSharedPtr<UE::FUsdGeomBBoxCache> InBBoxCache){};
 	bool HasData() const
@@ -4402,6 +4436,16 @@ ULevelSequence* FUsdLevelSequenceHelper::Init(const UE::FUsdStage& UsdStage)
 	}
 }
 
+bool FUsdLevelSequenceHelper::Serialize(FArchive& Ar)
+{
+	if (UsdSequencerImpl.IsValid())
+	{
+		return UsdSequencerImpl->Serialize(Ar);
+	}
+
+	return false;
+}
+
 void FUsdLevelSequenceHelper::OnStageActorRenamed()
 {
 	if (UsdSequencerImpl.IsValid())
diff --git a/Engine/Plugins/Importers/USDImporter/Source/USDStage/Private/USDStageActor.cpp b/Engine/Plugins/Importers/USDImporter/Source/USDStage/Private/USDStageActor.cpp
--- a/Engine/Plugins/Importers/USDImporter/Source/USDStage/Private/USDStageActor.cpp
+++ b/Engine/Plugins/Importers/USDImporter/Source/USDStage/Private/USDStageActor.cpp
@@ -3943,6 +3943,7 @@ void AUsdStageActor::Serialize(FArchive& Ar)
 		}
 
 		InfoCache->Serialize(Ar);
+		LevelSequenceHelper.Serialize(Ar);
 	}
 
 	// Make sure we clear the old deprecated cache if we still have one, as old scenes may be persisting textures
diff --git a/Engine/Plugins/Importers/USDImporter/Source/USDStage/Public/USDLevelSequenceHelper.h b/Engine/Plugins/Importers/USDImporter/Source/USDStage/Public/USDLevelSequenceHelper.h
--- a/Engine/Plugins/Importers/USDImporter/Source/USDStage/Public/USDLevelSequenceHelper.h
+++ b/Engine/Plugins/Importers/USDImporter/Source/USDStage/Public/USDLevelSequenceHelper.h
@@ -40,6 +40,9 @@ public:
 	/** Creates the main level sequence and subsequences from the usd stage layers */
 	ULevelSequence* Init(const UE::FUsdStage& UsdStage);
 
+	/** Allows serialization for transaction support */
+	bool Serialize(FArchive& Ar);
+
 	/** Sets the asset cache to use when fetching assets and asset info required for the level sequence animation, like UAnimSequences */
 	void SetInfoCache(TSharedPtr<FUsdInfoCache> InInfoCache);
 
