USD: Export skeletal and transform animations for components with anim blueprints and sockets.

diff --git a/Engine/Plugins/Importers/USDImporter/Source/USDExporter/Private/LevelSequenceExporterUSD.cpp b/Engine/Plugins/Importers/USDImporter/Source/USDExporter/Private/LevelSequenceExporterUSD.cpp
--- a/Engine/Plugins/Importers/USDImporter/Source/USDExporter/Private/LevelSequenceExporterUSD.cpp
+++ b/Engine/Plugins/Importers/USDImporter/Source/USDExporter/Private/LevelSequenceExporterUSD.cpp
@@ -25,6 +25,7 @@
 #include "Compilation/MovieSceneCompiledDataManager.h"
 #include "Components/SkeletalMeshComponent.h"
 #include "Editor.h"
+#include "Engine/SkeletalMesh.h"
 #include "Engine/SphereReflectionCapture.h"
 #include "EngineAnalytics.h"
 #include "Evaluation/MovieSceneSequenceHierarchy.h"
@@ -51,6 +52,15 @@
 
 #define LOCTEXT_NAMESPACE "LevelSequenceExporterUSD"
 
+static bool bExportAnimationsFromAllComponents = true;
+static FAutoConsoleVariableRef CVarExportAnimationsFromAllComponents(
+	TEXT("USD.ExportAnimationsFromAllComponents"),
+	bExportAnimationsFromAllComponents,
+	TEXT(
+		"If true it means that whenever we export LevelSequences to USD we may try exporting transforms and skeletal animations from all components of actors bound to the Sequence, even if those components aren't directly bound themselves. This is useful when using attach sockets or animation blueprints"
+	)
+);
+
 namespace UE::LevelSequenceExporterUSD::Private
 {
 	void SendAnalytics(
@@ -105,10 +115,10 @@ namespace UE::LevelSequenceExporterUSD::Private
 
 		bool bDestroyingJustHides = true;
 
-		virtual UObject* SpawnObject(const FGuid& Guid, 
-			UMovieScene& MovieScene, 
-			FMovieSceneSequenceIDRef TemplateID, 
-			TSharedRef<const FSharedPlaybackState> SharedPlaybackState, 
+		virtual UObject* SpawnObject(const FGuid& Guid,
+			UMovieScene& MovieScene,
+			FMovieSceneSequenceIDRef TemplateID,
+			TSharedRef<const FSharedPlaybackState> SharedPlaybackState,
 			int32 BindingIndex
 		) override
 		{
@@ -343,7 +353,7 @@ namespace UE::LevelSequenceExporterUSD::Private
 						break;
 					}
 				}
-				
+
 				FLevelSequenceEditorSpawnRegister::DestroySpawnedObject(Object, CustomSpawnableBinding);
 			}
 		}
@@ -925,6 +935,35 @@ namespace UE::LevelSequenceExporterUSD::Private
 			}
 		}
 
+		// Expand BoundObjects to include all components of all of its bound actors (even those without any binding to the LevelSequence).
+		// The idea here is that even if these don't have any tracks, the attach socket and AnimBlueprint fallbacks at the bottom of the
+		// loop below will still be triggered, letting us automatically capture the animations of these components that are "indirectly
+		// animated"
+		if (bExportAnimationsFromAllComponents)
+		{
+			TMap<UObject*, FSpawnedInstanceKey> NewEntries;
+			for (const TPair<UObject*, FSpawnedInstanceKey>& Pair : BoundObjects)
+			{
+				if (AActor* Actor = Cast<AActor>(Pair.Key))
+				{
+					if (USceneComponent* Root = Actor->GetRootComponent())
+					{
+						const bool bIncludeAllDescendants = true;
+						TArray<USceneComponent*> Children;
+						Root->GetChildrenComponents(bIncludeAllDescendants, Children);
+
+						for (USceneComponent* Child : Children)
+						{
+							NewEntries.Add(Child, FSpawnedInstanceKey{});
+						}
+					}
+				}
+			}
+			NewEntries.Append(BoundObjects);	// Prefer values from BoundObjects
+			NewEntries.Remove(nullptr);
+			Swap(NewEntries, BoundObjects);
+		}
+
 		// Generate bakers
 		for (const TPair<UObject*, FSpawnedInstanceKey>& Pair : BoundObjects)
 		{
@@ -970,8 +1009,7 @@ namespace UE::LevelSequenceExporterUSD::Private
 			// opening a referenced stage via an exported UsdStageActor
 			if (const FMovieSceneDynamicBinding* DynamicBinding = DynamicBindings.FindRef(InstanceKey))
 			{
-				if (const FMovieSceneDynamicBindingPayloadVariable* FoundPrimPathPayload = DynamicBinding->PayloadVariables.Find(TEXT("PrimPa"
-																																	  "th")))
+				if (const FMovieSceneDynamicBindingPayloadVariable* FoundPrimPathPayload = DynamicBinding->PayloadVariables.Find(TEXT("PrimPath")))
 				{
 					FString PrimPathInSourceStage = FoundPrimPathPayload->Value;
 					if (!PrimPathInSourceStage.IsEmpty())
@@ -1020,10 +1058,7 @@ namespace UE::LevelSequenceExporterUSD::Private
 										FText::Format(
 											LOCTEXT(
 												"NonIdealComposition",
-												"Exported animation for prim '{0}' may not compose correctly with the prims from referenced "
-												"layer '{1}' on the exported stage for the LevelSequence '{2}'. For best results, make sure "
-												"the referenced layer is saved to disk (i.e. not anonymous), has a defaultPrim setup, and "
-												"that the animation tracks are only bound to prims that are descendents of the defaultPrim."
+												"Exported animation for prim '{0}' may not compose correctly with the prims from referenced layer '{1}' on the exported stage for the LevelSequence '{2}'. For best results, make sure the referenced layer is saved to disk (i.e. not anonymous), has a defaultPrim setup, and that the animation tracks are only bound to prims that are descendents of the defaultPrim."
 											),
 											FText::FromString(PrimPathInSourceStage),
 											FText::FromString(LoadedStage.GetRootLayer().GetIdentifier()),
@@ -1051,7 +1086,42 @@ namespace UE::LevelSequenceExporterUSD::Private
 				continue;
 			}
 
+			TFunction<void(UnrealToUsd::FComponentBaker&)> AddBaker = [&InOutComponentBakers, &BoundComponent](UnrealToUsd::FComponentBaker& Baker)
+			{
+				// If we made a baker and we don't have one of this type for this component yet, add its lambda to the array
+				FCombinedComponentBakers& ExistingBakers = InOutComponentBakers.FindOrAdd(BoundComponent);
+				if (Baker.BakerType != UnrealToUsd::EBakingType::None && !EnumHasAnyFlags(ExistingBakers.CombinedBakingType, Baker.BakerType))
+				{
+					ExistingBakers.Bakers.Add(Baker);
+					ExistingBakers.CombinedBakingType |= Baker.BakerType;
+				}
+			};
+
+			TFunction<void(UnrealToUsd::FComponentBaker&)> GenerateSkeletalBaker =
+				[&UsdStage, &BoundComponent, &PrimPath](UnrealToUsd::FComponentBaker& InOutBaker)
+			{
+				if (USkeletalMeshComponent* SkeletalBoundComponent = Cast<USkeletalMeshComponent>(BoundComponent))
+				{
+					UE::FUsdPrim SkelAnimPrim = UsdStage.DefinePrim(UE::FSdfPath{*PrimPath}.AppendChild(TEXT("Anim")), TEXT("SkelAnimation"));
+
+					UE::FUsdPrim SkeletonPrim = UsdStage.DefinePrim(
+						UE::FSdfPath{*PrimPath}.AppendChild(UnrealIdentifiers::ExportedSkeletonPrimName),
+						TEXT("Skeleton")
+					);
+
+					if (SkelAnimPrim && SkeletonPrim)
+					{
+						UnrealToUsd::CreateSkeletalAnimationBaker(SkeletonPrim, SkelAnimPrim, *SkeletalBoundComponent, InOutBaker);
+					}
+					else
+					{
+						UE_LOG(LogUsd, Warning, TEXT("Failed to generate Skeleton or SkelAnimation prim when baking out SkelRoot '%s'"), *PrimPath);
+					}
+				}
+			};
+
 			bool bHasTransformBaker = false;
+			bool bHasSkeletalBaker = false;
 			if (const FMovieSceneBinding* Binding = MovieScene->FindBinding(InstanceKey.Key))
 			{
 				for (const UMovieSceneTrack* Track : Binding->GetTracks())
@@ -1082,27 +1152,7 @@ namespace UE::LevelSequenceExporterUSD::Private
 					// set the original skeletal animation track sections as disabled, so they'd fail the "IsTrackAnimated" check above
 					else if (Track->IsA<UMovieSceneSkeletalAnimationTrack>() || Track->IsA<UMovieSceneControlRigParameterTrack>())
 					{
-						if (USkeletalMeshComponent* SkeletalBoundComponent = Cast<USkeletalMeshComponent>(BoundComponent))
-						{
-							UE::FUsdPrim SkelAnimPrim = UsdStage.DefinePrim(UE::FSdfPath{*PrimPath}.AppendChild(TEXT("Anim")), TEXT("SkelAnimation"));
-
-							UE::FUsdPrim SkeletonPrim = UsdStage.DefinePrim(
-								UE::FSdfPath{*PrimPath}.AppendChild(UnrealIdentifiers::ExportedSkeletonPrimName),
-								TEXT("Skeleton")
-							);
-
-							if (!SkelAnimPrim || !SkeletonPrim)
-							{
-								UE_LOG(
-									LogUsd,
-									Warning,
-									TEXT("Failed to generate Skeleton or SkelAnimation prim when baking out SkelRoot '%s'"),
-									*PrimPath
-								);
-								continue;
-							}
-							UnrealToUsd::CreateSkeletalAnimationBaker(SkeletonPrim, SkelAnimPrim, *SkeletalBoundComponent, Baker);
-						}
+						GenerateSkeletalBaker(Baker);
 					}
 					// If we have an attach track that attaches the object to somewhere else, then we'll need to bake in that transform
 					// change, as we can't export "hierarchy changes" otherwise
@@ -1111,18 +1161,16 @@ namespace UE::LevelSequenceExporterUSD::Private
 						UnrealToUsd::CreateComponentPropertyBaker(Prim, *BoundComponent, TEXT("Transform"), Baker);
 					}
 
-					// If we made a baker and we don't have one of this type for this component yet, add its lambda to the array
-					FCombinedComponentBakers& ExistingBakers = InOutComponentBakers.FindOrAdd(BoundComponent);
-					if (Baker.BakerType != UnrealToUsd::EBakingType::None && !EnumHasAnyFlags(ExistingBakers.CombinedBakingType, Baker.BakerType))
-					{
-						ExistingBakers.Bakers.Add(Baker);
-						ExistingBakers.CombinedBakingType |= Baker.BakerType;
-					}
+					AddBaker(Baker);
 
 					if (Baker.BakerType == UnrealToUsd::EBakingType::Transform)
 					{
 						bHasTransformBaker = true;
 					}
+					else if (Baker.BakerType == UnrealToUsd::EBakingType::Skeletal)
+					{
+						bHasSkeletalBaker = true;
+					}
 				}
 			}
 
@@ -1132,17 +1180,62 @@ namespace UE::LevelSequenceExporterUSD::Private
 			// would cause the parent prim's skeletal animation to also affect its child prims.
 			// Ideally we'd actually search through the tracks to know for sure whether our parent has a SkeletalAnimation section,
 			// but it's probably safer to just do this in case it is hidden behind N subsequences or some obscure feature
-			if (!bHasTransformBaker && BoundComponent->GetAttachSocketName() != NAME_None)
+			if (!bHasTransformBaker)
 			{
-				UnrealToUsd::FComponentBaker Baker;
-				const FString PropertyPath = TEXT("Transform");
-				UnrealToUsd::CreateComponentPropertyBaker(Prim, *BoundComponent, PropertyPath, Baker);
+				// If any ancestor component has an attach socket, we may need to bake our transform
+				bool bHasSocketAttachment = false;
+				USceneComponent* Iterator = BoundComponent;
+				while (Iterator)
+				{
+					if (Iterator->GetAttachSocketName() != NAME_None)
+					{
+						bHasSocketAttachment = true;
+						break;
+					}
 
-				FCombinedComponentBakers& ExistingBakers = InOutComponentBakers.FindOrAdd(BoundComponent);
-				if (Baker.BakerType != UnrealToUsd::EBakingType::None && !EnumHasAnyFlags(ExistingBakers.CombinedBakingType, Baker.BakerType))
+					if (AActor* OwnerActor = Iterator->GetOwner())
+					{
+						if (OwnerActor->GetRootComponent() == Iterator)
+						{
+							// Don't climb out of the actor
+							break;
+						}
+					}
+
+					Iterator = Iterator->GetAttachParent();
+				}
+
+				if (bHasSocketAttachment)
 				{
-					ExistingBakers.Bakers.Add(Baker);
-					ExistingBakers.CombinedBakingType |= Baker.BakerType;
+					UnrealToUsd::FComponentBaker Baker;
+
+					const FString PropertyPath = TEXT("Transform");
+					UnrealToUsd::CreateComponentPropertyBaker(Prim, *BoundComponent, PropertyPath, Baker);
+
+					AddBaker(Baker);
+				}
+			}
+
+			// There are many different ways in which SkeletalMeshComponents may animate their joints without having any Sequencer track
+			// or even any binding, and this check here tries filling in that gap and generating a skeletal baker if needed.
+			// (Search for bExportAnimationsFromAllComponents in this file to see how we can get in here without having a binding)
+			if (!bHasSkeletalBaker)
+			{
+				if (USkeletalMeshComponent* SkeletalBoundComponent = Cast<USkeletalMeshComponent>(BoundComponent))
+				{
+					const bool bNeedsSkeletalBaker = SkeletalBoundComponent->HasValidAnimationInstance()
+													 || SkeletalBoundComponent->LeaderPoseComponent.IsValid()
+													 || (SkeletalBoundComponent->GetAnimationMode() == EAnimationMode::AnimationBlueprint
+														 && SkeletalBoundComponent->AnimClass);
+
+					if (bNeedsSkeletalBaker)
+					{
+						UnrealToUsd::FComponentBaker Baker;
+
+						GenerateSkeletalBaker(Baker);
+
+						AddBaker(Baker);
+					}
 				}
 			}
 		}
diff --git a/Engine/Plugins/Importers/USDImporter/Source/USDUtilities/Private/USDPrimConversion.cpp b/Engine/Plugins/Importers/USDImporter/Source/USDUtilities/Private/USDPrimConversion.cpp
--- a/Engine/Plugins/Importers/USDImporter/Source/USDUtilities/Private/USDPrimConversion.cpp
+++ b/Engine/Plugins/Importers/USDImporter/Source/USDUtilities/Private/USDPrimConversion.cpp
@@ -3797,6 +3797,73 @@ bool UnrealToUsd::CreateComponentPropertyBaker(
 	return false;
 }
 
+namespace UE::USDPrimConversion::Private
+{
+	// Returns bone-space joint transforms from the SkeletalMeshComponent while paying attention to whether
+	// it has a LeaderPoseComponent or not.
+	//
+	// References:
+	// - FMLDeformerEditorToolkit::GetDebugActorComponentSpaceTransforms
+	// - FAnimationRecorder::GetBoneTransforms
+	void GetBoneTransforms(USkeletalMeshComponent* Component, TArray<FTransform>& BoneTransforms)
+	{
+		if (!Component)
+		{
+			return;
+		}
+
+		int32 NumBones = INDEX_NONE;
+		if (USkeletalMesh* Mesh = Component->GetSkeletalMeshAsset())
+		{
+			const FReferenceSkeleton& RefSkel = Mesh->GetRefSkeleton();
+			NumBones = RefSkel.GetNum();
+		}
+		if (NumBones == INDEX_NONE)
+		{
+			return;
+		}
+
+		if (USkeletalMeshComponent* Leader = Cast<USkeletalMeshComponent>(Component->LeaderPoseComponent.Get()))
+		{
+			const TArray<FTransform>& LeaderTransforms = Leader->GetBoneSpaceTransforms();
+			const TArray<FTransform>& FollowerTransforms = Component->GetBoneSpaceTransforms();
+
+			const TArray<int32>& BoneMap = Component->GetLeaderBoneMap();
+
+			BoneTransforms.SetNumUninitialized(NumBones);
+			for (int32 BoneIndex = 0; BoneIndex < NumBones; BoneIndex++)
+			{
+				if (BoneMap.IsValidIndex(BoneIndex) && LeaderTransforms.IsValidIndex(BoneMap[BoneIndex]))
+				{
+					BoneTransforms[BoneIndex] = LeaderTransforms[BoneMap[BoneIndex]];
+				}
+				else if (FollowerTransforms.IsValidIndex(BoneIndex))
+				{
+					BoneTransforms[BoneIndex] = FollowerTransforms[BoneIndex];
+				}
+			}
+		}
+		else
+		{
+			BoneTransforms = Component->GetBoneSpaceTransforms();
+		}
+	}
+
+	void RefreshSkeletalMeshComponent(USkeletalMeshComponent& Component)
+	{
+		// This whole incantation is required or else the component will really not update until the next frame.
+		// Note: This will also cause the update of morph target weights.
+		Component.TickAnimation(0.f, false);
+		Component.UpdateLODStatus();
+		Component.RefreshBoneTransforms();
+		Component.RefreshFollowerComponents();
+		Component.UpdateComponentToWorld();
+		Component.FinalizeBoneTransform();
+		Component.MarkRenderTransformDirty();
+		Component.MarkRenderDynamicDataDirty();
+	}
+}	 // namespace UE::USDPrimConversion::Private
+
 bool UnrealToUsd::CreateSkeletalAnimationBaker(
 	UE::FUsdPrim& SkeletonPrim,
 	UE::FUsdPrim& SkelAnimation,
@@ -3805,6 +3872,8 @@ bool UnrealToUsd::CreateSkeletalAnimationBaker(
 )
 {
 #if WITH_EDITOR
+	using namespace UE::USDPrimConversion::Private;
+
 	USkeletalMesh* SkeletalMesh = Component.GetSkeletalMeshAsset();
 	if (!SkeletalMesh)
 	{
@@ -3920,16 +3989,11 @@ bool UnrealToUsd::CreateSkeletalAnimationBaker(
 		Rotations.resize(NumBones);
 		Scales.resize(NumBones);
 
-		// This whole incantation is required or else the component will really not update until the next frame.
-		// Note: This will also cause the update of morph target weights.
-		Component.TickAnimation(0.f, false);
-		Component.UpdateLODStatus();
-		Component.RefreshBoneTransforms();
-		Component.RefreshFollowerComponents();
-		Component.UpdateComponentToWorld();
-		Component.FinalizeBoneTransform();
-		Component.MarkRenderTransformDirty();
-		Component.MarkRenderDynamicDataDirty();
+		if (USkeletalMeshComponent* Leader = Cast<USkeletalMeshComponent>(Component.LeaderPoseComponent.Get()))
+		{
+			RefreshSkeletalMeshComponent(*Leader);
+		}
+		RefreshSkeletalMeshComponent(Component);
 
 		// I'm not entirely sure why this is needed but FFbxExporter::ExportAnimTrack and FFbxExporter::ExportLevelSequenceBaked3DTransformTrack
 		// do this so for safety maybe we should as well?
@@ -3938,7 +4002,8 @@ bool UnrealToUsd::CreateSkeletalAnimationBaker(
 			Owner->Tick(0.0f);
 		}
 
-		const TArray<FTransform>& LocalBoneTransforms = Component.GetBoneSpaceTransforms();
+		TArray<FTransform> LocalBoneTransforms;
+		GetBoneTransforms(&Component, LocalBoneTransforms);
 
 		// For whatever reason it seems that sometimes this is not ready for us, so let's force it to be recalculated
 		if (LocalBoneTransforms.Num() == 0)
