USD: Fixed possible packaging issues when packaging a project that uses USDImporter on platforms on which it is not supported.

diff --git a/Engine/Plugins/Importers/USDImporter/USDImporter.uplugin b/Engine/Plugins/Importers/USDImporter/USDImporter.uplugin
--- a/Engine/Plugins/Importers/USDImporter/USDImporter.uplugin
+++ b/Engine/Plugins/Importers/USDImporter/USDImporter.uplugin
@@ -128,7 +128,12 @@
 		},
 		{
 			"Name": "Interchange",
-			"Enabled": true
+			"Enabled": true,
+			"PlatformAllowList": [
+				"Win64",
+				"Mac",
+				"Linux"
+			]
 		},
 		{
 			"Name": "LevelSequenceEditor",
@@ -140,7 +145,12 @@
 		},
 		{
 			"Name": "MDLImporter",
-			"Enabled": true
+			"Enabled": true,
+			"PlatformAllowList": [
+				"Win64",
+				"Mac",
+				"Linux"
+			]
 		},
 		{
 			"Name": "PythonScriptPlugin",
