USD: Fix strange hang/crash when opening any USD stage.

diff --git a/Engine/Plugins/Importers/USDImporter/Source/USDStageEditorViewModels/Private/USDLayersViewModel.cpp b/Engine/Plugins/Importers/USDImporter/Source/USDStageEditorViewModels/Private/USDLayersViewModel.cpp
--- a/Engine/Plugins/Importers/USDImporter/Source/USDStageEditorViewModels/Private/USDLayersViewModel.cpp
+++ b/Engine/Plugins/Importers/USDImporter/Source/USDStageEditorViewModels/Private/USDLayersViewModel.cpp
@@ -164,7 +164,10 @@ void FUsdLayerViewModel::RefreshData()
 	const pxr::SdfLayerHandle& EditTargetLayer = UsdStageRef->GetEditTarget().GetLayer();
 	const pxr::SdfLayerHandle& FocusedEditTarget = IsolatedStageRef ? IsolatedStageRef->GetEditTarget().GetLayer() : pxr::SdfLayerHandle{};
 
-	LayerModel->bIsEditTarget = (IsolatedStageRef) ? LayerModel->bIsInIsolatedStage && FocusedEditTarget == UsdLayer : EditTargetLayer == UsdLayer;
+	// For whatever reason, doing the "EditTargetLayer == UsdLayer" comparison direcly leads to a hang on Windows and a crash
+	// on the Mac? It could be related to the USD 24.03 update somehow
+	LayerModel->bIsEditTarget = (IsolatedStageRef) ? LayerModel->bIsInIsolatedStage && FocusedEditTarget == UsdLayer
+												   : (EditTargetLayer && UsdLayer && EditTargetLayer->GetIdentifier() == UsdLayer->GetIdentifier());
 
 	LayerModel->bIsMuted = (IsolatedStageRef && LayerModel->bIsInIsolatedStage)
 							   // If we isolating, we're only muted if we're muted on the isolated stage
