USD: Tweak fix on cl 32722819 to also handled isolated stages.

diff --git a/Engine/Plugins/Importers/USDImporter/Source/USDStageEditorViewModels/Private/USDLayersViewModel.cpp b/Engine/Plugins/Importers/USDImporter/Source/USDStageEditorViewModels/Private/USDLayersViewModel.cpp
--- a/Engine/Plugins/Importers/USDImporter/Source/USDStageEditorViewModels/Private/USDLayersViewModel.cpp
+++ b/Engine/Plugins/Importers/USDImporter/Source/USDStageEditorViewModels/Private/USDLayersViewModel.cpp
@@ -164,10 +164,18 @@ void FUsdLayerViewModel::RefreshData()
 	const pxr::SdfLayerHandle& EditTargetLayer = UsdStageRef->GetEditTarget().GetLayer();
 	const pxr::SdfLayerHandle& FocusedEditTarget = IsolatedStageRef ? IsolatedStageRef->GetEditTarget().GetLayer() : pxr::SdfLayerHandle{};
 
-	// For whatever reason, doing the "EditTargetLayer == UsdLayer" comparison direcly leads to a hang on Windows and a crash
-	// on the Mac? It could be related to the USD 24.03 update somehow
-	LayerModel->bIsEditTarget = (IsolatedStageRef) ? LayerModel->bIsInIsolatedStage && FocusedEditTarget == UsdLayer
-												   : (EditTargetLayer && UsdLayer && EditTargetLayer->GetIdentifier() == UsdLayer->GetIdentifier());
+	if (UsdLayer)
+	{
+		// For whatever reason, doing the "EditTargetLayer == UsdLayer" comparison direcly leads to a hang on Windows and a crash
+		// on the Mac? It could be related to the USD 24.03 update somehow
+		LayerModel->bIsEditTarget = (IsolatedStageRef) ? LayerModel->bIsInIsolatedStage
+															 && (FocusedEditTarget && FocusedEditTarget->GetIdentifier() == UsdLayer->GetIdentifier())
+													   : (EditTargetLayer && EditTargetLayer->GetIdentifier() == UsdLayer->GetIdentifier());
+	}
+	else
+	{
+		LayerModel->bIsEditTarget = false;
+	}
 
 	LayerModel->bIsMuted = (IsolatedStageRef && LayerModel->bIsInIsolatedStage)
 							   // If we isolating, we're only muted if we're muted on the isolated stage
