[VFX Reference Update] Fix an issue in linux with the MaterialX library load ordering that was causing double free on editor exit

diff --git a/Engine/Plugins/Importers/USDImporter/Source/USDSchemas/USDSchemas.Build.cs b/Engine/Plugins/Importers/USDImporter/Source/USDSchemas/USDSchemas.Build.cs
--- a/Engine/Plugins/Importers/USDImporter/Source/USDSchemas/USDSchemas.Build.cs
+++ b/Engine/Plugins/Importers/USDImporter/Source/USDSchemas/USDSchemas.Build.cs
@@ -68,7 +68,15 @@ namespace UnrealBuildTool.Rules
 					}
 				);
 
-				AddEngineThirdPartyPrivateStaticDependencies(Target, "MaterialX");
+				if (Target.Platform == UnrealTargetPlatform.Win64)
+				{
+					PrivateDependencyModuleNames.AddRange(
+						new string[]
+						{
+							"MaterialX"
+						}
+					);
+				}
 			}
 		}
 	}
diff --git a/Engine/Plugins/Importers/USDImporter/Source/UnrealUSDWrapper/UnrealUSDWrapper.Build.cs b/Engine/Plugins/Importers/USDImporter/Source/UnrealUSDWrapper/UnrealUSDWrapper.Build.cs
--- a/Engine/Plugins/Importers/USDImporter/Source/UnrealUSDWrapper/UnrealUSDWrapper.Build.cs
+++ b/Engine/Plugins/Importers/USDImporter/Source/UnrealUSDWrapper/UnrealUSDWrapper.Build.cs
@@ -28,7 +28,6 @@ namespace UnrealBuildTool.Rules
 					"CoreUObject",
 					"Engine",
 					"IntelTBB",
-					"MaterialX", // Needed for the standard data libraries
 					"Projects", // For plugin manager within UnrealUSDWrapper.cpp
 					"USDClasses"
 				}
