USD: Fix missing placeholder value for string Printf.

diff --git a/Engine/Plugins/Importers/USDImporter/Source/USDStageEditorViewModels/Private/USDObjectFieldViewModel.cpp b/Engine/Plugins/Importers/USDImporter/Source/USDStageEditorViewModels/Private/USDObjectFieldViewModel.cpp
--- a/Engine/Plugins/Importers/USDImporter/Source/USDStageEditorViewModels/Private/USDObjectFieldViewModel.cpp
+++ b/Engine/Plugins/Importers/USDImporter/Source/USDStageEditorViewModels/Private/USDObjectFieldViewModel.cpp
@@ -503,7 +503,7 @@ void FUsdObjectFieldsViewModel::Refresh(const UE::FUsdStageWeak& InUsdStage, con
 					}
 					else if (Targets.size() > 1)
 					{
-						FString CombinedTargets = FString::Printf(TEXT("%d elements: ["));
+						FString CombinedTargets = FString::Printf(TEXT("%d elements: ["), static_cast<int32>(Targets.size()));
 						for (const pxr::SdfPath& Target : Targets)
 						{
 							CombinedTargets += UsdToUnreal::ConvertPath(Target) + TEXT(", ");
