USD: Prevent invalid static meshes being generated when one of the LOD MeshDescriptions is empty.

diff --git a/Engine/Plugins/Importers/USDImporter/Source/USDSchemas/Private/USDGeomMeshTranslator.cpp b/Engine/Plugins/Importers/USDImporter/Source/USDSchemas/Private/USDGeomMeshTranslator.cpp
--- a/Engine/Plugins/Importers/USDImporter/Source/USDSchemas/Private/USDGeomMeshTranslator.cpp
+++ b/Engine/Plugins/Importers/USDImporter/Source/USDSchemas/Private/USDGeomMeshTranslator.cpp
@@ -1978,15 +1978,32 @@ void FGeomMeshCreateAssetsTaskChain::SetupTasks()
 		   }
 		   bCollectedMetadata = true;
 
-		   // If we have at least one valid LOD, we should keep going
-		   for (const FMeshDescription& MeshDescription : LODIndexToMeshDescription)
+		   // Strip empty MeshDescriptions: If we have some valid LODs and some empty our StaticMesh RenderData
+		   // will end up with NaN bounds and will have to be discarded anyway
+		   bool bHasValidLOD = false;
+		   for (int32 LODIndex = LODIndexToMeshDescription.Num() - 1; LODIndex >= 0; --LODIndex)
 		   {
-			   if (!MeshDescription.IsEmpty())
+			   const FMeshDescription& MeshDescription = LODIndexToMeshDescription[LODIndex];
+			   if (MeshDescription.IsEmpty())
 			   {
-				   return true;
+				   LODIndexToMeshDescription.RemoveAt(LODIndex, EAllowShrinking::No);
+				   LODIndexToMaterialInfo.RemoveAt(LODIndex, EAllowShrinking::No);
+				   UE_LOG(
+					   LogUsd,
+					   Warning,
+					   TEXT("Ignoring mesh data collected for LOD%d of prim '%s' as it is empty. Is the prim invisible?"),
+					   LODIndex,
+					   bParseLODs ? *GetPrim().GetParent().GetPrimPath().GetString() : *GetPrim().GetPrimPath().GetString()
+				   );
+			   }
+			   else
+			   {
+				   bHasValidLOD = true;
 			   }
 		   }
-		   return false;
+
+		   // If we have at least one valid LOD, we should keep going
+		   return bHasValidLOD;
 	   });
 
 	FBuildStaticMeshTaskChain::SetupTasks();
