USD: Add basic support for UsdLuxShadowAPI (just the shadow:enable attribute for now).

diff --git a/Engine/Plugins/Importers/USDImporter/Source/USDUtilities/Private/USDLightConversion.cpp b/Engine/Plugins/Importers/USDImporter/Source/USDUtilities/Private/USDLightConversion.cpp
--- a/Engine/Plugins/Importers/USDImporter/Source/USDUtilities/Private/USDLightConversion.cpp
+++ b/Engine/Plugins/Importers/USDImporter/Source/USDUtilities/Private/USDLightConversion.cpp
@@ -32,6 +32,7 @@
 #include "pxr/usd/usdLux/domeLight.h"
 #include "pxr/usd/usdLux/lightAPI.h"
 #include "pxr/usd/usdLux/rectLight.h"
+#include "pxr/usd/usdLux/shadowAPI.h"
 #include "pxr/usd/usdLux/shapingAPI.h"
 #include "pxr/usd/usdLux/sphereLight.h"
 #include "USDIncludesEnd.h"
@@ -117,6 +118,18 @@ bool UsdToUnreal::ConvertLight(const pxr::UsdPrim& Prim, ULightComponentBase& Li
 		LightComponent->Temperature = UsdUtils::GetUsdValue<float>(LightAPI.GetColorTemperatureAttr(), UsdTimeCode);
 	}
 
+	if (const pxr::UsdLuxShadowAPI ShadowAPI{Prim})
+	{
+		if (pxr::UsdAttribute Attr = ShadowAPI.GetShadowEnableAttr())
+		{
+			bool bEnable = true;
+			if (Attr.Get(&bEnable, UsdTimeCode))
+			{
+				LightComponentBase.SetCastShadows(bEnable);
+			}
+		}
+	}
+
 	return true;
 }
 
@@ -470,6 +483,32 @@ bool UnrealToUsd::ConvertLightComponent(const ULightComponentBase& LightComponen
 		UsdUtils::NotifyIfOverriddenOpinion(Attr);
 	}
 
+	// Only author shadow stuff if we need to, as it involves applying an API schema. We don't want to
+	// open up an USD stage -> Change a light intensity and save -> End up adding she shadow API schema and attribute
+	// just to put the default value of true
+	bool bPrimCastsShadows = true;
+	if (pxr::UsdLuxShadowAPI ExistingShadowAPI{Prim})
+	{
+		if (pxr::UsdAttribute Attr = ExistingShadowAPI.GetShadowEnableAttr())
+		{
+			bool bEnable = true;
+			if (Attr.Get(&bEnable, UsdTimeCode))
+			{
+				bPrimCastsShadows = bEnable;
+			}
+		}
+	}
+	const bool bComponentCastsShadows = static_cast<bool>(LightComponent.CastShadows);
+	if (bComponentCastsShadows != bPrimCastsShadows)
+	{
+		pxr::UsdLuxShadowAPI ShadowAPI = pxr::UsdLuxShadowAPI::Apply(Prim);
+		if (pxr::UsdAttribute Attr = ShadowAPI.CreateShadowEnableAttr())
+		{
+			Attr.Set(bComponentCastsShadows, UsdTimeCode);
+			UsdUtils::NotifyIfOverriddenOpinion(Attr);
+		}
+	}
+
 	return true;
 }
 
