USD: Remove direct usage of static TCHAR as that causes trouble on the Mac.

diff --git a/Engine/Plugins/Importers/USDImporter/Source/USDUtilities/Private/USDValueConversion.cpp b/Engine/Plugins/Importers/USDImporter/Source/USDUtilities/Private/USDValueConversion.cpp
--- a/Engine/Plugins/Importers/USDImporter/Source/USDUtilities/Private/USDValueConversion.cpp
+++ b/Engine/Plugins/Importers/USDImporter/Source/USDUtilities/Private/USDValueConversion.cpp
@@ -26,9 +26,7 @@ namespace UE::USDValueConversion::Private
 	const static char DoubleQuoteChar = '\"';
 	const static char EscapeChar = '\\';
 	const static char AssetDelimiterChar = '@';
-	const static TCHAR SingleQuoteTCHAR = '\'';
 	const static TCHAR DoubleQuoteTCHAR = '\"';
-	const static TCHAR AssetDelimiterTCHAR = '@';
 }
 
 #if USE_USD_SDK
@@ -1168,7 +1166,7 @@ namespace UsdUtils
 	{
 		using namespace UE::USDValueConversion::Private;
 
-		return FString::Printf(TEXT("%s%s%s"), &AssetDelimiterTCHAR, *Value, &AssetDelimiterTCHAR);
+		return FString::Printf(TEXT("@%s@"), *Value);
 	}
 
 	FString StringifyAsMatrix2d(const FMatrix2D& Value)
@@ -1452,9 +1450,14 @@ namespace UsdUtils
 		{
 			// If the string contains a double quote, use a single quote as a delimiter, or else the string
 			// will be impossible to tokenize into the source string later when unstringifying
-			const TCHAR* Delimiter = Element.Contains(&DoubleQuoteTCHAR) ? &SingleQuoteTCHAR : &DoubleQuoteTCHAR;
-
-			Result += FString::Printf(TEXT("%s%s%s, "), Delimiter, *Element, Delimiter);
+            if (Element.Contains(&DoubleQuoteTCHAR))
+            {
+                Result += FString::Printf(TEXT("'%s', "), *Element);
+            }
+            else
+            {
+                Result += FString::Printf(TEXT("\"%s\", "), *Element);
+            }
 		}
 		Result.RemoveFromEnd(TEXT(", "));
 		Result += TEXT("]");
@@ -1479,7 +1482,7 @@ namespace UsdUtils
 		FString Result = TEXT("[");
 		for (const FString& Element : Value)
 		{
-			Result += FString::Printf(TEXT("%s%s%s, "), &AssetDelimiterTCHAR, *Element, &AssetDelimiterTCHAR);
+			Result += FString::Printf(TEXT("@%s@, "), *Element);
 		}
 		Result.RemoveFromEnd(TEXT(", "));
 		Result += TEXT("]");
@@ -1501,9 +1504,14 @@ namespace UsdUtils
 		{
 			// If the string contains a double quote, use a single quote as a delimiter, or else the string
 			// will be impossible to tokenize into the source string later when unstringifying
-			const TCHAR* Delimiter = Element.Contains(&DoubleQuoteTCHAR) ? &SingleQuoteTCHAR : &DoubleQuoteTCHAR;
-
-			Result += FString::Printf(TEXT("%s%s%s, "), Delimiter, *Element, Delimiter);
+            if (Element.Contains(&DoubleQuoteTCHAR))
+            {
+                Result += FString::Printf(TEXT("'%s', "), *Element);
+            }
+            else
+            {
+                Result += FString::Printf(TEXT("\"%s\", "), *Element);
+            }
 		}
 		Result.RemoveFromEnd(TEXT(", "));
 		Result += TEXT("])");
