USD: Fix relationship changes not being handled by USDTransactor;

Fix ensure on test_update_refreshes_assets.py due to not allowing enough time for Skeletal Mesh Component delayed tasks.

diff --git a/Engine/Plugins/Importers/USDImporter/Source/USDStage/Private/USDTransactor.cpp b/Engine/Plugins/Importers/USDImporter/Source/USDStage/Private/USDTransactor.cpp
--- a/Engine/Plugins/Importers/USDImporter/Source/USDStage/Private/USDTransactor.cpp
+++ b/Engine/Plugins/Importers/USDImporter/Source/USDStage/Private/USDTransactor.cpp
@@ -16,6 +16,7 @@
 #include "UsdWrappers/SdfPath.h"
 #include "UsdWrappers/UsdAttribute.h"
 #include "UsdWrappers/UsdPrim.h"
+#include "UsdWrappers/UsdRelationship.h"
 #include "UsdWrappers/UsdStage.h"
 
 #if WITH_EDITOR
@@ -35,6 +36,7 @@ namespace UsdUtils
 	struct FTransactorAttributeChange
 	{
 		FString PropertyName;
+		bool bIsRelationship = false;
 		FString Field;
 		FString AttributeTypeName;	   // Full SdfValueTypeName of the attribute (e.g. normal3f, bool, texCoord3d, float2) so that we can undo/redo
 									   // attribute creation
@@ -80,6 +82,7 @@ FArchive& operator<<(FArchive& Ar, UsdUtils::FTransactorAttributeChange& Change)
 	Ar << Change.PropertyName;
 	Ar << Change.Field;
 	Ar << Change.AttributeTypeName;
+	Ar << Change.bIsRelationship;
 	Ar << Change.OldValue;
 	Ar << Change.NewValue;
 	Ar << Change.TimeSamples;
@@ -261,6 +264,10 @@ namespace UsdUtils
 								}
 							}
 						}
+						else if (UE::FUsdRelationship Relationship = Prim.GetRelationship(*ConvertedAttributeChange.PropertyName))
+						{
+							ConvertedAttributeChange.bIsRelationship = true;
+						}
 						else
 						{
 							UE_LOG(
@@ -430,6 +437,7 @@ namespace UsdUtils
 		const FString& PropertyName,
 		const FString& Field,
 		const FString& AttributeTypeName,
+		bool bIsRelationship,
 		bool bRemoveProperty,
 		const FConvertedVtValue& Value,
 		UE::FUsdPrim& Prim,
@@ -442,41 +450,75 @@ namespace UsdUtils
 		}
 
 		bool bCreated = false;
-
 		UE::FUsdAttribute Attribute;
+		UE::FUsdRelationship Relationship;
+
 		if (PropertyName != TEXT("kind"))	 // Kind is prim metadata, not an attribute
 		{
 			if (bRemoveProperty)
 			{
-				Attribute = Prim.GetAttribute(*PropertyName);
-				if (!Attribute)
+				if (bIsRelationship)
 				{
-					return true;
+					Relationship = Prim.GetRelationship(*PropertyName);
+					if (!Relationship)
+					{
+						return true;
+					}
+				}
+				else
+				{
+					Attribute = Prim.GetAttribute(*PropertyName);
+					if (!Attribute)
+					{
+						return true;
+					}
 				}
 			}
 			else
 			{
-				bool bHadAttr = Prim.HasAttribute(*PropertyName);
-				Attribute = Prim.CreateAttribute(*PropertyName, *AttributeTypeName);
-				if (!Attribute)
+				if (bIsRelationship)
 				{
-					// We expect to fail to create an attribute if we have no typename here (e.g. undo remove property)
-					if (AttributeTypeName.IsEmpty())
+					bool bHadRelationship = Prim.HasRelationship(*PropertyName);
+					Relationship = Prim.CreateRelationship(*PropertyName);
+					if (!Relationship)
 					{
 						UE_LOG(
 							LogUsd,
 							Warning,
-							TEXT("Failed to create attribute '%s' with typename '%s' for prim '%s'"),
+							TEXT("Failed to create relationship '%s' with for prim '%s'"),
 							*PropertyName,
-							*AttributeTypeName,
 							*Prim.GetPrimPath().GetString()
 						);
+
+						return false;
 					}
 
-					return false;
+					bCreated = !bHadRelationship;
 				}
+				else
+				{
+					bool bHadAttr = Prim.HasAttribute(*PropertyName);
+					Attribute = Prim.CreateAttribute(*PropertyName, *AttributeTypeName);
+					if (!Attribute)
+					{
+						// We expect to fail to create an attribute if we have no typename here (e.g. undo remove property)
+						if (AttributeTypeName.IsEmpty())
+						{
+							UE_LOG(
+								LogUsd,
+								Warning,
+								TEXT("Failed to create attribute '%s' with typename '%s' for prim '%s'"),
+								*PropertyName,
+								*AttributeTypeName,
+								*Prim.GetPrimPath().GetString()
+							);
+						}
+
+						return false;
+					}
 
-				bCreated = !bHadAttr;
+					bCreated = !bHadAttr;
+				}
 			}
 		}
 
@@ -486,7 +528,7 @@ namespace UsdUtils
 			UE_LOG(
 				LogUsd,
 				Warning,
-				TEXT("Failed to convert VtValue back to USD when applying it to attribute '%s' of prim '%s'"),
+				TEXT("Failed to convert VtValue back to USD when applying it to property '%s' of prim '%s'"),
 				*PropertyName,
 				*Prim.GetPrimPath().GetString()
 			);
@@ -538,29 +580,89 @@ namespace UsdUtils
 				{
 					Prim.RemoveProperty(*PropertyName);
 				}
-				if (Time.IsSet())
+				if (Time.IsSet() && Attribute)
 				{
 					Attribute.ClearAtTime(Time.GetValue());
 				}
-				else
+				else if (Attribute)
 				{
 					Attribute.Clear();
 				}
 			}
-			else
+			else if (Attribute)
 			{
 				Attribute.Set(WrapperValue, Time);
 			}
 		}
+		// This seems to be the field name for the actual value in pxr:UsdRelationship
+		else if (Field == TEXT("targetPaths"))
+		{
+			if (WrapperValue.IsEmpty())
+			{
+				if (bRemoveProperty)
+				{
+					Prim.RemoveProperty(*PropertyName);
+				}
+				else
+				{
+					if (Attribute)
+					{
+						Attribute.Clear();
+					}
+					else if (Relationship)
+					{
+						bool bRemoveSpec = false;
+						Relationship.ClearTargets(bRemoveSpec);
+					}
+				}
+			}
+			else
+			{
+				// We have to manually convert from the TArray<FString> that our ConvertedValue is holding,
+				// as unlike for UE::FUsdAttribute, we can't just feed a VtValue into the UE::FUsdRelationship
+				if (Value.SourceType == EUsdBasicDataTypes::String && Value.bIsArrayValued)
+				{
+					TArray<UE::FSdfPath> Targets;
+					for (const FConvertedVtValueEntry& Entry : Value.Entries)
+					{
+						// For the relationship values we always put a single component per entry
+						if (Entry.Num() == 1)
+						{
+							const FConvertedVtValueComponent& Component = Entry[0];
+							if (const FString* HeldString = Component.TryGet<FString>())
+							{
+								Targets.Add(UE::FSdfPath{**HeldString});
+							}
+						}
+					}
+
+					Relationship.SetTargets(Targets);
+				}
+			}
+		}
 		else	// variability, colorSpace, etc.
 		{
 			if (WrapperValue.IsEmpty())
 			{
-				Attribute.Clear();
+				if (Attribute)
+				{
+					Attribute.ClearMetadata(*Field);
+				}
+				else if (Relationship)
+				{
+					Relationship.ClearMetadata(*Field);
+				}
 			}
 			else
 			{
-				Attribute.SetMetadata(*Field, WrapperValue);
+				if (Attribute)
+				{
+					Attribute.SetMetadata(*Field, WrapperValue);
+				}
+				else if (Relationship)
+				{
+					Relationship.SetMetadata(*Field, WrapperValue);
+				}
 			}
 		}
 
@@ -722,6 +824,7 @@ namespace UsdUtils
 									AttributeChange.PropertyName,
 									AttributeChange.Field,
 									AttributeChange.AttributeTypeName,
+									AttributeChange.bIsRelationship,
 									bShouldRemove,
 									Direction == EApplicationDirection::Forward ? AttributeChange.NewValue : AttributeChange.OldValue,
 									Prim
diff --git a/Engine/Plugins/Importers/USDImporter/Source/USDUtilities/Private/USDValueConversion.cpp b/Engine/Plugins/Importers/USDImporter/Source/USDUtilities/Private/USDValueConversion.cpp
--- a/Engine/Plugins/Importers/USDImporter/Source/USDUtilities/Private/USDValueConversion.cpp
+++ b/Engine/Plugins/Importers/USDImporter/Source/USDUtilities/Private/USDValueConversion.cpp
@@ -385,6 +385,20 @@ namespace UsdToUnreal
 			OutValue.SourceType = EUsdBasicDataTypes::Int4;
 			ConvertVecValue<int32, GfVec4i>(UsdValue, OutValue);
 		}
+		// This is the type used to store values for relationships. We'll basically treat it like a string array
+		else if (UnderlyingType.IsA<SdfListOp<SdfPath>>())
+		{
+			OutValue.SourceType = EUsdBasicDataTypes::String;
+			OutValue.bIsArrayValued = true;
+
+			const SdfListOp<SdfPath>& Val = UsdValue.UncheckedGet<SdfListOp<SdfPath>>();
+			const std::vector<SdfPath>& ItemVec = Val.GetExplicitItems();
+			OutValue.Entries.Reserve(ItemVec.size());
+			for (const pxr::SdfPath& Item : ItemVec)
+			{
+				OutValue.Entries.Add({UsdUtils::FConvertedVtValueComponent(TInPlaceType<FString>(), UsdToUnreal::ConvertPath(Item))});
+			}
+		}
 		// These types should only appear within metadata, and don't support arrays.
 		// There are more of them (e.g. pxr/usd/usd/crateDataTypes.h), but these are the most common.
 		// Also check pxr/usd/sdf/types.cpp for where these are defined. These are simple enums
@@ -2188,8 +2202,7 @@ namespace UE::USDValueConversion::Private
 				FlattenedValues[FlattenedIndex + 0],
 				FlattenedValues[FlattenedIndex + 1],
 				FlattenedValues[FlattenedIndex + 2],
-				FlattenedValues[FlattenedIndex + 3]
-			};
+				FlattenedValues[FlattenedIndex + 3]};
 		}
 
 		return true;
diff --git a/Engine/Plugins/Importers/USDImporter/Source/UnrealUSDWrapper/Private/UsdWrappers/UsdPrim.cpp b/Engine/Plugins/Importers/USDImporter/Source/UnrealUSDWrapper/Private/UsdWrappers/UsdPrim.cpp
--- a/Engine/Plugins/Importers/USDImporter/Source/UnrealUSDWrapper/Private/UsdWrappers/UsdPrim.cpp
+++ b/Engine/Plugins/Importers/USDImporter/Source/UnrealUSDWrapper/Private/UsdWrappers/UsdPrim.cpp
@@ -8,8 +8,8 @@
 #include "UsdWrappers/UsdAttribute.h"
 #include "UsdWrappers/UsdPayloads.h"
 #include "UsdWrappers/UsdReferences.h"
-#include "UsdWrappers/UsdStage.h"
 #include "UsdWrappers/UsdRelationship.h"
+#include "UsdWrappers/UsdStage.h"
 #include "UsdWrappers/UsdVariantSets.h"
 
 #if USE_USD_SDK
@@ -409,15 +409,6 @@ namespace UE
 #endif	  // #if USE_USD_SDK
 	}
 
-	FUsdRelationship FUsdPrim::GetRelationship(const TCHAR* RelationshipName) const
-	{
-#if USE_USD_SDK
-		return FUsdRelationship(Impl->PxrUsdPrim.Get().GetRelationship(pxr::TfToken(TCHAR_TO_ANSI(RelationshipName))));
-#else
-		return FUsdRelationship();
-#endif	  // #if USE_USD_SDK
-	}
-
 	FName FUsdPrim::GetName() const
 	{
 #if USE_USD_SDK
@@ -668,6 +659,51 @@ namespace UE
 #endif	  // #if USE_USD_SDK
 	}
 
+	FUsdRelationship FUsdPrim::CreateRelationship(const TCHAR* RelName, bool bCustom) const
+	{
+#if USE_USD_SDK
+		FScopedUsdAllocs UsdAllocs;
+
+		return FUsdRelationship(Impl->PxrUsdPrim.Get().CreateRelationship(pxr::TfToken(TCHAR_TO_ANSI(RelName)), bCustom));
+#else
+		return FUsdRelationship{};
+#endif	  // #if USE_USD_SDK
+	}
+
+	FUsdRelationship FUsdPrim::GetRelationship(const TCHAR* RelationshipName) const
+	{
+#if USE_USD_SDK
+		return FUsdRelationship(Impl->PxrUsdPrim.Get().GetRelationship(pxr::TfToken(TCHAR_TO_ANSI(RelationshipName))));
+#else
+		return FUsdRelationship();
+#endif	  // #if USE_USD_SDK
+	}
+
+	TArray<FUsdRelationship> FUsdPrim::GetRelationships() const
+	{
+		TArray<FUsdRelationship> Relationships;
+
+#if USE_USD_SDK
+		FScopedUsdAllocs UsdAllocs;
+
+		for (const pxr::UsdRelationship& Relationship : Impl->PxrUsdPrim.Get().GetRelationships())
+		{
+			Relationships.Emplace(UE::FUsdRelationship{Relationship});
+		}
+#endif	  // #if USE_USD_SDK
+
+		return Relationships;
+	}
+
+	bool FUsdPrim::HasRelationship(const TCHAR* RelName) const
+	{
+#if USE_USD_SDK
+		return Impl->PxrUsdPrim.Get().HasRelationship(pxr::TfToken(TCHAR_TO_ANSI(RelName)));
+#else
+		return false;
+#endif	  // #if USE_USD_SDK
+	}
+
 	bool FUsdPrim::RemoveProperty(FName PropName) const
 	{
 #if USE_USD_SDK
diff --git a/Engine/Plugins/Importers/USDImporter/Source/UnrealUSDWrapper/Private/UsdWrappers/UsdRelationship.cpp b/Engine/Plugins/Importers/USDImporter/Source/UnrealUSDWrapper/Private/UsdWrappers/UsdRelationship.cpp
--- a/Engine/Plugins/Importers/USDImporter/Source/UnrealUSDWrapper/Private/UsdWrappers/UsdRelationship.cpp
+++ b/Engine/Plugins/Importers/USDImporter/Source/UnrealUSDWrapper/Private/UsdWrappers/UsdRelationship.cpp
@@ -1,8 +1,9 @@
 ﻿// Copyright Epic Games, Inc. All Rights Reserved.
 
 #include "UsdWrappers/UsdRelationship.h"
-#include "UsdWrappers/SdfPath.h"
 #include "USDMemory.h"
+#include "UsdWrappers/SdfPath.h"
+#include "UsdWrappers/VtValue.h"
 
 #if USE_USD_SDK
 #include "USDIncludesStart.h"
@@ -29,7 +30,7 @@ namespace UE
 				: PxrUsdRelationship(MoveTemp(InUsdRelationship))
 			{
 			}
-			
+
 			TUsdStore<pxr::UsdRelationship> PxrUsdRelationship;
 #endif	  // #if USE_USD_SDK
 		};
@@ -147,27 +148,87 @@ namespace UE
 		return Impl->PxrUsdRelationship.Get();
 	}
 #endif
-	
-	bool FUsdRelationship::GetTargets(TArray<UE::FSdfPath>& TargetsPath) const
+
+	bool FUsdRelationship::GetMetadata(const TCHAR* Key, UE::FVtValue& Value) const
+	{
+#if USE_USD_SDK
+		return Impl->PxrUsdRelationship.Get().GetMetadata(pxr::TfToken{TCHAR_TO_ANSI(Key)}, &Value.GetUsdValue());
+#else
+		return false;
+#endif	  // #if USE_USD_SDK
+	}
+
+	bool FUsdRelationship::HasMetadata(const TCHAR* Key) const
+	{
+#if USE_USD_SDK
+		return Impl->PxrUsdRelationship.Get().HasMetadata(pxr::TfToken{TCHAR_TO_ANSI(Key)});
+#else
+		return false;
+#endif	  // #if USE_USD_SDK
+	}
+
+	bool FUsdRelationship::SetMetadata(const TCHAR* Key, const UE::FVtValue& Value) const
 	{
-		TargetsPath.Reset();
+#if USE_USD_SDK
+		return Impl->PxrUsdRelationship.Get().SetMetadata(pxr::TfToken{TCHAR_TO_ANSI(Key)}, Value.GetUsdValue());
+#else
+		return false;
+#endif	  // #if USE_USD_SDK
+	}
+
+	bool FUsdRelationship::ClearMetadata(const TCHAR* Key) const
+	{
+#if USE_USD_SDK
+		return Impl->PxrUsdRelationship.Get().ClearMetadata(pxr::TfToken{TCHAR_TO_ANSI(Key)});
+#else
+		return false;
+#endif	  // #if USE_USD_SDK
+	}
+
+	bool FUsdRelationship::SetTargets(const TArray<UE::FSdfPath>& Targets) const
+	{
+#if USE_USD_SDK
+		std::vector<pxr::SdfPath> UsdTargets;
+		UsdTargets.reserve(Targets.Num());
+		for (const UE::FSdfPath& UETarget : Targets)
+		{
+			UsdTargets.push_back(pxr::SdfPath{UETarget});
+		}
+
+		return Impl->PxrUsdRelationship.Get().SetTargets(UsdTargets);
+#else
+		return false;
+#endif	  // USE_USD_SDK
+	}
+
+	bool FUsdRelationship::ClearTargets(bool bRemoveSpec) const
+	{
+#if USE_USD_SDK
+		return Impl->PxrUsdRelationship.Get().ClearTargets(bRemoveSpec);
+#else
+		return false;
+#endif	  // USE_USD_SDK
+	}
+
+	bool FUsdRelationship::GetTargets(TArray<UE::FSdfPath>& Targets) const
+	{
+		Targets.Reset();
 #if USE_USD_SDK
 		FScopedUsdAllocs Allocs;
 
-		pxr::SdfPathVector Targets;
-		if(Impl->PxrUsdRelationship.Get().GetTargets(&Targets))
+		pxr::SdfPathVector UsdTargets;
+		if (Impl->PxrUsdRelationship.Get().GetTargets(&UsdTargets))
 		{
-			const uint32 NumTargets = Targets.size();
-			for(uint32 TargetIndex = 0; TargetIndex < NumTargets; ++TargetIndex)
+			const uint32 NumTargets = UsdTargets.size();
+			for (uint32 TargetIndex = 0; TargetIndex < NumTargets; ++TargetIndex)
 			{
-				TargetsPath.Add(UE::FSdfPath(Targets[TargetIndex]));
+				Targets.Add(UE::FSdfPath(UsdTargets[TargetIndex]));
 			}
 			return true;
 		}
-		
+
 #endif	  // USE_USD_SDK
 		return false;
 	}
-	
-	
+
 }	 // namespace UE
diff --git a/Engine/Plugins/Importers/USDImporter/Source/UnrealUSDWrapper/Public/UsdWrappers/UsdPrim.h b/Engine/Plugins/Importers/USDImporter/Source/UnrealUSDWrapper/Public/UsdWrappers/UsdPrim.h
--- a/Engine/Plugins/Importers/USDImporter/Source/UnrealUSDWrapper/Public/UsdWrappers/UsdPrim.h
+++ b/Engine/Plugins/Importers/USDImporter/Source/UnrealUSDWrapper/Public/UsdWrappers/UsdPrim.h
@@ -29,7 +29,6 @@ namespace UE
 	class FUsdVariantSet;
 	class FUsdVariantSets;
 	class FUsdRelationship;
-	
 
 	namespace Internal
 	{
@@ -112,7 +111,6 @@ namespace UE
 
 		const FSdfPath GetPrimPath() const;
 		FUsdStage GetStage() const;
-		FUsdRelationship GetRelationship(const TCHAR* RelationshipName) const;
 
 		FName GetName() const;
 
@@ -148,6 +146,11 @@ namespace UE
 		FUsdAttribute GetAttribute(const TCHAR* AttrName) const;
 		bool HasAttribute(const TCHAR* AttrName) const;
 
+		FUsdRelationship CreateRelationship(const TCHAR* RelName, bool bCustom = true) const;
+		TArray<FUsdRelationship> GetRelationships() const;
+		FUsdRelationship GetRelationship(const TCHAR* RelationshipName) const;
+		bool HasRelationship(const TCHAR* RelName) const;
+
 		bool IsInstanceable() const;
 		bool SetInstanceable(bool bInstanceable) const;
 		bool ClearInstanceable() const;
diff --git a/Engine/Plugins/Importers/USDImporter/Source/UnrealUSDWrapper/Public/UsdWrappers/UsdRelationship.h b/Engine/Plugins/Importers/USDImporter/Source/UnrealUSDWrapper/Public/UsdWrappers/UsdRelationship.h
--- a/Engine/Plugins/Importers/USDImporter/Source/UnrealUSDWrapper/Public/UsdWrappers/UsdRelationship.h
+++ b/Engine/Plugins/Importers/USDImporter/Source/UnrealUSDWrapper/Public/UsdWrappers/UsdRelationship.h
@@ -23,6 +23,7 @@ namespace UE
 		class FUsdRelationshipImpl;
 	}
 	class FSdfPath;
+	class FVtValue;
 
 	/**
 	 * Minimal pxr::UsdRelationship wrapper for Unreal that can be used from no-rtti modules.
@@ -61,11 +62,20 @@ namespace UE
 		operator pxr::UsdProperty&();
 		operator const pxr::UsdProperty&() const;
 #endif	  // #if USE_USD_SDK
-		
+
+		  // Wrapped pxr::UsdObject functions, refer to the USD SDK documentation
+	public:
+		bool GetMetadata(const TCHAR* Key, UE::FVtValue& Value) const;
+		bool HasMetadata(const TCHAR* Key) const;
+		bool SetMetadata(const TCHAR* Key, const UE::FVtValue& Value) const;
+		bool ClearMetadata(const TCHAR* Key) const;
+
+		// Wrapped pxr::UsdRelationship functions, refer to the USD SDK documentation
 	public:
-		/** Get the relationship targets (primitive path relative to the stage) */
-		bool GetTargets(TArray<UE::FSdfPath>& TargetsPath) const;
-		
+		bool SetTargets(const TArray<UE::FSdfPath>& Targets) const;
+		bool ClearTargets(bool bRemoveSpec) const;
+		bool GetTargets(TArray<UE::FSdfPath>& Targets) const;
+
 	private:
 		TUniquePtr<Internal::FUsdRelationshipImpl> Impl;
 	};
