USD: Add some more analytics for relevant schema and schema translator counts.

diff --git a/Engine/Plugins/Importers/USDImporter/Source/USDSchemas/Private/USDSchemaTranslator.cpp b/Engine/Plugins/Importers/USDImporter/Source/USDSchemas/Private/USDSchemaTranslator.cpp
--- a/Engine/Plugins/Importers/USDImporter/Source/USDSchemas/Private/USDSchemaTranslator.cpp
+++ b/Engine/Plugins/Importers/USDImporter/Source/USDSchemas/Private/USDSchemaTranslator.cpp
@@ -106,6 +106,8 @@ FRegisteredSchemaTranslatorHandle FUsdSchemaTranslatorRegistry::Register(const F
 
 	SchemaTranslatorsStack->Push(RegisteredSchemaTranslator);
 
+	ExternalSchemaTranslatorCount++;
+
 	return RegisteredSchemaTranslator.Handle;
 #else
 	return FRegisteredSchemaTranslatorHandle();
@@ -132,6 +134,11 @@ void FUsdSchemaTranslatorRegistry::Unregister(const FRegisteredSchemaTranslatorH
 	}
 }
 
+int32 FUsdSchemaTranslatorRegistry::GetExternalSchemaTranslatorCount()
+{
+	return ExternalSchemaTranslatorCount;
+}
+
 FUsdSchemaTranslationContext::FUsdSchemaTranslationContext(const UE::FUsdStage& InStage, UUsdAssetCache2& InAssetCache)
 	: Stage(InStage)
 	, AssetCache(&InAssetCache)
@@ -160,6 +167,11 @@ FUsdSchemaTranslatorRegistry::FSchemaTranslatorsStack* FUsdSchemaTranslatorRegis
 	}
 }
 
+void FUsdSchemaTranslatorRegistry::ResetExternalTranslatorCount()
+{
+	ExternalSchemaTranslatorCount = 0;
+}
+
 void FUsdSchemaTranslationContext::CompleteTasks()
 {
 	TRACE_CPUPROFILER_EVENT_SCOPE(FUsdSchemaTranslationContext::CompleteTasks);
diff --git a/Engine/Plugins/Importers/USDImporter/Source/USDSchemas/Private/USDSchemasModule.cpp b/Engine/Plugins/Importers/USDImporter/Source/USDSchemas/Private/USDSchemasModule.cpp
--- a/Engine/Plugins/Importers/USDImporter/Source/USDSchemas/Private/USDSchemasModule.cpp
+++ b/Engine/Plugins/Importers/USDImporter/Source/USDSchemas/Private/USDSchemasModule.cpp
@@ -3,11 +3,13 @@
 #include "USDSchemasModule.h"
 
 #include "USDSchemaTranslator.h"
+#include "USDTypesConversion.h"
 
 #if WITH_EDITOR
 #include "MDLImporterModule.h"
 #endif	  // #if WITH_EDITOR
 
+#include "AnalyticsEventAttribute.h"
 #include "Modules/ModuleManager.h"
 
 #if USE_USD_SDK
@@ -25,6 +27,11 @@
 #include "USDShadeMaterialTranslator.h"
 #include "USDSkelSkeletonTranslator.h"
 #include "USDVolVolumeTranslator.h"
+
+#include "USDIncludesStart.h"
+#include "pxr/usd/kind/registry.h"
+#include "pxr/usd/usd/primRange.h"
+#include "USDIncludesEnd.h"
 #endif	  // #if USE_USD_SDK
 
 class FUsdSchemasModule : public IUsdSchemasModule
@@ -79,6 +86,8 @@ public:
 		}
 #endif	  // WITH_EDITOR
 
+		Registry.ResetExternalTranslatorCount();
+
 #endif	  // #if USE_USD_SDK
 	}
 
@@ -115,4 +124,276 @@ protected:
 	TArray<FRegisteredSchemaTranslatorHandle> TranslatorHandles;
 };
 
+void UsdUnreal::Analytics::CollectSchemaAnalytics(const UE::FUsdStage& Stage, const FString& EventName)
+{
+#if USE_USD_SDK
+	if (!Stage)
+	{
+		return;
+	}
+
+	// Last collected at cl 32482283
+	const static TSet<FString> NativeSchemaNames{
+		"AssetPreviewsAPI",
+		"Backdrop",
+		"BasisCurves",
+		"BlendShape",
+		"Camera",
+		"Capsule",
+		"Capsule_1",
+		"ClipsAPI",
+		"CollectionAPI",
+		"Cone",
+		"ConnectableAPI",
+		"ControlRigAPI",
+		"CoordSysAPI",
+		"Cube",
+		"Cylinder",
+		"CylinderLight",
+		"Cylinder_1",
+		"DiskLight",
+		"DistantLight",
+		"DomeLight",
+		"DomeLight_1",
+		"Field3DAsset",
+		"GenerativeProcedural",
+		"GeomModelAPI",
+		"GeomSubset",
+		"GeometryLight",
+		"GroomAPI",
+		"GroomBindingAPI",
+		"HermiteCurves",
+		"HydraGenerativeProceduralAPI",
+		"LightAPI",
+		"LightFilter",
+		"LightListAPI",
+		"ListAPI",
+		"LiveLinkAPI",
+		"Material",
+		"MaterialBindingAPI",
+		"Mesh",
+		"MeshLightAPI",
+		"ModelAPI",
+		"MotionAPI",
+		"NodeDefAPI",
+		"NodeGraph",
+		"NodeGraphNodeAPI",
+		"NurbsCurves",
+		"NurbsPatch",
+		"OpenVDBAsset",
+		"PhysicsArticulationRootAPI",
+		"PhysicsCollisionAPI",
+		"PhysicsCollisionGroup",
+		"PhysicsDistanceJoint",
+		"PhysicsDriveAPI",
+		"PhysicsFilteredPairsAPI",
+		"PhysicsFixedJoint",
+		"PhysicsJoint",
+		"PhysicsLimitAPI",
+		"PhysicsMassAPI",
+		"PhysicsMaterialAPI",
+		"PhysicsMeshCollisionAPI",
+		"PhysicsPrismaticJoint",
+		"PhysicsRevoluteJoint",
+		"PhysicsRigidBodyAPI",
+		"PhysicsScene",
+		"PhysicsSphericalJoint",
+		"Plane",
+		"PluginLight",
+		"PluginLightFilter",
+		"PointInstancer",
+		"Points",
+		"PortalLight",
+		"PrimvarsAPI",
+		"RectLight",
+		"RenderDenoisePass",
+		"RenderPass",
+		"RenderProduct",
+		"RenderSettings",
+		"RenderVar",
+		"RiMaterialAPI",
+		"RiRenderPassAPI",
+		"RiSplineAPI",
+		"SceneGraphPrimAPI",
+		"Scope",
+		"Shader",
+		"ShadowAPI",
+		"ShapingAPI",
+		"SkelAnimation",
+		"SkelBindingAPI",
+		"SkelRoot",
+		"Skeleton",
+		"SparseVolumeTextureAPI",
+		"SpatialAudio",
+		"Sphere",
+		"SphereLight",
+		"StatementsAPI",
+		"TetMesh",
+		"VisibilityAPI",
+		"Volume",
+		"VolumeLightAPI",
+		"Xform",
+		"XformCommonAPI"};
+
+	// clang-format off
+	// "Interesting" here means either something we don't support, or something that would be interesting
+	// to check the usage of, like for our Groom/ControlRig/LiveLink support, or whether the new support
+	// for OpenVDB is used, etc.
+	TMap<FString, int32> InterestingNativeSchemaCounts{
+		{"AssetPreviewsAPI", 0},
+		{"Backdrop", 0},
+		{"BasisCurves", 0},
+		{"Capsule_1", 0},
+		{"ClipsAPI", 0},
+		{"CollectionAPI", 0},
+		{"ControlRigAPI", 0},
+		{"CoordSysAPI", 0},
+		{"CylinderLight", 0},
+		{"Cylinder_1", 0},
+		{"DomeLight_1", 0},
+		{"Field3DAsset", 0},
+		{"GenerativeProcedural", 0},
+		{"GeometryLight", 0},
+		{"GroomAPI", 0},
+		{"GroomBindingAPI", 0},
+		{"HermiteCurves", 0},
+		{"HydraGenerativeProceduralAPI", 0},
+		{"LightFilter", 0},
+		{"LightListAPI", 0},
+		{"ListAPI", 0},
+		{"LiveLinkAPI", 0},
+		{"MeshLightAPI", 0},
+		{"MotionAPI", 0},
+		{"NurbsCurves", 0},
+		{"NurbsPatch", 0},
+		{"OpenVDBAsset", 0},
+		{"PhysicsArticulationRootAPI", 0},
+		{"PhysicsCollisionAPI", 0},
+		{"PhysicsCollisionGroup", 0},
+		{"PhysicsDistanceJoint", 0},
+		{"PhysicsDriveAPI", 0},
+		{"PhysicsFilteredPairsAPI", 0},
+		{"PhysicsFixedJoint", 0},
+		{"PhysicsJoint", 0},
+		{"PhysicsLimitAPI", 0},
+		{"PhysicsMassAPI", 0},
+		{"PhysicsMaterialAPI", 0},
+		{"PhysicsMeshCollisionAPI", 0},
+		{"PhysicsPrismaticJoint", 0},
+		{"PhysicsRevoluteJoint", 0},
+		{"PhysicsRigidBodyAPI", 0},
+		{"PhysicsScene", 0},
+		{"PhysicsSphericalJoint", 0},
+		{"PluginLight", 0},
+		{"PluginLightFilter", 0},
+		{"Points", 0},
+		{"PortalLight", 0},
+		{"RenderDenoisePass", 0},
+		{"RenderPass", 0},
+		{"RenderProduct", 0},
+		{"RenderSettings", 0},
+		{"RenderVar", 0},
+		{"RiMaterialAPI", 0},
+		{"RiRenderPassAPI", 0},
+		{"RiSplineAPI", 0},
+		{"SceneGraphPrimAPI", 0},
+		{"ShadowAPI", 0},
+		{"SparseVolumeTextureAPI", 0},
+		{"SpatialAudio", 0},
+		{"SphereLight", 0},
+		{"StatementsAPI", 0},
+		{"TetMesh", 0},
+		{"VisibilityAPI", 0},
+		{"Volume", 0},
+		{"VolumeLightAPI", 0}
+	};
+	// clang-format on
+
+	TSet<FString> SeenSchemas;
+
+	int32 CustomSchemaCount = 0;
+	{
+		FScopedUsdAllocs Allocs;
+
+		pxr::UsdPrimRange PrimRange{Stage.GetPseudoRoot(), pxr::UsdTraverseInstanceProxies()};
+		for (pxr::UsdPrimRange::iterator PrimRangeIt = ++PrimRange.begin(); PrimRangeIt != PrimRange.end(); ++PrimRangeIt)
+		{
+			// It's perfectly fine to have a typeless prim (e.g. "def 'Cube'")
+			if (PrimRangeIt->HasAuthoredTypeName())
+			{
+				const pxr::TfToken& TypeName = PrimRangeIt->GetTypeName();
+				const FString TypeNameStr = UsdToUnreal::ConvertToken(TypeName);
+
+				if (!TypeNameStr.IsEmpty())
+				{
+					SeenSchemas.Add(TypeNameStr);
+
+					if (int32* InterestingSchemaCount = InterestingNativeSchemaCounts.Find(TypeNameStr))
+					{
+						*InterestingSchemaCount = *InterestingSchemaCount + 1;
+					}
+				}
+			}
+
+			const pxr::UsdPrimTypeInfo& PrimTypeInfo = PrimRangeIt->GetPrimTypeInfo();
+			for (const pxr::TfToken& AppliedSchema : PrimTypeInfo.GetAppliedAPISchemas())
+			{
+				std::pair<pxr::TfToken, pxr::TfToken> Pair = pxr::UsdSchemaRegistry::GetTypeNameAndInstance(AppliedSchema);
+				FString TypeName = UsdToUnreal::ConvertToken(Pair.first);
+
+				// These applied schema names shouldn't ever end up as the empty string... but we don't really want to pop
+				// an ensure or show a warning when analytics fails
+				if (!TypeName.IsEmpty())
+				{
+					SeenSchemas.Add(TypeName);
+
+					if (int32* InterestingSchemaCount = InterestingNativeSchemaCounts.Find(TypeName))
+					{
+						*InterestingSchemaCount = *InterestingSchemaCount + 1;
+					}
+				}
+			}
+		}
+	}
+
+	for (const FString& SchemaName : SeenSchemas)
+	{
+		// We only care about non-native schemas
+		if (!NativeSchemaNames.Contains(SchemaName))
+		{
+			++CustomSchemaCount;
+		}
+	}
+
+	TArray<FAnalyticsEventAttribute> EventAttributes;
+	if (CustomSchemaCount > 0)
+	{
+		EventAttributes.Emplace(TEXT("CustomSchemas"), CustomSchemaCount);
+	}
+
+	for (const TPair<FString, int32>& Pair : InterestingNativeSchemaCounts)
+	{
+		const FString& InterestingNativeSchemaName = Pair.Key;
+		int32 Count = Pair.Value;
+		if (Count > 0)
+		{
+			EventAttributes.Emplace(InterestingNativeSchemaName, Count);
+		}
+	}
+
+	IUsdSchemasModule& UsdSchemasModule = FModuleManager::Get().LoadModuleChecked<IUsdSchemasModule>(TEXT("USDSchemas"));
+	FUsdSchemaTranslatorRegistry& Registry = UsdSchemasModule.GetTranslatorRegistry();
+	int32 SchemaTranslatorCount = Registry.GetExternalSchemaTranslatorCount();
+	if (SchemaTranslatorCount > 0)
+	{
+		EventAttributes.Emplace(TEXT("CustomSchemaTranslatorCount"), SchemaTranslatorCount);
+	}
+
+	if (EventAttributes.Num() > 0)
+	{
+		IUsdClassesModule::SendAnalytics(MoveTemp(EventAttributes), FString::Printf(TEXT("%s.CustomSchemaCount"), *EventName));
+	}
+#endif	  // USE_USD_SDK
+}
+
 IMPLEMENT_MODULE_USD(FUsdSchemasModule, USDSchemas);
diff --git a/Engine/Plugins/Importers/USDImporter/Source/USDSchemas/Public/USDSchemaTranslator.h b/Engine/Plugins/Importers/USDImporter/Source/USDSchemas/Public/USDSchemaTranslator.h
--- a/Engine/Plugins/Importers/USDImporter/Source/USDSchemas/Public/USDSchemaTranslator.h
+++ b/Engine/Plugins/Importers/USDImporter/Source/USDSchemas/Public/USDSchemaTranslator.h
@@ -111,12 +111,20 @@ public:
 
 	void Unregister(const FRegisteredSchemaTranslatorHandle& TranslatorHandle);
 
+	int32 GetExternalSchemaTranslatorCount();
+
 protected:
 	FRegisteredSchemaTranslatorHandle Register(const FString& SchemaName, FCreateTranslator CreateFunction);
 
 	FSchemaTranslatorsStack* FindSchemaTranslatorStack(const FString& SchemaName);
 
 	TArray<TPair<FString, FSchemaTranslatorsStack>> RegisteredSchemaTranslators;
+
+private:
+	// Small machinery that lets us collect basic analytics about how many custom schema translators are being used in this session
+	friend class FUsdSchemasModule;
+	void ResetExternalTranslatorCount();
+	int32 ExternalSchemaTranslatorCount = 0;
 };
 
 class FRegisteredSchemaTranslator
diff --git a/Engine/Plugins/Importers/USDImporter/Source/USDSchemas/Public/USDSchemasModule.h b/Engine/Plugins/Importers/USDImporter/Source/USDSchemas/Public/USDSchemasModule.h
--- a/Engine/Plugins/Importers/USDImporter/Source/USDSchemas/Public/USDSchemasModule.h
+++ b/Engine/Plugins/Importers/USDImporter/Source/USDSchemas/Public/USDSchemasModule.h
@@ -2,6 +2,9 @@
 
 #pragma once
 
+#include "UsdWrappers/ForwardDeclarations.h"
+
+#include "CoreMinimal.h"
 #include "Modules/ModuleInterface.h"
 
 class FUsdRenderContextRegistry;
@@ -13,3 +16,9 @@ public:
 	virtual FUsdSchemaTranslatorRegistry& GetTranslatorRegistry() = 0;
 	virtual FUsdRenderContextRegistry& GetRenderContextRegistry() = 0;
 };
+
+namespace UsdUnreal::Analytics
+{
+	/** Collects analytics about custom schemas, unsupported native schemas, and the count of custom registered schema translators */
+	USDSCHEMAS_API void CollectSchemaAnalytics(const UE::FUsdStage& Stage, const FString& EventName);
+}
diff --git a/Engine/Plugins/Importers/USDImporter/Source/USDStage/Private/USDStageActor.cpp b/Engine/Plugins/Importers/USDImporter/Source/USDStage/Private/USDStageActor.cpp
--- a/Engine/Plugins/Importers/USDImporter/Source/USDStage/Private/USDStageActor.cpp
+++ b/Engine/Plugins/Importers/USDImporter/Source/USDStage/Private/USDStageActor.cpp
@@ -2505,7 +2505,7 @@ void AUsdStageActor::SetRootLayer(const FString& RootFilePath)
 
 	// Do this here instead of on OpenUsdStage/LoadUsdStage as those also get called when changing any of
 	// our properties, like render context, material purpose, etc.
-	UsdUtils::CollectSchemaAnalytics(UsdStage, TEXT("Open"));
+	UsdUnreal::Analytics::CollectSchemaAnalytics(UsdStage, TEXT("Open"));
 }
 
 void AUsdStageActor::SetStageState(EUsdStageState NewState)
diff --git a/Engine/Plugins/Importers/USDImporter/Source/USDStageImporter/Private/USDStageImporter.cpp b/Engine/Plugins/Importers/USDImporter/Source/USDStageImporter/Private/USDStageImporter.cpp
--- a/Engine/Plugins/Importers/USDImporter/Source/USDStageImporter/Private/USDStageImporter.cpp
+++ b/Engine/Plugins/Importers/USDImporter/Source/USDStageImporter/Private/USDStageImporter.cpp
@@ -2094,7 +2094,7 @@ void UUsdStageImporter::ImportFromFile(FUsdStageImportContext& ImportContext)
 		UE::USDStageImporter::Private::SendAnalytics(ImportContext, nullptr, TEXT("Import"), PublishedAssetsAndDependencies, ElapsedSeconds);
 		UE_LOG(LogUsd, Log, TEXT("Imported '%s' in %.3f seconds."), *ImportContext.FilePath, ElapsedSeconds);
 
-		UsdUtils::CollectSchemaAnalytics(ImportContext.Stage, TEXT("Import"));
+		UsdUnreal::Analytics::CollectSchemaAnalytics(ImportContext.Stage, TEXT("Import"));
 	}
 
 	UE::USDStageImporter::Private::CloseStageIfNeeded(ImportContext);
@@ -2313,7 +2313,7 @@ bool UUsdStageImporter::ReimportSingleAsset(
 		UE::USDStageImporter::Private::SendAnalytics(ImportContext, ReimportedObject, TEXT("Reimport"), {ReimportedObject}, ElapsedSeconds);
 		UE_LOG(LogUsd, Log, TEXT("Re-imported '%s' in %.3f seconds."), *ImportContext.FilePath, ElapsedSeconds);
 
-		UsdUtils::CollectSchemaAnalytics(ImportContext.Stage, TEXT("Reimport"));
+		UsdUnreal::Analytics::CollectSchemaAnalytics(ImportContext.Stage, TEXT("Reimport"));
 	}
 
 	UE::USDStageImporter::Private::CloseStageIfNeeded(ImportContext);
diff --git a/Engine/Plugins/Importers/USDImporter/Source/USDUtilities/Private/USDConversionUtils.cpp b/Engine/Plugins/Importers/USDImporter/Source/USDUtilities/Private/USDConversionUtils.cpp
--- a/Engine/Plugins/Importers/USDImporter/Source/USDUtilities/Private/USDConversionUtils.cpp
+++ b/Engine/Plugins/Importers/USDImporter/Source/USDUtilities/Private/USDConversionUtils.cpp
@@ -20,7 +20,6 @@
 #include "UsdWrappers/UsdPrim.h"
 #include "UsdWrappers/UsdStage.h"
 
-#include "AnalyticsEventAttribute.h"
 #include "Animation/AnimBlueprint.h"
 #include "Animation/AnimSequence.h"
 #include "Animation/Skeleton.h"
@@ -3408,164 +3407,9 @@ bool UsdUtils::ClearNonEssentialPrimMetadata(const pxr::UsdPrim& Prim)
 }
 #endif	  // USE_USD_SDK
 
+// Deprecated
 void UsdUtils::CollectSchemaAnalytics(const UE::FUsdStage& Stage, const FString& EventName)
 {
-#if USE_USD_SDK
-	if (!Stage)
-	{
-		return;
-	}
-
-	TSet<FString> SeenSchemas;
-
-	{
-		FScopedUsdAllocs Allocs;
-
-		pxr::UsdPrimRange PrimRange{Stage.GetPseudoRoot(), pxr::UsdTraverseInstanceProxies()};
-		for (pxr::UsdPrimRange::iterator PrimRangeIt = ++PrimRange.begin(); PrimRangeIt != PrimRange.end(); ++PrimRangeIt)
-		{
-			// It's perfectly fine to have a typeless prim (e.g. "def 'Cube'")
-			if (PrimRangeIt->HasAuthoredTypeName())
-			{
-				const pxr::TfToken& TypeName = PrimRangeIt->GetTypeName();
-				const FString TypeNameStr = UsdToUnreal::ConvertToken(TypeName);
-
-				if (!TypeNameStr.IsEmpty())
-				{
-					SeenSchemas.Add(TypeNameStr);
-				}
-			}
-
-			const pxr::UsdPrimTypeInfo& PrimTypeInfo = PrimRangeIt->GetPrimTypeInfo();
-			for (const pxr::TfToken& AppliedSchema : PrimTypeInfo.GetAppliedAPISchemas())
-			{
-				std::pair<pxr::TfToken, pxr::TfToken> Pair = pxr::UsdSchemaRegistry::GetTypeNameAndInstance(AppliedSchema);
-				FString TypeName = UsdToUnreal::ConvertToken(Pair.first);
-
-				// These applied schema names shouldn't ever end up as the empty string... but we don't really want to pop
-				// an ensure or show a warning when analytics fails
-				if (!TypeName.IsEmpty())
-				{
-					SeenSchemas.Add(TypeName);
-				}
-			}
-		}
-	}
-
-	const static TSet<FString> NativeSchemaNames{
-		"AssetPreviewsAPI",
-		"Backdrop",
-		"BasisCurves",
-		"BlendShape",
-		"Camera",
-		"Capsule",
-		"ClipsAPI",
-		"CollectionAPI",
-		"Cone",
-		"ConnectableAPI",
-		"ControlRigAPI",
-		"CoordSysAPI",
-		"Cube",
-		"Cylinder",
-		"CylinderLight",
-		"DiskLight",
-		"DistantLight",
-		"DomeLight",
-		"Field3DAsset",
-		"GenerativeProcedural",
-		"GeomModelAPI",
-		"GeomSubset",
-		"GeometryLight",
-		"GroomAPI",
-		"GroomBindingAPI",
-		"HermiteCurves",
-		"HydraGenerativeProceduralAPI",
-		"LightAPI",
-		"LightFilter",
-		"LightListAPI",
-		"ListAPI",
-		"LiveLinkAPI",
-		"Material",
-		"MaterialBindingAPI",
-		"Mesh",
-		"MeshLightAPI",
-		"ModelAPI",
-		"MotionAPI",
-		"NodeDefAPI",
-		"NodeGraph",
-		"NodeGraphNodeAPI",
-		"NurbsCurves",
-		"NurbsPatch",
-		"OpenVDBAsset",
-		"PackedJointAnimation",
-		"PhysicsArticulationRootAPI",
-		"PhysicsCollisionAPI",
-		"PhysicsCollisionGroup",
-		"PhysicsDistanceJoint",
-		"PhysicsDriveAPI",
-		"PhysicsFilteredPairsAPI",
-		"PhysicsFixedJoint",
-		"PhysicsJoint",
-		"PhysicsLimitAPI",
-		"PhysicsMassAPI",
-		"PhysicsMaterialAPI",
-		"PhysicsMeshCollisionAPI",
-		"PhysicsPrismaticJoint",
-		"PhysicsRevoluteJoint",
-		"PhysicsRigidBodyAPI",
-		"PhysicsScene",
-		"PhysicsSphericalJoint",
-		"Plane",
-		"PluginLight",
-		"PluginLightFilter",
-		"PointInstancer",
-		"Points",
-		"PortalLight",
-		"PrimvarsAPI",
-		"RectLight",
-		"RenderDenoisePass",
-		"RenderPass",
-		"RenderProduct",
-		"RenderSettings",
-		"RenderVar",
-		"RiMaterialAPI",
-		"RiSplineAPI",
-		"SceneGraphPrimAPI",
-		"Scope",
-		"Shader",
-		"ShadowAPI",
-		"ShapingAPI",
-		"SkelAnimation",
-		"SkelBindingAPI",
-		"SkelRoot",
-		"Skeleton",
-		"SpatialAudio",
-		"Sphere",
-		"SphereLight",
-		"StatementsAPI",
-		"VisibilityAPI",
-		"Volume",
-		"VolumeLightAPI",
-		"Xform",
-		"XformCommonAPI"};
-
-	int32 CustomSchemaCount = 0;
-	for (const FString& SchemaName : SeenSchemas)
-	{
-		// We only care about non-native schemas
-		if (!NativeSchemaNames.Contains(SchemaName))
-		{
-			++CustomSchemaCount;
-		}
-	}
-
-	if (CustomSchemaCount > 0)
-	{
-		TArray<FAnalyticsEventAttribute> EventAttributes;
-		EventAttributes.Emplace(TEXT("CustomSchemas"), CustomSchemaCount);
-		IUsdClassesModule::SendAnalytics(MoveTemp(EventAttributes), FString::Printf(TEXT("%s.CustomSchemaCount"), *EventName));
-	}
-#endif	  // USE_USD_SDK
 }
 
 #if USE_USD_SDK
diff --git a/Engine/Plugins/Importers/USDImporter/Source/USDUtilities/Public/USDConversionUtils.h b/Engine/Plugins/Importers/USDImporter/Source/USDUtilities/Public/USDConversionUtils.h
--- a/Engine/Plugins/Importers/USDImporter/Source/USDUtilities/Public/USDConversionUtils.h
+++ b/Engine/Plugins/Importers/USDImporter/Source/USDUtilities/Public/USDConversionUtils.h
@@ -569,5 +569,6 @@ namespace UsdUtils
 #endif	  // USE_USD_SDK
 
 	/** Collects how many times each schema shows up on the provided stage and send it as an analytics event */
+	UE_DEPRECATED(5.5, "This function has been moved to the USDSchemasModule")
 	USDUTILITIES_API void CollectSchemaAnalytics(const UE::FUsdStage& Stage, const FString& EventName);
 }	 // namespace UsdUtils
