USD: Better workaround for USD weakPtr and refPtr operator== causing a hang/crash;

Undo cls 32723057 and 32722819 as they are no longer needed.

diff --git a/Engine/Plugins/Importers/USDImporter/Source/USDStageEditorViewModels/Private/USDLayersViewModel.cpp b/Engine/Plugins/Importers/USDImporter/Source/USDStageEditorViewModels/Private/USDLayersViewModel.cpp
--- a/Engine/Plugins/Importers/USDImporter/Source/USDStageEditorViewModels/Private/USDLayersViewModel.cpp
+++ b/Engine/Plugins/Importers/USDImporter/Source/USDStageEditorViewModels/Private/USDLayersViewModel.cpp
@@ -164,18 +164,7 @@ void FUsdLayerViewModel::RefreshData()
 	const pxr::SdfLayerHandle& EditTargetLayer = UsdStageRef->GetEditTarget().GetLayer();
 	const pxr::SdfLayerHandle& FocusedEditTarget = IsolatedStageRef ? IsolatedStageRef->GetEditTarget().GetLayer() : pxr::SdfLayerHandle{};
 
-	if (UsdLayer)
-	{
-		// For whatever reason, doing the "EditTargetLayer == UsdLayer" comparison direcly leads to a hang on Windows and a crash
-		// on the Mac? It could be related to the USD 24.03 update somehow
-		LayerModel->bIsEditTarget = (IsolatedStageRef) ? LayerModel->bIsInIsolatedStage
-															 && (FocusedEditTarget && FocusedEditTarget->GetIdentifier() == UsdLayer->GetIdentifier())
-													   : (EditTargetLayer && EditTargetLayer->GetIdentifier() == UsdLayer->GetIdentifier());
-	}
-	else
-	{
-		LayerModel->bIsEditTarget = false;
-	}
+    LayerModel->bIsEditTarget = (IsolatedStageRef) ? LayerModel->bIsInIsolatedStage && FocusedEditTarget == UsdLayer : EditTargetLayer == UsdLayer;
 
 	LayerModel->bIsMuted = (IsolatedStageRef && LayerModel->bIsInIsolatedStage)
 							   // If we isolating, we're only muted if we're muted on the isolated stage
diff --git a/Engine/Plugins/Importers/USDImporter/Source/UnrealUSDWrapper/Private/UnrealUSDWrapper.cpp b/Engine/Plugins/Importers/USDImporter/Source/UnrealUSDWrapper/Private/UnrealUSDWrapper.cpp
--- a/Engine/Plugins/Importers/USDImporter/Source/UnrealUSDWrapper/Private/UnrealUSDWrapper.cpp
+++ b/Engine/Plugins/Importers/USDImporter/Source/UnrealUSDWrapper/Private/UnrealUSDWrapper.cpp
@@ -1092,6 +1092,25 @@ void UnrealUSDWrapper::ClearDiagnosticDelegate()
 #endif	  // USE_USD_SDK
 }
 
+#if USE_USD_SDK
+bool operator==(const pxr::SdfLayerRefPtr& LHS, const pxr::SdfLayerWeakPtr& RHS)
+{
+    return RHS.PointsTo(*LHS);
+}
+bool operator==(const pxr::SdfLayerWeakPtr& LHS, const pxr::SdfLayerWeakPtr& RHS)
+{
+    return LHS.GetUniqueIdentifier() == RHS.GetUniqueIdentifier();
+}
+bool operator==(const pxr::UsdStageRefPtr& LHS, const pxr::UsdStageWeakPtr& RHS)
+{
+    return RHS.PointsTo(*LHS);
+}
+bool operator==(const pxr::UsdStageWeakPtr& LHS, const pxr::UsdStageWeakPtr& RHS)
+{
+    return LHS.GetUniqueIdentifier() == RHS.GetUniqueIdentifier();
+}
+#endif      // USE_USD_SDK
+
 class FUnrealUSDWrapperModule : public IUnrealUSDWrapperModule
 {
 public:
diff --git a/Engine/Plugins/Importers/USDImporter/Source/UnrealUSDWrapper/Private/UsdWrappers/SdfLayer.cpp b/Engine/Plugins/Importers/USDImporter/Source/UnrealUSDWrapper/Private/UsdWrappers/SdfLayer.cpp
--- a/Engine/Plugins/Importers/USDImporter/Source/UnrealUSDWrapper/Private/UsdWrappers/SdfLayer.cpp
+++ b/Engine/Plugins/Importers/USDImporter/Source/UnrealUSDWrapper/Private/UsdWrappers/SdfLayer.cpp
@@ -124,28 +124,30 @@ namespace UE
 	}
 
 	template<typename PtrType>
-	template<typename OtherPtrType>
-	bool FSdfLayerBase<PtrType>::operator==(const FSdfLayerBase<OtherPtrType>& Other) const
+	bool FSdfLayerBase<PtrType>::operator==(const FSdfLayerBase<PtrType>& Other) const
 	{
 		return Impl->GetInner() == Other.Impl->GetInner();
 	}
 
-	template bool UNREALUSDWRAPPER_API FSdfLayer::operator==(const FSdfLayer& Other) const;
-	template bool UNREALUSDWRAPPER_API FSdfLayer::operator==(const FSdfLayerWeak& Other) const;
-	template bool UNREALUSDWRAPPER_API FSdfLayerWeak::operator==(const FSdfLayer& Other) const;
-	template bool UNREALUSDWRAPPER_API FSdfLayerWeak::operator==(const FSdfLayerWeak& Other) const;
-
 	template<typename PtrType>
-	template<typename OtherPtrType>
-	bool FSdfLayerBase<PtrType>::operator!=(const FSdfLayerBase<OtherPtrType>& Other) const
+	bool FSdfLayerBase<PtrType>::operator!=(const FSdfLayerBase<PtrType>& Other) const
 	{
 		return !(*this == Other);
 	}
 
-	template bool UNREALUSDWRAPPER_API FSdfLayer::operator!=(const FSdfLayer& Other) const;
-	template bool UNREALUSDWRAPPER_API FSdfLayer::operator!=(const FSdfLayerWeak& Other) const;
-	template bool UNREALUSDWRAPPER_API FSdfLayerWeak::operator!=(const FSdfLayer& Other) const;
-	template bool UNREALUSDWRAPPER_API FSdfLayerWeak::operator!=(const FSdfLayerWeak& Other) const;
+#if USE_USD_SDK
+	template<typename PtrType>
+	bool FSdfLayerBase<PtrType>::operator==(const pxr::SdfLayerWeakPtr& Other) const
+	{
+		return pxr::SdfLayerWeakPtr{*this} == Other;
+	}
+
+	template<typename PtrType>
+	bool FSdfLayerBase<PtrType>::operator!=(const pxr::SdfLayerWeakPtr& Other) const
+	{
+		return !(*this == Other);
+	}
+#endif	  // USE_USD_SDK
 
 	template<typename PtrType>
 	FSdfLayerBase<PtrType>::operator bool() const
diff --git a/Engine/Plugins/Importers/USDImporter/Source/UnrealUSDWrapper/Private/UsdWrappers/UsdStage.cpp b/Engine/Plugins/Importers/USDImporter/Source/UnrealUSDWrapper/Private/UsdWrappers/UsdStage.cpp
--- a/Engine/Plugins/Importers/USDImporter/Source/UnrealUSDWrapper/Private/UsdWrappers/UsdStage.cpp
+++ b/Engine/Plugins/Importers/USDImporter/Source/UnrealUSDWrapper/Private/UsdWrappers/UsdStage.cpp
@@ -132,6 +132,20 @@ namespace UE
 		return !(*this == Other);
 	}
 
+#if USE_USD_SDK
+	template<typename PtrType>
+	bool FUsdStageBase<PtrType>::operator==(const pxr::UsdStageWeakPtr& Other) const
+	{
+		return pxr::UsdStageWeakPtr{*this} == Other;
+	}
+
+	template<typename PtrType>
+	bool FUsdStageBase<PtrType>::operator!=(const pxr::UsdStageWeakPtr& Other) const
+	{
+		return !(*this == Other);
+	}
+#endif	  // USE_USD_SDK
+
 	template<typename PtrType>
 	FUsdStageBase<PtrType>::~FUsdStageBase()
 	{
diff --git a/Engine/Plugins/Importers/USDImporter/Source/UnrealUSDWrapper/Public/UnrealUSDWrapper.h b/Engine/Plugins/Importers/USDImporter/Source/UnrealUSDWrapper/Public/UnrealUSDWrapper.h
--- a/Engine/Plugins/Importers/USDImporter/Source/UnrealUSDWrapper/Public/UnrealUSDWrapper.h
+++ b/Engine/Plugins/Importers/USDImporter/Source/UnrealUSDWrapper/Public/UnrealUSDWrapper.h
@@ -30,19 +30,27 @@
 
 #if USE_USD_SDK
 PXR_NAMESPACE_OPEN_SCOPE
-	class GfMatrix4d;
-	class SdfPath;
-	class TfToken;
-	class UsdAttribute;
-	class UsdGeomMesh;
-	class UsdPrim;
-	class UsdStage;
-	class UsdStageCache;
-
-	template<typename T>
-	class TfRefPtr;
+    class GfMatrix4d;
+    class SdfPath;
+    class TfToken;
+    class UsdAttribute;
+    class UsdGeomMesh;
+    class UsdPrim;
+    class UsdStage;
+    class UsdStageCache;
+
+    template<typename T>
+    class TfRefPtr;
+    template<typename T>
+    class TfWeakPtr;
+
+    using UsdStageRefPtr = TfRefPtr<UsdStage>;
+    using UsdStageWeakPtr = TfWeakPtr<UsdStage>;
+
+    using SdfLayerRefPtr = TfRefPtr<SdfLayer>;
+    using SdfLayerWeakPtr = TfWeakPtr<SdfLayer>;
 PXR_NAMESPACE_CLOSE_SCOPE
-#endif	  // #if USE_USD_SDK
+#endif      // #if USE_USD_SDK
 
 class IUsdPrim;
 class FUsdDiagnosticDelegate;
@@ -531,3 +539,32 @@ struct UNREALUSDWRAPPER_API FUsdDelegates
 	static FUsdImportDelegate OnPreUsdImport;
 	static FUsdImportDelegate OnPostUsdImport;
 };
+
+#if USE_USD_SDK
+// This is a temp patch, because somewhere around cl 32323934 (so for 5.5, with USD 24.03,
+// see UE-210206 and UE-211302) the comparison operator between const SdfLayerRefPtr and SdfLayerWeakPtr
+// was leading to a crash, without any related change on our side. These operators are not defined by
+// USD so they are implicitly generated, so to work around it we just define them here and implement
+// them in a way that doesn't crash.
+//
+// We can't make these templated for any type as the implementations can't be generated on modules
+// that have RTTI disabled. This unfortunately means that we may get the same issues if we make some
+// ptr::TfRefPtr<T> and compare with pxr::TfWeakPtr<T>, but that should hopefully be rare in general.
+// We mostly use these pointer types when manipulating layers and stages.
+//
+// Additionally, spelling out the types here also lets us just forward declare them easily and not put
+// more USD includes in this header file.
+//
+// Note that for this workaround to work on Clang it was also necessary to define these operators on
+// the UE::FSdfLayer wrapper:
+//
+//	bool operator==(const pxr::SdfLayerWeakPtr& Other) const;
+// 	bool operator!=(const pxr::SdfLayerWeakPtr& Other) const;
+//
+// Otherwise we'd end up with ambiguous comparisons as the FSdfLayer types have implicit conversion
+// operators to the underlying pxr::SdfLayerWeakPtr/pxr::SdfLayerRefPtr
+UNREALUSDWRAPPER_API bool operator==(const pxr::SdfLayerRefPtr& LHS, const pxr::SdfLayerWeakPtr& RHS);
+UNREALUSDWRAPPER_API bool operator==(const pxr::SdfLayerWeakPtr& LHS, const pxr::SdfLayerWeakPtr& RHS);
+UNREALUSDWRAPPER_API bool operator==(const pxr::UsdStageRefPtr& LHS, const pxr::UsdStageWeakPtr& RHS);
+UNREALUSDWRAPPER_API bool operator==(const pxr::UsdStageWeakPtr& LHS, const pxr::UsdStageWeakPtr& RHS);
+#endif      // USE_USD_SDK
diff --git a/Engine/Plugins/Importers/USDImporter/Source/UnrealUSDWrapper/Public/UsdWrappers/SdfLayer.h b/Engine/Plugins/Importers/USDImporter/Source/UnrealUSDWrapper/Public/UsdWrappers/SdfLayer.h
--- a/Engine/Plugins/Importers/USDImporter/Source/UnrealUSDWrapper/Public/UsdWrappers/SdfLayer.h
+++ b/Engine/Plugins/Importers/USDImporter/Source/UnrealUSDWrapper/Public/UsdWrappers/SdfLayer.h
@@ -90,10 +90,15 @@ namespace UE
 		FSdfLayerBase& operator=(const FSdfLayerWeak& Other);
 		FSdfLayerBase& operator=(FSdfLayerWeak&& Other);
 
-		template<typename OtherPtrType>
-		bool operator==(const FSdfLayerBase<OtherPtrType>& Other) const;
-		template<typename OtherPtrType>
-		bool operator!=(const FSdfLayerBase<OtherPtrType>& Other) const;
+		bool operator==(const FSdfLayerBase& Other) const;
+		bool operator!=(const FSdfLayerBase& Other) const;
+
+#if USE_USD_SDK
+		// This is for the UE-210206 workaround. For more details see the comment
+		// at the bottom of UnrealUSDWrapper.h
+		bool operator==(const pxr::SdfLayerWeakPtr& Other) const;
+		bool operator!=(const pxr::SdfLayerWeakPtr& Other) const;
+#endif	  // #if USE_USD_SDK
 
 		explicit operator bool() const;
 
diff --git a/Engine/Plugins/Importers/USDImporter/Source/UnrealUSDWrapper/Public/UsdWrappers/UsdStage.h b/Engine/Plugins/Importers/USDImporter/Source/UnrealUSDWrapper/Public/UsdWrappers/UsdStage.h
--- a/Engine/Plugins/Importers/USDImporter/Source/UnrealUSDWrapper/Public/UsdWrappers/UsdStage.h
+++ b/Engine/Plugins/Importers/USDImporter/Source/UnrealUSDWrapper/Public/UsdWrappers/UsdStage.h
@@ -46,6 +46,13 @@ namespace UE
 		bool operator==(const FUsdStageBase& Other) const;
 		bool operator!=(const FUsdStageBase& Other) const;
 
+#if USE_USD_SDK
+		// This is for the UE-210206 workaround. For more details see the comment
+		// at the bottom of UnrealUSDWrapper.h
+		bool operator==(const pxr::UsdStageWeakPtr& Other) const;
+		bool operator!=(const pxr::UsdStageWeakPtr& Other) const;
+#endif	  // USE_USD_SDK
+
 		// Auto conversion from/to PtrType. We use concrete pointer types here
 		// because we should also be able to convert between them
 	public:
