USD: Implements UsdAssetCache3 which should be a simpler and more flexible way of managing USD assets;

Refactor dome light handling to create assets from within CreateAssets, like all other asset types;
Refactor groom cache parsing into more tasks so that it fits nicely into the new AssetCache interface;
Fix asset parsing race condition when translating assets from two identical prims (this could happen for all asset types...);
Move AssetUserData/ImportData handling functions into USDObjectUtils;
Start the UsdUnreal namespace convention with separate namespaces for each individual "library".

diff --git a/Engine/Plugins/Importers/USDImporter/Source/GeometryCacheUSD/Private/GeometryCacheTrackUSD.cpp b/Engine/Plugins/Importers/USDImporter/Source/GeometryCacheUSD/Private/GeometryCacheTrackUSD.cpp
--- a/Engine/Plugins/Importers/USDImporter/Source/GeometryCacheUSD/Private/GeometryCacheTrackUSD.cpp
+++ b/Engine/Plugins/Importers/USDImporter/Source/GeometryCacheUSD/Private/GeometryCacheTrackUSD.cpp
@@ -307,11 +307,14 @@ void UGeometryCacheTrackUsd::UpdateTime(float Time, bool bLooping)
 
 void UGeometryCacheTrackUsd::RegisterStream()
 {
-	const bool bNeedPrefetch = !IGeometryCacheStreamer::Get().IsTrackRegistered(this);
-	IGeometryCacheStreamer::Get().RegisterTrack(this, UsdStream.Get());
-	if (bNeedPrefetch)
+	if (UsdStream)
 	{
-		UsdStream->Prefetch(StartFrameIndex);
+		const bool bNeedPrefetch = !IGeometryCacheStreamer::Get().IsTrackRegistered(this);
+		IGeometryCacheStreamer::Get().RegisterTrack(this, UsdStream.Get());
+		if (bNeedPrefetch)
+		{
+			UsdStream->Prefetch(StartFrameIndex);
+		}
 	}
 }
 
diff --git a/Engine/Plugins/Importers/USDImporter/Source/USDClasses/Private/USDClassesModule.cpp b/Engine/Plugins/Importers/USDImporter/Source/USDClasses/Private/USDClassesModule.cpp
--- a/Engine/Plugins/Importers/USDImporter/Source/USDClasses/Private/USDClassesModule.cpp
+++ b/Engine/Plugins/Importers/USDImporter/Source/USDClasses/Private/USDClassesModule.cpp
@@ -2,13 +2,17 @@
 
 #include "USDClassesModule.h"
 
+#include "USDAssetCache3.h"
 #include "USDLog.h"
+#include "USDObjectUtils.h"
 #include "USDProjectSettings.h"
 
 #include "AnalyticsEventAttribute.h"
 #include "Animation/AnimBlueprint.h"
 #include "Animation/AnimSequence.h"
 #include "Animation/Skeleton.h"
+#include "AssetRegistry/AssetData.h"
+#include "AssetRegistry/AssetRegistryModule.h"
 #include "Engine/Engine.h"
 #include "Engine/SkeletalMesh.h"
 #include "Engine/SkinnedAssetCommon.h"
@@ -16,6 +20,7 @@
 #include "Engine/Texture.h"
 #include "Engine/Texture2D.h"
 #include "EngineAnalytics.h"
+#include "Framework/Notifications/NotificationManager.h"
 #include "GeometryCache.h"
 #include "GroomAsset.h"
 #include "GroomBindingAsset.h"
@@ -25,11 +30,13 @@
 #include "Materials/MaterialInstance.h"
 #include "Materials/MaterialInstanceConstant.h"
 #include "Materials/MaterialInstanceDynamic.h"
+#include "Materials/MaterialInterface.h"
 #include "Misc/EngineVersion.h"
 #include "Misc/FileHelper.h"
 #include "Misc/PackageName.h"
 #include "Misc/Paths.h"
 #include "Misc/SecureHash.h"
+#include "Modules/ModuleManager.h"
 #include "PhysicsEngine/PhysicsAsset.h"
 #include "Serialization/JsonSerializer.h"
 #include "SparseVolumeTexture/SparseVolumeTexture.h"
@@ -37,13 +44,17 @@
 #include "UObject/ObjectSaveContext.h"
 #include "UObject/Package.h"
 #include "UObject/UObjectGlobals.h"
+#include "Widgets/Notifications/SNotificationList.h"
 
 #if WITH_EDITOR
-#include "Editor/MaterialEditor/Public/MaterialEditingLibrary.h"
+#include "AssetToolsModule.h"
+#include "MaterialEditingLibrary.h"
 #endif	  // WITH_EDITOR
 
 DEFINE_LOG_CATEGORY(LogUsd);
 
+#define LOCTEXT_NAMESPACE "USDClassesModule"
+
 namespace UE::USDClasses::Private
 {
 	static const FString DisplayColorID = TEXT("!DisplayColor");
@@ -440,17 +451,140 @@ TSet<UObject*> IUsdClassesModule::GetAssetDependencies(UObject* Asset)
 	return Result;
 }
 
-FString IUsdClassesModule::SanitizeObjectName(const FString& InObjectName)
+UUsdAssetCache3* IUsdClassesModule::GetAssetCacheForProject()
 {
-	FString SanitizedText = InObjectName;
-	const TCHAR* InvalidChar = INVALID_OBJECTNAME_CHARACTERS;
-	while (*InvalidChar)
+	// First check if we have any valid AssetCache on the project settings
+	UUsdProjectSettings* ProjectSettings = GetMutableDefault<UUsdProjectSettings>();
+	if (ProjectSettings)
+	{
+		// Check if the package exists on disk first to try and avoid some ugly warnings if we try calling TryLoad with a broken path
+		const FAssetRegistryModule& AssetRegistryModule = FModuleManager::GetModuleChecked<FAssetRegistryModule>(TEXT("AssetRegistry"));
+		FAssetData AssetData = AssetRegistryModule.Get().GetAssetByObjectPath(ProjectSettings->DefaultAssetCache);
+		if (AssetData.IsValid())
+		{
+			if (UUsdAssetCache3* DefaultCache = Cast<UUsdAssetCache3>(ProjectSettings->DefaultAssetCache.TryLoad()))
+			{
+				return DefaultCache;
+			}
+		}
+	}
+
+	// Don't record the creation of the asset cache itself into the transaction buffer (and the setting of it on the project settings)
+	// so that it remains if we undo. In general I don't think we ever want to "undo the creation of assets" anyway, but here
+	// it also helps prevent this issue: Get stage actor without asset cache -> Open stage (and get a new
+	// asset cache that we create right here) -> Undo -> Open stage again.
+	// That would give the stage actor a brand new asset cache without the property values (including the tracked assets)
+	// the previous one had, which could lead us to create duplicates of every asset...
+	TGuardValue<ITransaction*> SuppressTransaction{GUndo, nullptr};
+
+	UUsdAssetCache3* DefaultCache = nullptr;
+#if WITH_EDITOR
+	if (GIsEditor)
+	{
+		// Nothing is set on the project settings yet, so let's create a brand new asset cache.
+		// First let's find a unique package name
+		FString DesiredPath = TEXT("/Game/UsdAssetCache");
+		FString Suffix = TEXT("");
+		FString UniquePackageName;
+		FString UniqueAssetName;
+		const FAssetToolsModule& AssetToolsModule = FModuleManager::LoadModuleChecked<FAssetToolsModule>("AssetTools");
+		AssetToolsModule.Get().CreateUniqueAssetName(DesiredPath, Suffix, UniquePackageName, UniqueAssetName);
+		if (UniquePackageName.EndsWith(UniqueAssetName))
+		{
+			UniquePackageName = UniquePackageName.LeftChop(UniqueAssetName.Len() + 1);
+		}
+
+		// Create the new asset cache
+		UFactory* Factory = nullptr;
+		DefaultCache = Cast<UUsdAssetCache3>(
+			AssetToolsModule.Get().CreateAsset(*UniqueAssetName, *UniquePackageName, UUsdAssetCache3::StaticClass(), Factory)
+		);
+	}
+	else
+#endif	  // WITH_EDITOR
+	{
+		DefaultCache = NewObject<UUsdAssetCache3>();
+	}
+
+	if (DefaultCache && ProjectSettings && GIsEditor)
 	{
-		SanitizedText.ReplaceCharInline(*InvalidChar, TCHAR('_'), ESearchCase::CaseSensitive);
-		++InvalidChar;
+		const FText Text = LOCTEXT("NewAssetCacheToastText", "USD Asset Cache");
+
+		const FText SubText = FText::Format(
+			LOCTEXT(
+				"NewAssetCacheToastSubText",
+				"A new default UsdAssetCache asset was created for the project at path '{0}', but it can be changed in the project settings.\n\nThe UsdAssetCache is used by UsdStageActors in order to share and reuse assets generated from USD."
+			),
+			FText::FromString(DefaultCache->GetPathName())
+		);
+
+		UE_LOG(LogUsd, Log, TEXT("%s"), *SubText.ToString().Replace(TEXT("\n\n"), TEXT(" ")));
+
+		const UUsdProjectSettings* Settings = GetDefault<UUsdProjectSettings>();
+		if (Settings && Settings->bShowCreateDefaultAssetCacheDialog)
+		{
+			static TWeakPtr<SNotificationItem> Notification;
+
+			FNotificationInfo Toast(Text);
+			Toast.SubText = SubText;
+			Toast.Image = FCoreStyle::Get().GetBrush(TEXT("MessageLog.Warning"));
+			Toast.CheckBoxText = LOCTEXT("DontAskAgain", "Don't prompt again");
+			Toast.bUseLargeFont = false;
+			Toast.bFireAndForget = false;
+			Toast.FadeOutDuration = 0.0f;
+			Toast.ExpireDuration = 0.0f;
+			Toast.bUseThrobber = false;
+			Toast.bUseSuccessFailIcons = false;
+			Toast.ButtonDetails.Emplace(
+				LOCTEXT("OverridenOpinionMessageOk", "Ok"),
+				FText::GetEmpty(),
+				FSimpleDelegate::CreateLambda(
+					[]()
+					{
+						if (TSharedPtr<SNotificationItem> PinnedNotification = Notification.Pin())
+						{
+							PinnedNotification->SetCompletionState(SNotificationItem::CS_Success);
+							PinnedNotification->ExpireAndFadeout();
+						}
+					}
+				)
+			);
+			// This is flipped because the default checkbox message is "Don't prompt again"
+			Toast.CheckBoxState = Settings->bShowCreateDefaultAssetCacheDialog ? ECheckBoxState::Unchecked : ECheckBoxState::Checked;
+			Toast.CheckBoxStateChanged = FOnCheckStateChanged::CreateStatic(
+				[](ECheckBoxState NewState)
+				{
+					if (UUsdProjectSettings* Settings = GetMutableDefault<UUsdProjectSettings>())
+					{
+						// This is flipped because the default checkbox message is "Don't prompt again"
+						Settings->bShowCreateDefaultAssetCacheDialog = NewState == ECheckBoxState::Unchecked;
+						Settings->SaveConfig();
+					}
+				}
+			);
+
+			// Only show one at a time
+			if (!Notification.IsValid())
+			{
+				Notification = FSlateNotificationManager::Get().AddNotification(Toast);
+			}
+
+			if (TSharedPtr<SNotificationItem> PinnedNotification = Notification.Pin())
+			{
+				PinnedNotification->SetCompletionState(SNotificationItem::CS_Pending);
+			}
+		}
+
+		ProjectSettings->DefaultAssetCache = DefaultCache;
+		ProjectSettings->SaveConfig();
 	}
 
-	return SanitizedText;
+	return DefaultCache;
+}
+
+FString IUsdClassesModule::SanitizeObjectName(const FString& InObjectName)
+{
+	return UsdUnreal::ObjectUtils::SanitizeObjectName(InObjectName);
 }
 
 FString IUsdClassesModule::FDisplayColorMaterial::ToString()
@@ -520,9 +654,9 @@ UMaterialInstanceDynamic* IUsdClassesModule::CreateDisplayColorMaterialInstanceD
 			GetTransientPackage(),
 			UMaterialInstanceConstant::StaticClass(),
 			*FString::Printf(
-				TEXT("DisplayColor_%s_%s"),
-				DisplayColorDescription.bHasOpacity ? TEXT("Opacity") : TEXT("NoOpacity"),
-				DisplayColorDescription.bIsDoubleSided ? TEXT("DoubleSided") : TEXT("SingleSided")
+				TEXT("DisplayColor%s%s"),
+				DisplayColorDescription.bHasOpacity ? TEXT("_Translucent") : TEXT(""),
+				DisplayColorDescription.bIsDoubleSided ? TEXT("_TwoSided") : TEXT("")
 			)
 		);
 
@@ -550,9 +684,9 @@ UMaterialInstanceConstant* IUsdClassesModule::CreateDisplayColorMaterialInstance
 			GetTransientPackage(),
 			UMaterialInstanceConstant::StaticClass(),
 			*FString::Printf(
-				TEXT("DisplayColor_%s_%s"),
-				DisplayColorDescription.bHasOpacity ? TEXT("Opacity") : TEXT("NoOpacity"),
-				DisplayColorDescription.bIsDoubleSided ? TEXT("DoubleSided") : TEXT("SingleSided")
+				TEXT("DisplayColor%s%s"),
+				DisplayColorDescription.bHasOpacity ? TEXT("_Translucent") : TEXT(""),
+				DisplayColorDescription.bIsDoubleSided ? TEXT("_TwoSided") : TEXT("")
 			)
 		);
 
@@ -603,4 +737,6 @@ private:
 	FDelegateHandle PackageSavedWithContextEventHandle;
 };
 
+#undef LOCTEXT_NAMESPACE
+
 IMPLEMENT_MODULE(FUsdClassesModule, USDClasses);
diff --git a/Engine/Plugins/Importers/USDImporter/Source/USDClasses/Private/USDDrawModeComponent.cpp b/Engine/Plugins/Importers/USDImporter/Source/USDClasses/Private/USDDrawModeComponent.cpp
--- a/Engine/Plugins/Importers/USDImporter/Source/USDClasses/Private/USDDrawModeComponent.cpp
+++ b/Engine/Plugins/Importers/USDImporter/Source/USDClasses/Private/USDDrawModeComponent.cpp
@@ -83,8 +83,11 @@ namespace UE::UsdDrawModeComponentImpl::Private
 				// Reference: PrimitiveDrawingUtils.cpp, DrawBox function
 
 				// Calculate verts for a face pointing down Z
-				FVector3f Positions[4] =
-					{FVector3f(-0.5, -0.5, +0.5), FVector3f(-0.5, +0.5, +0.5), FVector3f(+0.5, +0.5, +0.5), FVector3f(+0.5, -0.5, +0.5)};
+				FVector3f Positions[4] = {
+					FVector3f(-0.5, -0.5, +0.5),
+					FVector3f(-0.5, +0.5, +0.5),
+					FVector3f(+0.5, +0.5, +0.5),
+					FVector3f(+0.5, -0.5, +0.5)};
 
 				FVector2f UVs[4] = {
 					FVector2f(0, 0),
@@ -662,21 +665,14 @@ namespace UE::UsdDrawModeComponentImpl::Private
 		FMaterialParameterInfo Info;
 		Info.Name = TEXT("Texture");
 
-#if WITH_EDITOR
-		if (GIsEditor)
+		if (UMaterialInstanceDynamic* Dynamic = Cast<UMaterialInstanceDynamic>(Instance))
 		{
-			if (UMaterialInstanceConstant* Constant = Cast<UMaterialInstanceConstant>(Instance))
-			{
-				Constant->SetTextureParameterValueEditorOnly(Info, Texture);
-			}
+			Dynamic->SetTextureParameterValueByInfo(Info, Texture);
 		}
 		else
-#endif	  // WITH_EDITOR
 		{
-			if (UMaterialInstanceDynamic* Dynamic = Cast<UMaterialInstanceDynamic>(Instance))
-			{
-				Dynamic->SetTextureParameterValueByInfo(Info, Texture);
-			}
+			// We should always create MIDs on-demand now, and never a MIC
+			ensure(false);
 		}
 	}
 };	  // namespace UE::UsdDrawModeComponentImpl::Private
@@ -1130,22 +1126,10 @@ UMaterialInstance* UUsdDrawModeComponent::GetOrCreateTextureMaterial(EUsdModelCa
 	{
 		if (UMaterialInterface* ParentMaterial = Cast<UMaterialInterface>(ProjectSettings->ReferenceModelCardTextureMaterial.TryLoad()))
 		{
-#if WITH_EDITOR
-			if (GIsEditor)
-			{
-				if (UMaterialInstanceConstant* NewMIC = NewObject<UMaterialInstanceConstant>(GetTransientPackage(), NAME_None, RF_Transient))
-				{
-					UMaterialEditingLibrary::SetMaterialInstanceParent(NewMIC, ParentMaterial);
-					Instance = NewMIC;
-				}
-			}
-			else
-#endif
+			UObject* Outer = this;
+			if (UMaterialInstanceDynamic* NewMID = UMaterialInstanceDynamic::Create(ParentMaterial, Outer, NAME_None))
 			{
-				if (UMaterialInstanceDynamic* NewMID = UMaterialInstanceDynamic::Create(ParentMaterial, GetTransientPackage(), NAME_None))
-				{
-					Instance = NewMID;
-				}
+				Instance = NewMID;
 			}
 		}
 	}
@@ -1199,7 +1183,7 @@ void UUsdDrawModeComponent::RefreshMaterialInstances()
 		{
 			SetTextureParameter(GetOrCreateTextureMaterial(Face), OppositeTexture);
 		}
-		// Face is not authored, but the opposite face is. It doesn't have a texture though. (+ other cases) --> Draw DisplyColor
+		// Face is not authored, but the opposite face is. It doesn't have a texture though. (+ other cases) --> Draw DisplayColor
 		else
 		{
 			IndicesToSetDrawModeColor.Add(FaceIndex);
@@ -1238,16 +1222,7 @@ void UUsdDrawModeComponent::RefreshMaterialInstances()
 			Desc.bHasOpacity = false;
 			Desc.bIsDoubleSided = false;
 
-#if WITH_EDITOR
-			if (GIsEditor)
-			{
-				DisplayColorInstance = IUsdClassesModule::CreateDisplayColorMaterialInstanceConstant(Desc);
-			}
-			else
-#endif	  // WITH_EDITOR
-			{
-				DisplayColorInstance = IUsdClassesModule::CreateDisplayColorMaterialInstanceDynamic(Desc);
-			}
+			DisplayColorInstance = IUsdClassesModule::CreateDisplayColorMaterialInstanceDynamic(Desc);
 
 			if (DisplayColorInstance)
 			{
diff --git a/Engine/Plugins/Importers/USDImporter/Source/USDClasses/Public/USDAssetCache2.h b/Engine/Plugins/Importers/USDImporter/Source/USDClasses/Public/USDAssetCache2.h
--- a/Engine/Plugins/Importers/USDImporter/Source/USDClasses/Public/USDAssetCache2.h
+++ b/Engine/Plugins/Importers/USDImporter/Source/USDClasses/Public/USDAssetCache2.h
@@ -12,7 +12,7 @@
 /**
  * Owns the assets generated and reused by USD Stages, allowing thread-safe retrieval/storage.
  */
-UCLASS(BlueprintType, Blueprintable, meta = (DisplayName = "USD Asset Cache", ScriptName = "UsdAssetCache"))
+UCLASS(BlueprintType, Blueprintable)
 class USDCLASSES_API UUsdAssetCache2 : public UObject
 {
 	GENERATED_BODY()
diff --git a/Engine/Plugins/Importers/USDImporter/Source/USDClasses/Public/USDAssetUserData.h b/Engine/Plugins/Importers/USDImporter/Source/USDClasses/Public/USDAssetUserData.h
--- a/Engine/Plugins/Importers/USDImporter/Source/USDClasses/Public/USDAssetUserData.h
+++ b/Engine/Plugins/Importers/USDImporter/Source/USDClasses/Public/USDAssetUserData.h
@@ -24,6 +24,10 @@ public:
 	// bits of metadata
 	UPROPERTY(EditAnywhere, BlueprintReadWrite, Category = "USD")
 	TMap<FString, FUsdCombinedPrimMetadata> StageIdentifierToMetadata;
+
+	// Hash of the prim data that was used to generate this asset
+	UPROPERTY(VisibleAnywhere, BlueprintReadOnly, Category = "USD")
+	FString OriginalHash;
 };
 
 UCLASS(BlueprintType)
diff --git a/Engine/Plugins/Importers/USDImporter/Source/USDClasses/Public/USDClassesModule.h b/Engine/Plugins/Importers/USDImporter/Source/USDClasses/Public/USDClassesModule.h
--- a/Engine/Plugins/Importers/USDImporter/Source/USDClasses/Public/USDClassesModule.h
+++ b/Engine/Plugins/Importers/USDImporter/Source/USDClasses/Public/USDClassesModule.h
@@ -8,6 +8,7 @@
 
 class UMaterialInstanceConstant;
 class UMaterialInstanceDynamic;
+class UUsdAssetCache3;
 struct FAnalyticsEventAttribute;
 
 enum class EUsdReferenceMaterialProperties : uint8
@@ -60,7 +61,11 @@ public:
 	 */
 	USDCLASSES_API static TSet<UObject*> GetAssetDependencies(UObject* Asset);
 
+	// Returns the default asset cache for the project or create a new one at the project root
+	USDCLASSES_API static UUsdAssetCache3* GetAssetCacheForProject();
+
 	// Adapted from ObjectTools as it is within an Editor-only module
+	UE_DEPRECATED(5.5, "This function has been moved to USDObjectUtils.h")
 	USDCLASSES_API static FString SanitizeObjectName(const FString& InObjectName);
 
 	/** Describes the type of vertex color/DisplayColor material that we would need in order to render a prim's displayColor data as intended */
diff --git a/Engine/Plugins/Importers/USDImporter/Source/USDClasses/Public/USDDrawModeComponent.h b/Engine/Plugins/Importers/USDImporter/Source/USDClasses/Public/USDDrawModeComponent.h
--- a/Engine/Plugins/Importers/USDImporter/Source/USDClasses/Public/USDDrawModeComponent.h
+++ b/Engine/Plugins/Importers/USDImporter/Source/USDClasses/Public/USDDrawModeComponent.h
@@ -7,7 +7,6 @@
 #include "USDDrawModeComponent.generated.h"
 
 class UMaterialInstance;
-class UUsdAssetCache2;
 enum class EUsdUpAxis : uint8;
 
 UENUM()
diff --git a/Engine/Plugins/Importers/USDImporter/Source/USDClasses/Public/USDProjectSettings.h b/Engine/Plugins/Importers/USDImporter/Source/USDClasses/Public/USDProjectSettings.h
--- a/Engine/Plugins/Importers/USDImporter/Source/USDClasses/Public/USDProjectSettings.h
+++ b/Engine/Plugins/Importers/USDImporter/Source/USDClasses/Public/USDProjectSettings.h
@@ -67,7 +67,7 @@ public:
 	 * USD Asset Cache to use for USD Stage Actors that don't have any asset cache specified.
 	 * Leave this empty to have each stage actor generate it's on transient cache instead.
 	 */
-	UPROPERTY(config, EditAnywhere, Category = USD, meta = (AllowedClasses = "/Script/USDClasses.UsdAssetCache2"))
+	UPROPERTY(config, EditAnywhere, Category = USD, meta = (AllowedClasses = "/Script/USDClasses.UsdAssetCache3"))
 	FSoftObjectPath DefaultAssetCache;
 
 	UPROPERTY(config, EditAnywhere, Category = "USD|Dialogs")
diff --git a/Engine/Plugins/Importers/USDImporter/Source/USDClasses/USDClasses.Build.cs b/Engine/Plugins/Importers/USDImporter/Source/USDClasses/USDClasses.Build.cs
--- a/Engine/Plugins/Importers/USDImporter/Source/USDClasses/USDClasses.Build.cs
+++ b/Engine/Plugins/Importers/USDImporter/Source/USDClasses/USDClasses.Build.cs
@@ -29,6 +29,7 @@ namespace UnrealBuildTool.Rules
 					"GeometryFramework", // For BaseDynamicMeshSceneProxy.h
 					"HairStrandsCore",
 					"Json", // To read/write plugInfo.json files from UnrealUSDWrapper.cpp
+					"LevelSequence",
 					"MeshDescription",
 					"RenderCore", // For UUsdDrawModeComponent implementation
 					"RHI",
@@ -40,6 +41,7 @@ namespace UnrealBuildTool.Rules
 				PrivateDependencyModuleNames.AddRange(
 					new string[]
 					{
+						"UnrealEd", // For ResetTransaction within UE::USDAssetCache3::Private::SafeDeleteObjects
 						"MaterialEditor", // For SetMaterialInstanceParent, when managing the DisplayColor materials
 					}
 				);
diff --git a/Engine/Plugins/Importers/USDImporter/Source/USDClassesEditor/Private/USDAssetCacheAssetActions.cpp b/Engine/Plugins/Importers/USDImporter/Source/USDClassesEditor/Private/USDAssetCacheAssetActions.cpp
--- a/Engine/Plugins/Importers/USDImporter/Source/USDClassesEditor/Private/USDAssetCacheAssetActions.cpp
+++ b/Engine/Plugins/Importers/USDImporter/Source/USDClassesEditor/Private/USDAssetCacheAssetActions.cpp
@@ -2,7 +2,7 @@
 
 #include "USDAssetCacheAssetActions.h"
 
-#include "USDAssetCache2.h"
+#include "USDAssetCache3.h"
 #include "USDAssetCacheAssetEditorToolkit.h"
 
 #define LOCTEXT_NAMESPACE "AssetCacheAssetActions"
@@ -19,7 +19,7 @@ FText FUsdAssetCacheAssetActions::GetName() const
 
 UClass* FUsdAssetCacheAssetActions::GetSupportedClass() const
 {
-	return UUsdAssetCache2::StaticClass();
+	return UUsdAssetCache3::StaticClass();
 }
 
 FColor FUsdAssetCacheAssetActions::GetTypeColor() const
@@ -33,7 +33,7 @@ void FUsdAssetCacheAssetActions::OpenAssetEditor(const TArray<UObject*>& InObjec
 
 	for (auto ObjIt = InObjects.CreateConstIterator(); ObjIt; ++ObjIt)
 	{
-		if (UUsdAssetCache2* Asset = Cast<UUsdAssetCache2>(*ObjIt))
+		if (UUsdAssetCache3* Asset = Cast<UUsdAssetCache3>(*ObjIt))
 		{
 			TSharedRef<FUsdAssetCacheAssetEditorToolkit> NewCustomAssetEditor = MakeShared<FUsdAssetCacheAssetEditorToolkit>();
 			NewCustomAssetEditor->Initialize(Mode, EditWithinLevelEditor, Asset);
diff --git a/Engine/Plugins/Importers/USDImporter/Source/USDClassesEditor/Private/USDAssetCacheAssetEditorToolkit.cpp b/Engine/Plugins/Importers/USDImporter/Source/USDClassesEditor/Private/USDAssetCacheAssetEditorToolkit.cpp
--- a/Engine/Plugins/Importers/USDImporter/Source/USDClassesEditor/Private/USDAssetCacheAssetEditorToolkit.cpp
+++ b/Engine/Plugins/Importers/USDImporter/Source/USDClassesEditor/Private/USDAssetCacheAssetEditorToolkit.cpp
@@ -2,7 +2,7 @@
 
 #include "USDAssetCacheAssetEditorToolkit.h"
 
-#include "USDAssetCache2.h"
+#include "USDAssetCache3.h"
 
 #include "Modules/ModuleManager.h"
 #include "PropertyEditorModule.h"
@@ -15,7 +15,7 @@ const FName FUsdAssetCacheAssetEditorToolkit::TabId("AssetCacheEditor");
 void FUsdAssetCacheAssetEditorToolkit::Initialize(
 	const EToolkitMode::Type Mode,
 	const TSharedPtr<IToolkitHost>& InitToolkitHost,
-	UUsdAssetCache2* InAssetCache
+	UUsdAssetCache3* InAssetCache
 )
 {
 	const TSharedRef<FTabManager::FLayout>
diff --git a/Engine/Plugins/Importers/USDImporter/Source/USDClassesEditor/Private/USDAssetCacheAssetEditorToolkit.h b/Engine/Plugins/Importers/USDImporter/Source/USDClassesEditor/Private/USDAssetCacheAssetEditorToolkit.h
--- a/Engine/Plugins/Importers/USDImporter/Source/USDClassesEditor/Private/USDAssetCacheAssetEditorToolkit.h
+++ b/Engine/Plugins/Importers/USDImporter/Source/USDClassesEditor/Private/USDAssetCacheAssetEditorToolkit.h
@@ -5,14 +5,14 @@
 #include "Templates/SharedPointer.h"
 #include "Toolkits/AssetEditorToolkit.h"
 
-class UUsdAssetCache2;
+class UUsdAssetCache3;
 
 class FUsdAssetCacheAssetEditorToolkit
 	: public FAssetEditorToolkit
 	, public FGCObject
 {
 public:
-	void Initialize(const EToolkitMode::Type Mode, const TSharedPtr<IToolkitHost>& InitToolkitHost, UUsdAssetCache2* InAssetCache);
+	void Initialize(const EToolkitMode::Type Mode, const TSharedPtr<IToolkitHost>& InitToolkitHost, UUsdAssetCache3* InAssetCache);
 
 private:
 	virtual FText GetBaseToolkitName() const override;
@@ -29,7 +29,7 @@ private:
 	virtual FString GetReferencerName() const override;
 
 private:
-	TObjectPtr<UUsdAssetCache2> AssetCache;
+	TObjectPtr<UUsdAssetCache3> AssetCache;
 
 	TSharedPtr<class IDetailsView> AssetCacheEditorWidget;
 
diff --git a/Engine/Plugins/Importers/USDImporter/Source/USDClassesEditor/Private/USDAssetCacheFactory.cpp b/Engine/Plugins/Importers/USDImporter/Source/USDClassesEditor/Private/USDAssetCacheFactory.cpp
--- a/Engine/Plugins/Importers/USDImporter/Source/USDClassesEditor/Private/USDAssetCacheFactory.cpp
+++ b/Engine/Plugins/Importers/USDImporter/Source/USDClassesEditor/Private/USDAssetCacheFactory.cpp
@@ -2,14 +2,14 @@
 
 #include "USDAssetCacheFactory.h"
 
-#include "USDAssetCache2.h"
+#include "USDAssetCache3.h"
 
 UUsdAssetCacheFactory::UUsdAssetCacheFactory(const FObjectInitializer& ObjectInitializer)
 	: Super(ObjectInitializer)
 {
 	bCreateNew = true;
 	bEditAfterNew = true;
-	SupportedClass = UUsdAssetCache2::StaticClass();
+	SupportedClass = UUsdAssetCache3::StaticClass();
 }
 
 UObject* UUsdAssetCacheFactory::FactoryCreateNew(
@@ -21,7 +21,7 @@ UObject* UUsdAssetCacheFactory::FactoryCreateNew(
 	FFeedbackContext* Warn
 )
 {
-	return NewObject<UUsdAssetCache2>(InParent, Name, Flags | RF_Transactional | RF_Public | RF_Standalone);
+	return NewObject<UUsdAssetCache3>(InParent, Name, Flags | RF_Transactional | RF_Public | RF_Standalone);
 }
 
 bool UUsdAssetCacheFactory::ShouldShowInNewMenu() const
diff --git a/Engine/Plugins/Importers/USDImporter/Source/USDClassesEditor/Private/USDClassesEditorModule.cpp b/Engine/Plugins/Importers/USDImporter/Source/USDClassesEditor/Private/USDClassesEditorModule.cpp
--- a/Engine/Plugins/Importers/USDImporter/Source/USDClassesEditor/Private/USDClassesEditorModule.cpp
+++ b/Engine/Plugins/Importers/USDImporter/Source/USDClassesEditor/Private/USDClassesEditorModule.cpp
@@ -4,50 +4,24 @@
 
 #include "USDAssetCache2.h"
 #include "USDAssetCacheAssetActions.h"
-#include "USDDefaultAssetCacheDialog.h"
 
 #include "AssetToolsModule.h"
-#include "Framework/Application/SlateApplication.h"
-#include "Interfaces/IMainFrameModule.h"
 #include "Modules/ModuleManager.h"
 
 #define LOCTEXT_NAMESPACE "USDClassesEditorModule"
 
 UUsdAssetCache2* IUsdClassesEditorModule::ShowMissingDefaultAssetCacheDialog()
 {
-	UUsdAssetCache2* Result = nullptr;
-	EDefaultAssetCacheDialogOption Unused = ShowMissingDefaultAssetCacheDialog(Result);
-	return Result;
+	return nullptr;
 }
 
 void IUsdClassesEditorModule::ShowMissingDefaultAssetCacheDialog(UUsdAssetCache2*& OutCreatedCache, bool& bOutUserAccepted)
 {
-	EDefaultAssetCacheDialogOption Result = ShowMissingDefaultAssetCacheDialog(OutCreatedCache);
-	bOutUserAccepted = Result != EDefaultAssetCacheDialogOption::Cancel;
 }
 
 EDefaultAssetCacheDialogOption IUsdClassesEditorModule::ShowMissingDefaultAssetCacheDialog(UUsdAssetCache2*& OutCreatedCache)
 {
-	TSharedPtr<SWindow> ParentWindow;
-
-	if (FModuleManager::Get().IsModuleLoaded("MainFrame"))
-	{
-		IMainFrameModule& MainFrame = FModuleManager::LoadModuleChecked<IMainFrameModule>("MainFrame");
-		ParentWindow = MainFrame.GetParentWindow();
-	}
-
-	FText WindowTitle = LOCTEXT("WindowTitle", "Set the default USD Asset Cache");
-
-	TSharedRef<SWindow> Window = SNew(SWindow).Title(WindowTitle).SizingRule(ESizingRule::Autosized).AdjustInitialSizeAndPositionForDPIScale(false);
-
-	TSharedPtr<SUsdDefaultAssetCacheDialog> OptionsWindow;
-	Window->SetContent(SAssignNew(OptionsWindow, SUsdDefaultAssetCacheDialog).WidgetWindow(Window));
-
-	const bool bSlowTaskWindow = false;
-	FSlateApplication::Get().AddModalWindow(Window, ParentWindow, bSlowTaskWindow);
-
-	OutCreatedCache = OptionsWindow->GetCreatedCache();
-	return OptionsWindow->GetDialogOutcome();
+	return EDefaultAssetCacheDialogOption::Cancel;
 }
 
 class FUsdClassesEditorModule : public IUsdClassesEditorModule
diff --git a/Engine/Plugins/Importers/USDImporter/Source/USDClassesEditor/Public/USDClassesEditorModule.h b/Engine/Plugins/Importers/USDImporter/Source/USDClassesEditor/Public/USDClassesEditorModule.h
--- a/Engine/Plugins/Importers/USDImporter/Source/USDClassesEditor/Public/USDClassesEditorModule.h
+++ b/Engine/Plugins/Importers/USDImporter/Source/USDClassesEditor/Public/USDClassesEditorModule.h
@@ -24,5 +24,6 @@ public:
 	UE_DEPRECATED(5.4, "Use the other signature that also returns the dialog outcome as an Enum")
 	USDCLASSESEDITOR_API static void ShowMissingDefaultAssetCacheDialog(UUsdAssetCache2*& OutCreatedCache, bool& bOutUserAccepted);
 
+	UE_DEPRECATED(5.5, "The Default Asset Cache dialog has been deprecated. An asset cache for the project will always be created on-demand now")
 	USDCLASSESEDITOR_API static EDefaultAssetCacheDialogOption ShowMissingDefaultAssetCacheDialog(UUsdAssetCache2*& OutCreatedCache);
 };
diff --git a/Engine/Plugins/Importers/USDImporter/Source/USDExporter/Private/AnimSequenceExporterUSD.cpp b/Engine/Plugins/Importers/USDImporter/Source/USDExporter/Private/AnimSequenceExporterUSD.cpp
--- a/Engine/Plugins/Importers/USDImporter/Source/USDExporter/Private/AnimSequenceExporterUSD.cpp
+++ b/Engine/Plugins/Importers/USDImporter/Source/USDExporter/Private/AnimSequenceExporterUSD.cpp
@@ -8,6 +8,7 @@
 #include "USDExporterModule.h"
 #include "USDLayerUtils.h"
 #include "USDLog.h"
+#include "USDObjectUtils.h"
 #include "USDOptionsWindow.h"
 #include "USDPrimConversion.h"
 #include "USDSkeletalDataConversion.h"
@@ -357,7 +358,7 @@ bool UAnimSequenceExporterUSD::ExportBinary(
 
 	if (Options->MetadataOptions.bExportAssetMetadata)
 	{
-		if (UUsdAssetUserData* UserData = UsdUtils::GetAssetUserData(AnimSequence))
+		if (UUsdAssetUserData* UserData = UsdUnreal::ObjectUtils::GetAssetUserData(AnimSequence))
 		{
 			UnrealToUsd::ConvertMetadata(
 				UserData,
diff --git a/Engine/Plugins/Importers/USDImporter/Source/USDExporter/Private/GeometryCacheExporterUSD.cpp b/Engine/Plugins/Importers/USDImporter/Source/USDExporter/Private/GeometryCacheExporterUSD.cpp
--- a/Engine/Plugins/Importers/USDImporter/Source/USDExporter/Private/GeometryCacheExporterUSD.cpp
+++ b/Engine/Plugins/Importers/USDImporter/Source/USDExporter/Private/GeometryCacheExporterUSD.cpp
@@ -9,6 +9,7 @@
 #include "USDExporterModule.h"
 #include "USDGeomMeshConversion.h"
 #include "USDLog.h"
+#include "USDObjectUtils.h"
 #include "USDOptionsWindow.h"
 #include "USDPrimConversion.h"
 #include "USDUnrealAssetInfo.h"
@@ -306,7 +307,7 @@ bool UGeometryCacheExporterUSD::ExportBinary(
 
 		if (Options->MetadataOptions.bExportAssetMetadata)
 		{
-			if (UUsdAssetUserData* UserData = UsdUtils::GetAssetUserData(GeometryCache))
+			if (UUsdAssetUserData* UserData = UsdUnreal::ObjectUtils::GetAssetUserData(GeometryCache))
 			{
 				UnrealToUsd::ConvertMetadata(
 					UserData,
diff --git a/Engine/Plugins/Importers/USDImporter/Source/USDExporter/Private/LevelSequenceExporterUSD.cpp b/Engine/Plugins/Importers/USDImporter/Source/USDExporter/Private/LevelSequenceExporterUSD.cpp
--- a/Engine/Plugins/Importers/USDImporter/Source/USDExporter/Private/LevelSequenceExporterUSD.cpp
+++ b/Engine/Plugins/Importers/USDImporter/Source/USDExporter/Private/LevelSequenceExporterUSD.cpp
@@ -11,6 +11,7 @@
 #include "USDExporterModule.h"
 #include "USDLayerUtils.h"
 #include "USDLog.h"
+#include "USDObjectUtils.h"
 #include "USDOptionsWindow.h"
 #include "USDPrimConversion.h"
 #include "USDStageActor.h"
@@ -102,7 +103,6 @@ namespace UE::LevelSequenceExporterUSD::Private
 
 #if USE_USD_SDK
 
-
 	using FSpawnedInstanceKey = TPair<FGuid, int32>;
 
 	// Custom spawn register so that when DestroySpawnedObject is called while bDestroyingJustHides is true we
@@ -115,21 +115,23 @@ namespace UE::LevelSequenceExporterUSD::Private
 
 		bool bDestroyingJustHides = true;
 
-		virtual UObject* SpawnObject(const FGuid& Guid,
+		virtual UObject* SpawnObject(
+			const FGuid& Guid,
 			UMovieScene& MovieScene,
 			FMovieSceneSequenceIDRef TemplateID,
 			TSharedRef<const FSharedPlaybackState> SharedPlaybackState,
 			int32 BindingIndex
 		) override
 		{
-
 			FSpawnedInstanceKey InstanceKey(Guid, BindingIndex);
 
 			// Never spawn ASphereReflectionCapture actors. These are useless in USD anyway, and we run into
 			// trouble after we're done exporting them because on the tick where they're destroyed the editor
 			// will still attempt to update their captures and some downstream code doesn't like that their
 			// components are pending kill (check UE-167593 for more info)
-			if (Cast<ASphereReflectionCapture>(MovieSceneHelpers::GetObjectTemplate(MovieScene.GetTypedOuter<UMovieSceneSequence>(), Guid, SharedPlaybackState)))
+			if (Cast<ASphereReflectionCapture>(
+					MovieSceneHelpers::GetObjectTemplate(MovieScene.GetTypedOuter<UMovieSceneSequence>(), Guid, SharedPlaybackState)
+				))
 			{
 				return nullptr;
 			}
@@ -144,9 +146,8 @@ namespace UE::LevelSequenceExporterUSD::Private
 
 			TArray<UObject*>& ExistingInstancesForGuid = SpawnableInstances.FindOrAdd(InstanceKey);
 
-			TMap<FMovieSceneSequenceID, TMap<FSpawnedInstanceKey, int32>>& SequenceInstanceToSpawnableIndices = RootSequenceToSpawnableInstanceIndices.FindOrAdd(
-				RootSequence
-			);
+			TMap<FMovieSceneSequenceID, TMap<FSpawnedInstanceKey, int32>>& SequenceInstanceToSpawnableIndices = RootSequenceToSpawnableInstanceIndices
+																													.FindOrAdd(RootSequence);
 			TMap<FSpawnedInstanceKey, int32>& SpawnableIndices = SequenceInstanceToSpawnableIndices.FindOrAdd(TemplateID);
 
 			// Already have an instance of this spawnable for this movie scene sequence instance
@@ -234,7 +235,7 @@ namespace UE::LevelSequenceExporterUSD::Private
 				{
 					// Rename the spawn to a unique name or else in case of name collisions they will overwrite each other when writing
 					// animation data. The level exporter will rename actors to unique prims by itself though.
-					FString NewLabel = UsdUtils::GetUniqueName(SpawnedActor->GetActorLabel(), UsedActorLabels);
+					FString NewLabel = UsdUnreal::ObjectUtils::GetUniqueName(SpawnedActor->GetActorLabel(), UsedActorLabels);
 					if (NewLabel != SpawnedActor->GetActorLabel())
 					{
 						const bool bMarkDirty = false;
@@ -384,8 +385,8 @@ namespace UE::LevelSequenceExporterUSD::Private
 				if (Object)
 				{
 					// TODO: I think this will likely be fine for most cases, but I could see it potentially being problematic in all cases to destroy
-					// an object in the default way this way. However, to do this 'correctly' would also involve storing the UCustomBinding object in this process,
-					// which I'm skipping for now.
+					// an object in the default way this way. However, to do this 'correctly' would also involve storing the UCustomBinding object in
+					// this process, which I'm skipping for now.
 					DestroySpawnedObject(*Object, nullptr);
 				}
 			}
@@ -399,7 +400,9 @@ namespace UE::LevelSequenceExporterUSD::Private
 		UObject* GetExistingSpawn(const UMovieSceneSequence& RootSequence, FMovieSceneSequenceID SequenceID, const FSpawnedInstanceKey& InstanceKey)
 		{
 			int32 SpawnableIndex = INDEX_NONE;
-			if (TMap<FMovieSceneSequenceID, TMap<FSpawnedInstanceKey, int32>>* SequenceIDsToSpawns = RootSequenceToSpawnableInstanceIndices.Find(&RootSequence))
+			if (TMap<FMovieSceneSequenceID, TMap<FSpawnedInstanceKey, int32>>* SequenceIDsToSpawns = RootSequenceToSpawnableInstanceIndices.Find(
+					&RootSequence
+				))
 			{
 				if (TMap<FSpawnedInstanceKey, int32>* Spawns = SequenceIDsToSpawns->Find(SequenceID))
 				{
@@ -591,10 +594,17 @@ namespace UE::LevelSequenceExporterUSD::Private
 							LastGuid = BindingReference.ID;
 							BindingIndex = 0;
 						}
-						if (BindingReference.CustomBinding && BindingReference.CustomBinding->WillSpawnObject(Context.Sequencer->GetSharedPlaybackState()))
+						if (BindingReference.CustomBinding
+							&& BindingReference.CustomBinding->WillSpawnObject(Context.Sequencer->GetSharedPlaybackState()))
 						{
 							StaticCastSharedRef<FMovieSceneSpawnRegister>(Context.SpawnRegister)
-								->SpawnObject(BindingReference.ID, *MovieScene, SequenceInstance, Context.Sequencer->GetSharedPlaybackState(), BindingIndex++);
+								->SpawnObject(
+									BindingReference.ID,
+									*MovieScene,
+									SequenceInstance,
+									Context.Sequencer->GetSharedPlaybackState(),
+									BindingIndex++
+								);
 						}
 					}
 				}
@@ -847,9 +857,11 @@ namespace UE::LevelSequenceExporterUSD::Private
 							LastGuid = BindingReference.ID;
 							BindingIndex = 0;
 						}
-						if (BindingReference.CustomBinding && BindingReference.CustomBinding->WillSpawnObject(Context.Sequencer->GetSharedPlaybackState()))
+						if (BindingReference.CustomBinding
+							&& BindingReference.CustomBinding->WillSpawnObject(Context.Sequencer->GetSharedPlaybackState()))
 						{
-							BoundObject = Context.SpawnRegister->GetExistingSpawn(*RootSequence, SequenceInstance, FSpawnedInstanceKey(Guid, BindingIndex++));
+							BoundObject = Context.SpawnRegister
+											  ->GetExistingSpawn(*RootSequence, SequenceInstance, FSpawnedInstanceKey(Guid, BindingIndex++));
 							if (!BoundObject)
 							{
 								// This should never happen as we preemptively spawn everything:
@@ -1417,7 +1429,7 @@ namespace UE::LevelSequenceExporterUSD::Private
 
 		// Make sure we don't overwrite a file we just wrote *during this export*.
 		// Overwriting other files is OK, as we want to allow a "repeatedly export over the same files" workflow
-		FString UniqueFilePath = UsdUtils::GetUniqueName(FilePath, Context.UsedFilePaths);
+		FString UniqueFilePath = UsdUnreal::ObjectUtils::GetUniqueName(FilePath, Context.UsedFilePaths);
 
 		// Try exporting subsequences if needed
 		if (Context.ExportOptions)
@@ -1616,7 +1628,7 @@ namespace UE::LevelSequenceExporterUSD::Private
 
 			if (Context.ExportOptions->LevelExportOptions.MetadataOptions.bExportAssetMetadata)
 			{
-				if (UUsdAssetUserData* UserData = UsdUtils::GetAssetUserData(Cast<ULevelSequence>(&MovieSceneSequence)))
+				if (UUsdAssetUserData* UserData = UsdUnreal::ObjectUtils::GetAssetUserData(Cast<ULevelSequence>(&MovieSceneSequence)))
 				{
 					UnrealToUsd::ConvertMetadata(
 						UserData,
@@ -1800,7 +1812,7 @@ bool ULevelSequenceExporterUsd::ExportBinary(
 	UAssetEditorSubsystem* AssetEditorSubsystem = nullptr;
 	{
 		AssetEditorSubsystem = GEditor->GetEditorSubsystem<UAssetEditorSubsystem>();
-		if (AssetEditorSubsystem)
+		if (AssetEditorSubsystem && !IsEngineExitRequested())
 		{
 			for (const TWeakPtr<ISequencer>& Sequencer : FLevelEditorSequencerIntegration::Get().GetSequencers())
 			{
diff --git a/Engine/Plugins/Importers/USDImporter/Source/USDExporter/Private/MaterialExporterUSD.cpp b/Engine/Plugins/Importers/USDImporter/Source/USDExporter/Private/MaterialExporterUSD.cpp
--- a/Engine/Plugins/Importers/USDImporter/Source/USDExporter/Private/MaterialExporterUSD.cpp
+++ b/Engine/Plugins/Importers/USDImporter/Source/USDExporter/Private/MaterialExporterUSD.cpp
@@ -10,6 +10,7 @@
 #include "USDGeomMeshConversion.h"
 #include "USDLog.h"
 #include "USDMetadataExportOptions.h"
+#include "USDObjectUtils.h"
 #include "USDOptionsWindow.h"
 #include "USDPrimConversion.h"
 #include "USDShadeConversion.h"
@@ -372,7 +373,7 @@ bool UMaterialExporterUsd::ExportMaterial(
 
 	if (MetadataOptions.bExportAssetMetadata)
 	{
-		if (UUsdAssetUserData* UserData = UsdUtils::GetAssetUserData(&Material))
+		if (UUsdAssetUserData* UserData = UsdUnreal::ObjectUtils::GetAssetUserData(&Material))
 		{
 			UnrealToUsd::ConvertMetadata(UserData, RootPrim, MetadataOptions.BlockedPrefixFilters, MetadataOptions.bInvertFilters);
 		}
@@ -473,7 +474,7 @@ bool UMaterialExporterUsd::ExportMaterialsForStage(
 		FString MaterialFilePath = FPaths::Combine(FPaths::GetPath(UExporter::CurrentFilename), FPaths::GetBaseFilename(MaterialPathName));
 
 		// "C:/MyFolder/Export/Blue_4"
-		FString FinalPathNoExt = UsdUtils::GetUniqueName(MaterialFilePath, UsedFilePathsWithoutExt);
+		FString FinalPathNoExt = UsdUnreal::ObjectUtils::GetUniqueName(MaterialFilePath, UsedFilePathsWithoutExt);
 
 		// "C:/MyFolder/Export/Blue_4.usda"
 		FString FinalPath = FString::Printf(TEXT("%s.%s"), *FinalPathNoExt, *ExtensionNoDot);
diff --git a/Engine/Plugins/Importers/USDImporter/Source/USDExporter/Private/SkeletalMeshExporterUSD.cpp b/Engine/Plugins/Importers/USDImporter/Source/USDExporter/Private/SkeletalMeshExporterUSD.cpp
--- a/Engine/Plugins/Importers/USDImporter/Source/USDExporter/Private/SkeletalMeshExporterUSD.cpp
+++ b/Engine/Plugins/Importers/USDImporter/Source/USDExporter/Private/SkeletalMeshExporterUSD.cpp
@@ -8,6 +8,7 @@
 #include "USDConversionUtils.h"
 #include "USDExporterModule.h"
 #include "USDLog.h"
+#include "USDObjectUtils.h"
 #include "USDOptionsWindow.h"
 #include "USDPrimConversion.h"
 #include "USDSkeletalDataConversion.h"
@@ -84,7 +85,7 @@ namespace UE::SkeletalMeshExporterUSD::Private
 			// just write out the slots with UsdGeomSubsets named "Section0", "Section1", ..., "SectionN" anyway
 		}
 	}
-}
+}	 // namespace UE::SkeletalMeshExporterUSD::Private
 
 USkeletalMeshExporterUsd::USkeletalMeshExporterUsd()
 {
@@ -333,7 +334,7 @@ bool USkeletalMeshExporterUsd::ExportBinary(
 
 		if (Options->MetadataOptions.bExportAssetMetadata)
 		{
-			if (UUsdAssetUserData* UserData = UsdUtils::GetAssetUserData(SkeletalMesh))
+			if (UUsdAssetUserData* UserData = UsdUnreal::ObjectUtils::GetAssetUserData(SkeletalMesh))
 			{
 				UnrealToUsd::ConvertMetadata(
 					UserData,
@@ -345,7 +346,7 @@ bool USkeletalMeshExporterUsd::ExportBinary(
 
 			if (USkeleton* Skeleton = SkeletalMesh->GetSkeleton())
 			{
-				if (UUsdAssetUserData* UserData = UsdUtils::GetAssetUserData(Skeleton))
+				if (UUsdAssetUserData* UserData = UsdUnreal::ObjectUtils::GetAssetUserData(Skeleton))
 				{
 					if (UserData->StageIdentifierToMetadata.Num() > 0)
 					{
diff --git a/Engine/Plugins/Importers/USDImporter/Source/USDExporter/Private/StaticMeshExporterUSD.cpp b/Engine/Plugins/Importers/USDImporter/Source/USDExporter/Private/StaticMeshExporterUSD.cpp
--- a/Engine/Plugins/Importers/USDImporter/Source/USDExporter/Private/StaticMeshExporterUSD.cpp
+++ b/Engine/Plugins/Importers/USDImporter/Source/USDExporter/Private/StaticMeshExporterUSD.cpp
@@ -2,16 +2,14 @@
 
 #include "StaticMeshExporterUSD.h"
 
-#include "EngineAnalytics.h"
 #include "MaterialExporterUSD.h"
-#include "Misc/EngineVersion.h"
 #include "StaticMeshExporterUSDOptions.h"
-#include "StaticMeshResources.h"
 #include "USDClassesModule.h"
 #include "USDConversionUtils.h"
 #include "USDExporterModule.h"
 #include "USDGeomMeshConversion.h"
 #include "USDLog.h"
+#include "USDObjectUtils.h"
 #include "USDOptionsWindow.h"
 #include "USDPrimConversion.h"
 #include "USDTypesConversion.h"
@@ -24,6 +22,9 @@
 
 #include "AssetExportTask.h"
 #include "Engine/StaticMesh.h"
+#include "EngineAnalytics.h"
+#include "Misc/EngineVersion.h"
+#include "StaticMeshResources.h"
 
 namespace UE::StaticMeshExporterUSD::Private
 {
@@ -82,7 +83,7 @@ namespace UE::StaticMeshExporterUSD::Private
 			// just write out the slots with UsdGeomSubsets named "Section0", "Section1", ..., "SectionN" anyway
 		}
 	}
-}
+}	 // namespace UE::StaticMeshExporterUSD::Private
 
 bool UStaticMeshExporterUsd::IsUsdAvailable()
 {
@@ -336,7 +337,7 @@ bool UStaticMeshExporterUsd::ExportBinary(UObject* Object, const TCHAR* Type, FA
 
 		if (Options->MetadataOptions.bExportAssetMetadata)
 		{
-			if (UUsdAssetUserData* UserData = UsdUtils::GetAssetUserData(StaticMesh))
+			if (UUsdAssetUserData* UserData = UsdUnreal::ObjectUtils::GetAssetUserData(StaticMesh))
 			{
 				UnrealToUsd::ConvertMetadata(
 					UserData,
diff --git a/Engine/Plugins/Importers/USDImporter/Source/USDExporter/Private/USDConversionBlueprintLibrary.cpp b/Engine/Plugins/Importers/USDImporter/Source/USDExporter/Private/USDConversionBlueprintLibrary.cpp
--- a/Engine/Plugins/Importers/USDImporter/Source/USDExporter/Private/USDConversionBlueprintLibrary.cpp
+++ b/Engine/Plugins/Importers/USDImporter/Source/USDExporter/Private/USDConversionBlueprintLibrary.cpp
@@ -10,6 +10,7 @@
 #include "USDExporterModule.h"
 #include "USDLayerUtils.h"
 #include "USDLog.h"
+#include "USDObjectUtils.h"
 #include "USDValueConversion.h"
 
 #include "UsdWrappers/SdfLayer.h"
@@ -767,12 +768,12 @@ TArray<FString> UUsdConversionBlueprintLibrary::DuplicatePrims(
 
 UUsdAssetUserData* UUsdConversionBlueprintLibrary::GetUsdAssetUserData(UObject* Object)
 {
-	return UsdUtils::GetAssetUserData(Object);
+	return UsdUnreal::ObjectUtils::GetAssetUserData(Object);
 }
 
 bool UUsdConversionBlueprintLibrary::SetUsdAssetUserData(UObject* Object, UUsdAssetUserData* AssetUserData)
 {
-	return UsdUtils::SetAssetUserData(Object, AssetUserData);
+	return UsdUnreal::ObjectUtils::SetAssetUserData(Object, AssetUserData);
 }
 
 namespace UE::UsdConversionBlueprintLibrary::Private
diff --git a/Engine/Plugins/Importers/USDImporter/Source/USDSchemas/Private/Custom/MDLUSDShadeMaterialTranslator.cpp b/Engine/Plugins/Importers/USDImporter/Source/USDSchemas/Private/Custom/MDLUSDShadeMaterialTranslator.cpp
--- a/Engine/Plugins/Importers/USDImporter/Source/USDSchemas/Private/Custom/MDLUSDShadeMaterialTranslator.cpp
+++ b/Engine/Plugins/Importers/USDImporter/Source/USDSchemas/Private/Custom/MDLUSDShadeMaterialTranslator.cpp
@@ -22,54 +22,48 @@
 #include "Materials/MaterialInstance.h"
 #include "Materials/MaterialInstanceConstant.h"
 #include "MaterialShared.h"
-#include "Misc/Paths.h"
 #include "MDLImporterOptions.h"
 #include "MDLMaterialImporter.h"
+#include "Misc/Paths.h"
 #include "UObject/StrongObjectPtr.h"
 
 #include "USDIncludesStart.h"
-	#include "pxr/usd/sdf/assetPath.h"
-	#include "pxr/usd/usdShade/material.h"
+#include "pxr/usd/sdf/assetPath.h"
+#include "pxr/usd/usdShade/material.h"
 #include "USDIncludesEnd.h"
 
 #define LOCTEXT_NAMESPACE "MDLUSDShadeMaterialTranslator"
 
 FName FMdlUsdShadeMaterialTranslator::MdlRenderContext = TEXT("mdl");
 
-namespace UE
+namespace UE::MDLShadeTranslator::Private
 {
-	namespace MDLShadeTranslatorImpl
+	void NotifyIfMaterialNeedsVirtualTextures(UMaterialInterface* MaterialInterface)
 	{
-		namespace Private
+		if (UMaterial* Material = Cast<UMaterial>(MaterialInterface))
 		{
-			void NotifyIfMaterialNeedsVirtualTextures( UMaterialInterface* MaterialInterface )
-			{
-				if ( UMaterial* Material = Cast<UMaterial>( MaterialInterface ) )
-				{
-					TArray<UTexture*> UsedTextures;
-					const bool bAllQualityLevels = true;
-					const bool bAllFeatureLevels = true;
-					Material->GetUsedTextures( UsedTextures, EMaterialQualityLevel::High, bAllQualityLevels, ERHIFeatureLevel::SM5, bAllFeatureLevels );
+			TArray<UTexture*> UsedTextures;
+			const bool bAllQualityLevels = true;
+			const bool bAllFeatureLevels = true;
+			Material->GetUsedTextures(UsedTextures, EMaterialQualityLevel::High, bAllQualityLevels, ERHIFeatureLevel::SM5, bAllFeatureLevels);
 
-					for ( UTexture* UsedTexture : UsedTextures )
-					{
-						UsdUtils::NotifyIfVirtualTexturesNeeded( UsedTexture );
-					}
-				}
-				else if ( UMaterialInstance* MaterialInstance = Cast<UMaterialInstance>( MaterialInterface ) )
+			for (UTexture* UsedTexture : UsedTextures)
+			{
+				UsdUtils::NotifyIfVirtualTexturesNeeded(UsedTexture);
+			}
+		}
+		else if (UMaterialInstance* MaterialInstance = Cast<UMaterialInstance>(MaterialInterface))
+		{
+			for (const FTextureParameterValue& TextureValue : MaterialInstance->TextureParameterValues)
+			{
+				if (UTexture* Texture = TextureValue.ParameterValue)
 				{
-					for ( const FTextureParameterValue& TextureValue : MaterialInstance->TextureParameterValues )
-					{
-						if ( UTexture* Texture = TextureValue.ParameterValue )
-						{
-							UsdUtils::NotifyIfVirtualTexturesNeeded( Texture );
-						}
-					}
+					UsdUtils::NotifyIfVirtualTexturesNeeded(Texture);
 				}
 			}
 		}
 	}
-}
+}	 // namespace UE::MDLShadeTranslator::Private
 
 void FMdlUsdShadeMaterialTranslator::CreateAssets()
 {
@@ -78,7 +72,7 @@ void FMdlUsdShadeMaterialTranslator::CreateAssets()
 	//   info:mdl:sourceAsset:subIdentifier -> Name of the material in the MDL file
 	//   inputs -> material parameters
 
-	if ( Context->RenderContext != MdlRenderContext )
+	if (Context->RenderContext != MdlRenderContext)
 	{
 		Super::CreateAssets();
 		return;
@@ -88,25 +82,30 @@ void FMdlUsdShadeMaterialTranslator::CreateAssets()
 	{
 		if (!Context->InfoCache->IsMaterialUsed(PrimPath))
 		{
-			UE_LOG(LogUsd, Verbose, TEXT("Skipping creating assets for material prim '%s' as it is not currently bound by any prim."), *PrimPath.GetString());
+			UE_LOG(
+				LogUsd,
+				Verbose,
+				TEXT("Skipping creating assets for material prim '%s' as it is not currently bound by any prim."),
+				*PrimPath.GetString()
+			);
 			return;
 		}
 	}
 
 	FScopedUsdAllocs UsdAllocs;
 
-	pxr::UsdShadeMaterial ShadeMaterial( GetPrim() );
+	pxr::UsdShadeMaterial ShadeMaterial(GetPrim());
 
-	if ( !ShadeMaterial )
+	if (!ShadeMaterial)
 	{
 		return;
 	}
 
-	const pxr::TfToken MdlToken = UnrealToUsd::ConvertToken( *MdlRenderContext.ToString() ).Get();
+	const pxr::TfToken MdlToken = UnrealToUsd::ConvertToken(*MdlRenderContext.ToString()).Get();
 
-	pxr::UsdShadeShader SurfaceShader = ShadeMaterial.ComputeSurfaceSource( MdlToken );
+	pxr::UsdShadeShader SurfaceShader = ShadeMaterial.ComputeSurfaceSource(MdlToken);
 
-	if ( !SurfaceShader )
+	if (!SurfaceShader)
 	{
 		Super::CreateAssets();
 		return;
@@ -115,76 +114,132 @@ void FMdlUsdShadeMaterialTranslator::CreateAssets()
 	const FString MdlRootPath = Context->Stage.GetRootLayer().GetRealPath();
 
 	pxr::SdfAssetPath SurfaceSourceAssetPath;
-	SurfaceShader.GetSourceAsset( &SurfaceSourceAssetPath, MdlToken );
+	SurfaceShader.GetSourceAsset(&SurfaceSourceAssetPath, MdlToken);
 
-	if ( SurfaceSourceAssetPath.GetAssetPath().empty() )
+	if (SurfaceSourceAssetPath.GetAssetPath().empty())
 	{
 		// Old Mdl Schema
-		const pxr::TfToken ModuleToken( "module" );
-		const pxr::UsdAttribute& MDLModule = SurfaceShader.GetPrim().GetAttribute( ModuleToken );
+		const pxr::TfToken ModuleToken("module");
+		const pxr::UsdAttribute& MDLModule = SurfaceShader.GetPrim().GetAttribute(ModuleToken);
 
-		if ( MDLModule.GetTypeName().GetAsToken() == pxr::SdfValueTypeNames->Asset )
+		if (MDLModule.GetTypeName().GetAsToken() == pxr::SdfValueTypeNames->Asset)
 		{
-			SurfaceSourceAssetPath = UsdUtils::GetUsdValue< pxr::SdfAssetPath >( MDLModule, Context->Time );
+			SurfaceSourceAssetPath = UsdUtils::GetUsdValue<pxr::SdfAssetPath>(MDLModule, Context->Time);
 		}
 	}
 
-	const FString MdlAbsoluteAssetPath = UsdToUnreal::ConvertString( SurfaceSourceAssetPath.GetResolvedPath() );
+	const FString MdlAbsoluteAssetPath = UsdToUnreal::ConvertString(SurfaceSourceAssetPath.GetResolvedPath());
 
-	const FString MdlModuleName = [ &MdlAbsoluteAssetPath, &MdlRootPath ]()
+	const FString MdlModuleName = [&MdlAbsoluteAssetPath, &MdlRootPath]()
 	{
 		FString ModuleRelativePath = MdlAbsoluteAssetPath;
-		FPaths::MakePathRelativeTo( ModuleRelativePath, *MdlRootPath );
+		FPaths::MakePathRelativeTo(ModuleRelativePath, *MdlRootPath);
 
-		FString ModuleName = UE::Mdl::Util::ConvertFilePathToModuleName( *ModuleRelativePath );
+		FString ModuleName = UE::Mdl::Util::ConvertFilePathToModuleName(*ModuleRelativePath);
 		return ModuleName;
 	}();
 
-	if (!MdlModuleName.IsEmpty() && Context->AssetCache)
+	if (!MdlModuleName.IsEmpty() && Context->UsdAssetCache)
 	{
 		pxr::TfToken MdlDefinitionToken;
-		SurfaceShader.GetSourceAssetSubIdentifier( &MdlDefinitionToken, MdlToken );
+		SurfaceShader.GetSourceAssetSubIdentifier(&MdlDefinitionToken, MdlToken);
 
-		const FString MdlDefinitionName = UsdToUnreal::ConvertToken( MdlDefinitionToken );
+		const FString MdlDefinitionName = UsdToUnreal::ConvertToken(MdlDefinitionToken);
 
 		const FString MdlFullName = MdlModuleName + TEXT("::") + MdlDefinitionName;
 		const FString MdlFullInstanceName = MdlFullName + TEXT("_Instance");
 		const FString HashPrefix = UsdUtils::GetAssetHashPrefix(GetPrim(), Context->bReuseIdenticalAssets);
+		const FString MaterialHash = HashPrefix + MdlFullName;
+		const FString MdlSearchPath = FPaths::GetPath(Context->Stage.GetRootLayer().GetRealPath());
+
+		TSet<UObject*> Dependencies;
+
+		FScopedUnrealAllocs UEAllocs;
+
+		// Create reference material
+		//
+		// The FMdlMaterialImporter will (strangely) take the provided UPackage "ParentPackage" and instead just create
+		// another package, using "ParentPackage.GetName() / MdlMaterialName" as it's path... This is not ideal especially
+		// since it will lead to asset collisions if we're not reusing assets via hash. Here we just put everything
+		// on the transient package and rename the asset into the target destination package instead, like we do for MaterialX
+		//
+		// We have to do all of that inside the lambda though because we want to reuse that reference material if we have one
+		// cached already
+		bool bCreatedMaterial = false;
+		UMaterialInterface* MdlMaterial = Context->UsdAssetCache->GetOrCreateCustomCachedAsset<UMaterialInterface>(
+			MaterialHash,
+			MdlFullName,
+			Context->ObjectFlags,
+			[&MdlSearchPath, &MdlModuleName, &MdlDefinitionName](UPackage* Outer, FName SanitizedName, EObjectFlags FlagsToUse) -> UMaterialInterface*
+			{
+				// Add the USD root as a search path for MDL
+				FMdlMaterialImporter::FScopedSearchPath UsdDirMdlSearchPath(MdlSearchPath);
+
+				TStrongObjectPtr<UMDLImporterOptions> ImportOptions(NewObject<UMDLImporterOptions>());
+
+				UPackage* ParentPackage = GetTransientPackage();
+				UMaterialInterface* ReferenceMaterial = FMdlMaterialImporter::ImportMaterialFromModule(
+					ParentPackage,
+					FlagsToUse,
+					MdlModuleName,
+					MdlDefinitionName,
+					*ImportOptions.Get()
+				);
+				if (!ensure(ReferenceMaterial))
+				{
+					return nullptr;
+				}
 
-		UMaterialInterface* MdlMaterial = Cast< UMaterialInterface >( Context->AssetCache->GetCachedAsset( HashPrefix + MdlFullName ) );
-		if ( !MdlMaterial )
+				// Rename the UMaterialInterface into the target UPackage the asset cache created for us.
+				// SanitizedName will already match it.
+				const bool bRenamed = ReferenceMaterial->Rename(
+					*SanitizedName.ToString(),
+					Outer,
+					REN_NonTransactional | REN_DontCreateRedirectors | REN_ForceNoResetLoaders
+				);
+				ensure(bRenamed);
+
+				// Let's not trust the flags the MDLImporter used and just use our own instead
+				ReferenceMaterial->ClearFlags(ReferenceMaterial->GetFlags());
+				ReferenceMaterial->SetFlags(FlagsToUse);
+
+				return ReferenceMaterial;
+			},
+			&bCreatedMaterial
+		);
+		if (bCreatedMaterial && MdlMaterial)
 		{
-			FScopedUnrealAllocs UnrealAllocs;
-
-			// Add the USD root as a search path for MDL
-			FMdlMaterialImporter::FScopedSearchPath UsdDirMdlSearchPath( FPaths::GetPath( Context->Stage.GetRootLayer().GetRealPath() ) );
-
-			TStrongObjectPtr< UMDLImporterOptions > ImportOptions( NewObject< UMDLImporterOptions >() );
-
-			MdlMaterial = FMdlMaterialImporter::ImportMaterialFromModule( GetTransientPackage(), Context->ObjectFlags | RF_Transient, MdlModuleName, MdlDefinitionName, *ImportOptions.Get() );
-
-			if ( MdlMaterial )
-			{
-				UE::MDLShadeTranslatorImpl::Private::NotifyIfMaterialNeedsVirtualTextures( MdlMaterial );
-			}
-			else
-			{
-				FUsdLogManager::LogMessage(
-					EMessageSeverity::Warning,
-					FText::Format( LOCTEXT("UsdMdlConversionFailed", "Failed to create MDL material for prim {0}."), FText::FromString( PrimPath.GetString() ) ) );
-			}
+			UE::MDLShadeTranslator::Private::NotifyIfMaterialNeedsVirtualTextures(MdlMaterial);
+			Dependencies.Append(IUsdClassesModule::GetAssetDependencies(MdlMaterial));
+		}
+		else if (!MdlMaterial)
+		{
+			FUsdLogManager::LogMessage(
+				EMessageSeverity::Warning,
+				FText::Format(
+					LOCTEXT("UsdMdlConversionFailed", "Failed to create MDL material for prim {0}."),
+					FText::FromString(PrimPath.GetString())
+				)
+			);
 		}
 
-		UMaterialInstanceConstant* MdlMaterialInstance = Cast< UMaterialInstanceConstant >(Context->AssetCache->GetCachedAsset(HashPrefix + MdlFullInstanceName));
-		if ( !MdlMaterialInstance && MdlMaterial )
+		// Create material instance
+		bool bCreatedInstance = false;
+		const FString MaterialInstanceHash = HashPrefix + MdlFullInstanceName;
+		UMaterialInstanceConstant* MdlMaterialInstance = Context->UsdAssetCache->GetOrCreateCachedAsset<UMaterialInstanceConstant>(
+			MaterialInstanceHash,
+			MdlFullInstanceName,
+			Context->ObjectFlags,
+			&bCreatedInstance
+		);
+		if (bCreatedInstance && MdlMaterialInstance)
 		{
-			MdlMaterialInstance = NewObject< UMaterialInstanceConstant >(GetTransientPackage(), NAME_None, Context->ObjectFlags | RF_Transient);
-			MdlMaterialInstance->SetParentEditorOnly( MdlMaterial );
+			MdlMaterialInstance->SetParentEditorOnly(MdlMaterial);
 
 			UsdToUnreal::ConvertShadeInputsToParameters(
 				ShadeMaterial,
 				*MdlMaterialInstance,
-				Context->AssetCache.Get(),
+				Context->UsdAssetCache.Get(),
 				*Context->RenderContext.ToString(),
 				Context->bReuseIdenticalAssets
 			);
@@ -193,21 +248,85 @@ void FMdlUsdShadeMaterialTranslator::CreateAssets()
 			// render states for some components, and UWorld::AddToWorld calls FScene::AddPrimitive which expects the component to not have
 			// primitives yet
 			FMaterialUpdateContext::EOptions::Type Options = FMaterialUpdateContext::EOptions::Default;
-			if ( Context->Level->bIsAssociatingLevel )
+			if (Context->Level && Context->Level->bIsAssociatingLevel)
 			{
-				Options = ( FMaterialUpdateContext::EOptions::Type ) ( Options & ~FMaterialUpdateContext::EOptions::RecreateRenderStates );
+				Options = (FMaterialUpdateContext::EOptions::Type)(Options & ~FMaterialUpdateContext::EOptions::RecreateRenderStates);
 			}
 
-			FMaterialUpdateContext UpdateContext( Options, GMaxRHIShaderPlatform );
-			UpdateContext.AddMaterialInstance( MdlMaterialInstance );
-			MdlMaterialInstance->PreEditChange( nullptr );
+			FMaterialUpdateContext UpdateContext(Options, GMaxRHIShaderPlatform);
+			UpdateContext.AddMaterialInstance(MdlMaterialInstance);
+			MdlMaterialInstance->PreEditChange(nullptr);
 			MdlMaterialInstance->PostEditChange();
 
-			UE::MDLShadeTranslatorImpl::Private::NotifyIfMaterialNeedsVirtualTextures( MdlMaterialInstance );
+			UE::MDLShadeTranslator::Private::NotifyIfMaterialNeedsVirtualTextures(MdlMaterialInstance);
+			Dependencies.Append(IUsdClassesModule::GetAssetDependencies(MdlMaterialInstance));
+		}
+
+		// If we created a material here, we could have created some textures. Let's cache those too
+		// so that PostImportMaterial can take care of them
+		for (UObject* Object : Dependencies)
+		{
+			if (UTexture* Texture = Cast<UTexture>(Object))
+			{
+				// We have to watch out because it seems like FMdlMaterialImporter can generate some materials using
+				// default engine textures like FlatNormal. We don't want to cache those.
+				// Note that we can only test for that by checking whether the asset has been saved or not because the
+				// MDL importer (unfortunately) can't use our asset cache with its previously created materials and textures...
+				// This means that every time an MDL has to create a new texture it will re-read the image from disk again and create
+				// a new (as of yet unsaved) asset. This means if the asset is saved it must have been an engine texture.
+				// Also Note: I think the MDLImporter has a bug here where it will put those textures in paths like
+				// "/Engine/TextureName"... We can take them anyway though.
+				UPackage* TexturePackage = Texture->GetOutermost();
+				if (TexturePackage && TexturePackage->GetFileSize() > 0)
+				{
+					continue;
+				}
+
+				// Note how we add MaterialHash to the texture hash: This because unfortunately there's no way to get
+				// the MDLImporter to easily use our cached textures (it does have a ResourcesDir on UMDLImporterOptions but
+				// those seem to be file directory paths instead?). This means each MDL material must be allowed to create and
+				// generate new textures for itself, which we must guarantee don't collide with anything inside the asset cache.
+				const FString FilePath = Texture->AssetImportData ? Texture->AssetImportData->GetFirstFilename() : TEXT("");
+				const FString PrefixedTextureHash = HashPrefix + MaterialHash
+													+ UsdUtils::GetTextureHash(
+														FilePath,
+														Texture->SRGB,
+														Texture->CompressionSettings,
+														Texture->GetTextureAddressX(),
+														Texture->GetTextureAddressY()
+													);
+
+				bool bCreatedTexture = false;
+				UObject* PackagedTexture = Context->UsdAssetCache->GetOrCreateCustomCachedAsset(
+					PrefixedTextureHash,
+					Texture->GetClass(),
+					Texture->GetName(),
+					Texture->GetFlags(),
+					[Texture](UPackage* Outer, FName SanitizedName, EObjectFlags FlagsToUse) -> UObject*
+					{
+						// Rename the asset into the target UPackage the asset cache created for us.
+						// SanitizedName will already match it.
+						const bool bRenamed = Texture->Rename(
+							*SanitizedName.ToString(),
+							Outer,
+							REN_NonTransactional | REN_DontCreateRedirectors | REN_ForceNoResetLoaders
+						);
+						ensure(bRenamed);
+
+						// Let's not trust the flags the MDLImporter used and just use our own instead
+						Texture->ClearFlags(Texture->GetFlags());
+						Texture->SetFlags(FlagsToUse);
+
+						return Texture;
+					},
+					&bCreatedTexture
+				);
+				ensure(PackagedTexture && bCreatedTexture);
+			}
 		}
 
-		PostImportMaterial(HashPrefix + MdlFullName, MdlMaterial);
-		PostImportMaterial(HashPrefix + MdlFullInstanceName, MdlMaterialInstance);
+		PostImportMaterial(MaterialHash, MdlMaterial);
+		PostImportMaterial(MaterialInstanceHash, MdlMaterialInstance);
 	}
 	else
 	{
@@ -217,4 +336,4 @@ void FMdlUsdShadeMaterialTranslator::CreateAssets()
 
 #undef LOCTEXT_NAMESPACE
 
-#endif // #if USE_USD_SDK && WITH_EDITOR
+#endif	  // #if USE_USD_SDK && WITH_EDITOR
diff --git a/Engine/Plugins/Importers/USDImporter/Source/USDSchemas/Private/Custom/MaterialXUSDShadeMaterialTranslator.cpp b/Engine/Plugins/Importers/USDImporter/Source/USDSchemas/Private/Custom/MaterialXUSDShadeMaterialTranslator.cpp
--- a/Engine/Plugins/Importers/USDImporter/Source/USDSchemas/Private/Custom/MaterialXUSDShadeMaterialTranslator.cpp
+++ b/Engine/Plugins/Importers/USDImporter/Source/USDSchemas/Private/Custom/MaterialXUSDShadeMaterialTranslator.cpp
@@ -2,12 +2,13 @@
 
 #include "Custom/MaterialXUSDShadeMaterialTranslator.h"
 #if USE_USD_SDK
-#include "USDAssetCache.h"
+#include "USDAssetCache3.h"
 #include "USDAssetUserData.h"
 #include "USDClassesModule.h"
 #include "USDConversionUtils.h"
 #include "USDErrorUtils.h"
 #include "USDLog.h"
+#include "USDObjectUtils.h"
 #include "USDShadeConversion.h"
 #include "USDTypesConversion.h"
 
@@ -30,9 +31,9 @@
 #include "UObject/StrongObjectPtr.h"
 
 #include "USDIncludesStart.h"
-	#include "pxr/usd/sdf/assetPath.h"
-	#include "pxr/usd/usd/primCompositionQuery.h"
-	#include "pxr/usd/usdShade/material.h"
+#include "pxr/usd/sdf/assetPath.h"
+#include "pxr/usd/usd/primCompositionQuery.h"
+#include "pxr/usd/usdShade/material.h"
 #include "USDIncludesEnd.h"
 
 #if WITH_EDITOR
@@ -45,7 +46,10 @@ static bool GUseInterchangeMaterialXTranslator = true;
 static FAutoConsoleVariableRef CvarUseInterchangeMaterialXTranslator(
 	TEXT("USD.UseInterchangeMaterialXTranslator"),
 	GUseInterchangeMaterialXTranslator,
-	TEXT("Whether to translate MaterialX materials referenced by USD files with Unreal's own MaterialX importer. If instead this is false, we will try parsing the generated UsdShadeMaterial prims as generated by USD's usdMtlx plugin directly"));
+	TEXT(
+		"Whether to translate MaterialX materials referenced by USD files with Unreal's own MaterialX importer. If instead this is false, we will try parsing the generated UsdShadeMaterial prims as generated by USD's usdMtlx plugin directly"
+	)
+);
 
 FName FMaterialXUsdShadeMaterialTranslator::MaterialXRenderContext = TEXT("mtlx");
 
@@ -162,7 +166,8 @@ namespace UE::USDMaterialXTranslator::Private
 		const FString& MaterialXFilePath,
 		const FString& FileHash,
 		const pxr::UsdPrim& MaterialXReferencerPrim,
-		UUsdAssetCache2& AssetCache,
+		UUsdAssetCache3& AssetCache,
+		EObjectFlags ObjectFlags,
 		bool bReuseIdenticalAssets
 	)
 	{
@@ -179,93 +184,113 @@ namespace UE::USDMaterialXTranslator::Private
 		FImportAssetParameters InterchangeParameters;
 		InterchangeParameters.bIsAutomated = true;
 
-		// Get a temp content folder we can give interchange as it must have a destination
-		int32 PackageSuffix = FMath::RandRange(0, TNumericLimits<int32>::Max()-1);
+		const FString HashPrefix = UsdUtils::GetAssetHashPrefix(MaterialXReferencerPrim, bReuseIdenticalAssets);
+
+		// Annoyingly we have to make a new target folder for the interchange import, and then rename all the assets over.
+		// This because for direct imports we'll be dealing with an asset cache that is pointing at the transient package.
+		// If we naively give the transient package path to Interchange as the target import location, it will create packages
+		// inside of the transient package, which sounds like very bad news.
+		//
+		// Additionally, importing to a guaranteed empty folder and then renaming the assets into their target packages
+		// means that we avoid another annoyance: If Interchange finds an asset with the target name at the target location
+		// and of the target class, it will just silently stop importing that factory node and use that asset instead. There
+		// is no way for us to tell when it did this, so from our end it just looks like sometimes some assets will be missing
+		// from the import for no reason... when that happens we'll end up falling back to translating the Material prim with
+		// UsdShadeMaterialTranslator and get a broken black plastic material instead, and a warning on the output log.
+		// By importing into an empty folder we avoid this issue.
+		int32 PackageSuffix = FMath::RandRange(0, TNumericLimits<int32>::Max() - 1);
 		FString TempPackagePath = FString::Printf(TEXT("/Engine/USDImporter/Transient/%d"), PackageSuffix++);
 
 		// Make sure that content folder doesn't have any packages already in it
-		TArray<FString> Filenames;
-		while (FPackageName::FindPackagesInDirectory(Filenames, TempPackagePath))
 		{
-			TempPackagePath = FString::Printf(TEXT("/Engine/USDImporter/Transient/%d"), PackageSuffix++);
+			TArray<FString> Filenames;
+			while (FPackageName::FindPackagesInDirectory(Filenames, TempPackagePath))
+			{
+				TempPackagePath = FString::Printf(TEXT("/Engine/USDImporter/Transient/%d"), PackageSuffix++);
+			}
 		}
 
-		const FString HashPrefix = UsdUtils::GetAssetHashPrefix(MaterialXReferencerPrim, bReuseIdenticalAssets);
-
-		// Currently Interchange cannot import directly into the transient package, so we'll import into a temporary folder and then
-		// move everything over
 		InterchangeParameters.OnAssetsImportDoneNative.BindLambda(
-			[&AssetCache, &FileHash, &ReferencerPrimPath, &TempPackagePath, &MaterialXFilePath, &HashPrefix](const TArray<UObject*>& ImportedObjects)
+			[&AssetCache, &FileHash, &ReferencerPrimPath, &TempPackagePath, &MaterialXFilePath, &HashPrefix, ObjectFlags](
+				const TArray<UObject*>& ImportedObjects
+			)
 			{
-				UE_LOG(LogUsd, Verbose, TEXT("Translated '%d' assets from MaterialX file '%s'"), ImportedObjects.Num(), *MaterialXFilePath);
+				UE_LOG(LogUsd, Log, TEXT("Translated '%d' assets from MaterialX file '%s'"), ImportedObjects.Num(), *MaterialXFilePath);
 
 				// Move all our assets to the transient package
 				for (UObject* ImportedObject : ImportedObjects)
 				{
-					const FName NewName = MakeUniqueObjectName(
-						GetTransientPackage(),
-						ImportedObject->GetClass(),
-						ImportedObject->GetFName()
-					);
-					bool bSuccess = ImportedObject->Rename(
-						*NewName.ToString(),
-						GetTransientPackage(),
-						REN_NonTransactional | REN_DontCreateRedirectors | REN_ForceNoResetLoaders
-					);
-					if (!bSuccess)
-					{
-						UE_LOG(LogUsd, Warning, TEXT("Failed to rename imported MaterialX material '%s', generated for referencer prim '%s'."),
-							*ImportedObject->GetPathName(),
-							*ReferencerPrimPath
-						);
-						continue;
-					}
-
-					ImportedObject->SetFlags(RF_Transient);
-
+					FString PrefixedAssetHash;
 					if (UMaterialInterface* Material = Cast<UMaterialInterface>(ImportedObject))
 					{
 						// MaterialX names are unique, and can only have alphanumeric and the "_" character, so we should
 						// always have a solid enough mapping to assume UAsset name == Prim name == MaterialX name
-						FString NewPrefixedMaterialHash = HashPrefix + FileHash + TEXT("/") + Material->GetFName().GetPlainNameString();
-
-						// We generate all assets from the MaterialX file once, but it's possible we're just updating a
-						// single material prim here. If we were to cache all assets here, we'd potentially be trying
-						// to overwrite the existing assets that are being used by other prims and wouldn't otherwise
-						// be discarded, so make sure we don't do that
-						UMaterialInterface* ExistingMaterial = Cast<UMaterialInterface>(AssetCache.GetCachedAsset(NewPrefixedMaterialHash));
-						if (!ExistingMaterial)
-						{
-							AssetCache.CacheAsset(NewPrefixedMaterialHash, Material);
-						}
+						PrefixedAssetHash = HashPrefix + FileHash + TEXT("/") + Material->GetFName().GetPlainNameString();
 					}
 					else if (UTexture* Texture = Cast<UTexture>(ImportedObject))
 					{
 						const FString FilePath = Texture->AssetImportData ? Texture->AssetImportData->GetFirstFilename() : TEXT("");
-
-						FString PrefixedTextureHash = HashPrefix + UsdUtils::GetTextureHash(
-							FilePath,
-							Texture->SRGB,
-							Texture->CompressionSettings,
-							Texture->GetTextureAddressX(),
-							Texture->GetTextureAddressY()
-						);
-
-						// See comment on the analogous part of the Material case above
-						UTexture* ExistingTexture = Cast<UTexture>(AssetCache.GetCachedAsset(PrefixedTextureHash));
-						if (!ExistingTexture)
-						{
-							UUsdAssetUserData* UserData = NewObject<UUsdAssetUserData>(Texture, TEXT("USDAssetUserData"));
-							UserData->PrimPaths = {ReferencerPrimPath};
-							Texture->AddAssetUserData(UserData);
-
-							AssetCache.CacheAsset(PrefixedTextureHash, Texture);
-						}
+						PrefixedAssetHash = HashPrefix
+											+ UsdUtils::GetTextureHash(
+												FilePath,
+												Texture->SRGB,
+												Texture->CompressionSettings,
+												Texture->GetTextureAddressX(),
+												Texture->GetTextureAddressY()
+											);
 					}
 					else
 					{
 						ensureMsgf(false, TEXT("Asset type unsupported!"));
+						continue;
 					}
+
+					// We generate all assets from the MaterialX file once, but it's possible we're just updating a
+					// single material prim here. If we were to cache all assets here, we'd potentially be trying
+					// to overwrite the existing assets that are being used by other prims and wouldn't otherwise
+					// be discarded, so make sure we don't do that
+					const bool bAssetAlreadyCached = AssetCache.IsAssetTrackedByCache(AssetCache.GetCachedAssetPath(PrefixedAssetHash));
+					if (bAssetAlreadyCached)
+					{
+						continue;
+					}
+
+					// Cache the asset for the first time
+					bool bCreatedAsset = false;
+					UClass* ImportedObjectClass = ImportedObject->GetClass();
+					UObject* NewCachedAsset = AssetCache.GetOrCreateCustomCachedAsset(
+						PrefixedAssetHash,
+						ImportedObjectClass,
+						ImportedObject->GetName(),
+						ObjectFlags,
+						[ImportedObject](UPackage* Outer, FName SanitizedName, EObjectFlags FlagsToUse) -> UObject*
+						{
+							UPackage* InterchangePackage = ImportedObject->GetOutermost();
+
+							// Rename the UMaterialInterface into the target UPackage the asset cache created for us.
+							// SanitizedName will already match it.
+							const bool bRenamed = ImportedObject->Rename(
+								*SanitizedName.ToString(),
+								Outer,
+								REN_NonTransactional | REN_DontCreateRedirectors | REN_ForceNoResetLoaders
+							);
+							ensure(bRenamed);
+
+							ImportedObject->ClearFlags(ImportedObject->GetFlags());
+							ImportedObject->SetFlags(FlagsToUse);
+
+							// Get rid of the original package interchange made for this asset
+							if (InterchangePackage && InterchangePackage != GetTransientPackage())
+							{
+								InterchangePackage->MarkAsGarbage();
+								InterchangePackage->ClearFlags(RF_Public | RF_Standalone);
+							}
+
+							return ImportedObject;
+						},
+						&bCreatedAsset
+					);
+					ensure(bCreatedAsset && NewCachedAsset);
 				}
 
 				// Now that the asset cache took everything we imported we shouldn't have anything else
@@ -279,8 +304,8 @@ namespace UE::USDMaterialXTranslator::Private
 
 		return Manager.ImportAsset(TempPackagePath, SourceData, InterchangeParameters);
 	}
-#endif // WITH_EDITOR
-}
+#endif	  // WITH_EDITOR
+}	 // namespace UE::USDMaterialXTranslator::Private
 
 void FMaterialXUsdShadeMaterialTranslator::CreateAssets()
 {
@@ -347,14 +372,19 @@ void FMaterialXUsdShadeMaterialTranslator::CreateAssets()
 		return;
 	}
 
-	if (!Context->AssetCache || !Context->InfoCache)
+	if (!Context->UsdAssetCache || !Context->InfoCache)
 	{
 		return;
 	}
 
 	if (Context->bTranslateOnlyUsedMaterials && !Context->InfoCache->IsMaterialUsed(PrimPath))
 	{
-		UE_LOG(LogUsd, Verbose, TEXT("Skipping creating assets for material prim '%s' as it is not currently bound by any prim."), *PrimPath.GetString());
+		UE_LOG(
+			LogUsd,
+			Verbose,
+			TEXT("Skipping creating assets for material prim '%s' as it is not currently bound by any prim."),
+			*PrimPath.GetString()
+		);
 		return;
 	}
 
@@ -420,7 +450,14 @@ void FMaterialXUsdShadeMaterialTranslator::CreateAssets()
 
 	if (MaterialXFilePaths.IsEmpty() || !MaterialXReferencerPrim)
 	{
-		UE_LOG(LogUsd, Warning, TEXT("Recognized potential MaterialX materials on prim '%s', but failed to find a valid referenced MaterialX file. Reverting to parsing the generated Material prims instead."), *PrimPath.GetString());
+		UE_LOG(
+			LogUsd,
+			Warning,
+			TEXT(
+				"Recognized potential MaterialX materials on prim '%s', but failed to find a valid referenced MaterialX file. Reverting to parsing the generated Material prims instead."
+			),
+			*PrimPath.GetString()
+		);
 		Super::CreateAssets();
 		return;
 	}
@@ -443,7 +480,7 @@ void FMaterialXUsdShadeMaterialTranslator::CreateAssets()
 		}
 
 		FString MaterialAssetHashForThisFile = TargetHashPrefix + MaterialXHash + TargetHashSuffix;
-		ParsedMaterial = Cast<UMaterialInterface>(Context->AssetCache->GetCachedAsset(MaterialAssetHashForThisFile));
+		ParsedMaterial = Cast<UMaterialInterface>(Context->UsdAssetCache->GetCachedAsset(MaterialAssetHashForThisFile));
 		if (ParsedMaterial)
 		{
 			FoundMaterialAssetHash = MaterialAssetHashForThisFile;
@@ -464,7 +501,8 @@ void FMaterialXUsdShadeMaterialTranslator::CreateAssets()
 				MaterialXFilePath,
 				MaterialXFileHash,
 				MaterialXReferencerPrim,
-				*Context->AssetCache,
+				*Context->UsdAssetCache,
+				Context->ObjectFlags,
 				Context->bReuseIdenticalAssets
 			);
 
@@ -483,7 +521,7 @@ void FMaterialXUsdShadeMaterialTranslator::CreateAssets()
 			if (!ParsedMaterial)
 			{
 				FString MaterialAssetHashForThisFile = TargetHashPrefix + MaterialXFileHash + TargetHashSuffix;
-				ParsedMaterial = Cast<UMaterialInterface>(Context->AssetCache->GetCachedAsset(MaterialAssetHashForThisFile));
+				ParsedMaterial = Cast<UMaterialInterface>(Context->UsdAssetCache->GetCachedAsset(MaterialAssetHashForThisFile));
 				if (ParsedMaterial)
 				{
 					FoundMaterialAssetHash = MaterialAssetHashForThisFile;
@@ -520,9 +558,9 @@ void FMaterialXUsdShadeMaterialTranslator::CreateAssets()
 	}
 #else
 	Super::CreateAssets();
-#endif // WITH_EDITOR
+#endif	  // WITH_EDITOR
 }
 
 #undef LOCTEXT_NAMESPACE
 
-#endif // #if USE_USD_SDK
+#endif	  // #if USE_USD_SDK
diff --git a/Engine/Plugins/Importers/USDImporter/Source/USDSchemas/Private/MeshTranslationImpl.cpp b/Engine/Plugins/Importers/USDImporter/Source/USDSchemas/Private/MeshTranslationImpl.cpp
--- a/Engine/Plugins/Importers/USDImporter/Source/USDSchemas/Private/MeshTranslationImpl.cpp
+++ b/Engine/Plugins/Importers/USDImporter/Source/USDSchemas/Private/MeshTranslationImpl.cpp
@@ -2,7 +2,7 @@
 
 #include "MeshTranslationImpl.h"
 
-#include "USDAssetCache2.h"
+#include "USDAssetCache3.h"
 #include "USDAssetUserData.h"
 #include "USDClassesModule.h"
 #include "USDConversionUtils.h"
@@ -10,6 +10,7 @@
 #include "USDInfoCache.h"
 #include "USDLog.h"
 #include "USDMemory.h"
+#include "USDObjectUtils.h"
 #include "USDProjectSettings.h"
 #include "USDShadeConversion.h"
 #include "USDTypesConversion.h"
@@ -45,7 +46,11 @@ static_assert(
 
 namespace UE::MeshTranslationImplInternal::Private
 {
-	UMaterialInterface* CreateTwoSidedVersionOfMaterial(UMaterialInterface* OneSidedMat)
+	UMaterialInterface* GetOrCreateTwoSidedVersionOfMaterial(
+		UMaterialInterface* OneSidedMat,
+		const FString& PrefixedTwoSidedHash,
+		UUsdAssetCache3& AssetCache
+	)
 	{
 		if (!OneSidedMat)
 		{
@@ -65,13 +70,10 @@ namespace UE::MeshTranslationImplInternal::Private
 			ReferenceMaterialTwoSided = MeshTranslationImpl::GetTwoSidedVersionOfReferencePreviewSurfaceMaterial(ReferenceMaterial);
 		}
 
-		const FName NewInstanceName = MakeUniqueObjectName(
-			GetTransientPackage(),
-			UMaterialInstance::StaticClass(),
-			*(OneSidedMat->GetName() + UnrealIdentifiers::TwoSidedMaterialSuffix)
-		);
+		const FString& DesiredMaterialName = OneSidedMat->GetName() + UnrealIdentifiers::TwoSidedMaterialSuffix;
+		const EObjectFlags DesiredFlags = OneSidedMat->GetFlags();
 
-		UMaterialInstanceDynamic* MID = Cast<UMaterialInstanceDynamic>(OneSidedMat);
+		UMaterialInstance* MI = Cast<UMaterialInstance>(OneSidedMat);
 #if WITH_EDITOR
 		UMaterialInstanceConstant* MIC = Cast<UMaterialInstanceConstant>(OneSidedMat);
 
@@ -79,14 +81,19 @@ namespace UE::MeshTranslationImplInternal::Private
 		// Just create an instance of the TwoSided version of the same reference material and copy parameter values.
 		if (GIsEditor && MIC && ReferenceMaterialTwoSided)
 		{
-			UMaterialInstanceConstant* TwoSidedMIC = NewObject<UMaterialInstanceConstant>(
-				GetTransientPackage(),
-				NewInstanceName,
-				OneSidedMat->GetFlags()
+			bool bCreatedAsset = false;
+			UMaterialInstanceConstant* TwoSidedMIC = AssetCache.GetOrCreateCachedAsset<UMaterialInstanceConstant>(
+				PrefixedTwoSidedHash,
+				DesiredMaterialName,
+				DesiredFlags,
+				&bCreatedAsset
 			);
 
-			TwoSidedMIC->SetParentEditorOnly(ReferenceMaterialTwoSided);
-			TwoSidedMIC->CopyMaterialUniformParametersEditorOnly(OneSidedMat);
+			if (bCreatedAsset && TwoSidedMIC)
+			{
+				TwoSidedMIC->SetParentEditorOnly(ReferenceMaterialTwoSided);
+				TwoSidedMIC->CopyMaterialUniformParametersEditorOnly(OneSidedMat);
+			}
 
 			TwoSidedMat = TwoSidedMIC;
 		}
@@ -95,18 +102,23 @@ namespace UE::MeshTranslationImplInternal::Private
 		// Create a new material instance of it and set the override to two-sided.
 		else if (GIsEditor)
 		{
-			UMaterialInstanceConstant* TwoSidedMIC = NewObject<UMaterialInstanceConstant>(
-				GetTransientPackage(),
-				NewInstanceName,
-				OneSidedMat->GetFlags()
+			bool bCreatedAsset = false;
+			UMaterialInstanceConstant* TwoSidedMIC = AssetCache.GetOrCreateCachedAsset<UMaterialInstanceConstant>(
+				PrefixedTwoSidedHash,
+				DesiredMaterialName,
+				DesiredFlags,
+				&bCreatedAsset
 			);
 
-			TwoSidedMIC->SetParentEditorOnly(OneSidedMat);
-			TwoSidedMIC->BasePropertyOverrides.bOverride_TwoSided = true;
-			TwoSidedMIC->BasePropertyOverrides.TwoSided = true;
+			if (bCreatedAsset && TwoSidedMIC)
+			{
+				TwoSidedMIC->SetParentEditorOnly(OneSidedMat);
+				TwoSidedMIC->BasePropertyOverrides.bOverride_TwoSided = true;
+				TwoSidedMIC->BasePropertyOverrides.TwoSided = true;
 
-			FMaterialUpdateContext UpdateContext(FMaterialUpdateContext::EOptions::Default, GMaxRHIShaderPlatform);
-			UpdateContext.AddMaterialInstance(TwoSidedMIC);
+				FMaterialUpdateContext UpdateContext(FMaterialUpdateContext::EOptions::Default, GMaxRHIShaderPlatform);
+				UpdateContext.AddMaterialInstance(TwoSidedMIC);
+			}
 
 			TwoSidedMat = TwoSidedMIC;
 		}
@@ -116,21 +128,32 @@ namespace UE::MeshTranslationImplInternal::Private
 
 			// At runtime all we can do is create another instance of our two-sided reference materials, we cannot set
 			// another override
-			if (MID && ReferenceMaterialTwoSided)
+			if (MI && ReferenceMaterialTwoSided)
 			{
-				UMaterialInstanceDynamic* TwoSidedMID = UMaterialInstanceDynamic::Create(
-					ReferenceMaterialTwoSided,
-					GetTransientPackage(),
-					NewInstanceName
+				// Note how we're requesting just a UMaterialInstance here, instead of spelling out the MID. This because
+				// if we're a runtime we may have a cooked MIC assigned to this hash, and in that case we want to use it
+				// instead of overwriting it with a MID. Our creation func will ensure we create a MID as a fallback anyway
+				bool bCreatedAsset = false;
+				UMaterialInstance* TwoSidedMI = AssetCache.GetOrCreateCustomCachedAsset<UMaterialInstance>(
+					PrefixedTwoSidedHash,
+					DesiredMaterialName,
+					DesiredFlags | RF_Transient,	// We never want MIDs to become assets in the content browser
+					[ReferenceMaterialTwoSided](UPackage* Outer, FName SanitizedName, EObjectFlags FlagsToUse)
+					{
+						UMaterialInstanceDynamic* NewMID = UMaterialInstanceDynamic::Create(ReferenceMaterialTwoSided, Outer, SanitizedName);
+						NewMID->ClearFlags(NewMID->GetFlags());
+						NewMID->SetFlags(FlagsToUse);
+						return NewMID;
+					},
+					&bCreatedAsset
 				);
-				if (!ensure(TwoSidedMID))
+
+				if (UMaterialInstanceDynamic* TwoSidedMID = Cast<UMaterialInstanceDynamic>(TwoSidedMI))
 				{
-					return nullptr;
+					TwoSidedMID->CopyParameterOverrides(MI);
 				}
 
-				TwoSidedMID->CopyParameterOverrides(MID);
-
-				TwoSidedMat = TwoSidedMID;
+				TwoSidedMat = TwoSidedMI;
 			}
 
 		return TwoSidedMat;
@@ -145,7 +168,7 @@ namespace UE::MeshTranslationImplInternal::Private
 	UMaterialInterface* CreatePrimvarCompatibleVersionOfMaterial(
 		UMaterialInterface& Material,
 		const TMap<FString, int32>& MeshPrimvarToUVIndex,
-		UUsdAssetCache2* AssetCache,
+		UUsdAssetCache3* AssetCache,
 		FUsdInfoCache* InfoCache,
 		const FString& MaterialHashPrefix,
 		bool bReuseIdenticalAssets
@@ -262,8 +285,8 @@ namespace UE::MeshTranslationImplInternal::Private
 		}
 
 		FString ExistingHash = AssetCache->GetHashForAsset(&Material);
-		const bool bMaterialBelongsToAssetCache = !ExistingHash.IsEmpty();
-		if (!bMaterialBelongsToAssetCache)
+		const bool bMaterialTrackedByAssetCache = !ExistingHash.IsEmpty();
+		if (!bMaterialTrackedByAssetCache)
 		{
 			return nullptr;
 		}
@@ -296,65 +319,54 @@ namespace UE::MeshTranslationImplInternal::Private
 		// In theory we don't even need to add the prefix here because our ExistingHash will already have the same prefix...
 		// However for consistency it's probably for the best to have both assets have the same prefix, so you can tell
 		// from the hash that they originated from the same prim
-		FString PrefixedMaterialHash = MaterialHashPrefix + Hash.ToString();
-
-		if (UMaterialInterface* ExistingCompatibleMaterial = Cast<UMaterialInterface>(AssetCache->GetCachedAsset(PrefixedMaterialHash)))
-		{
-			UE_LOG(
-				LogUsd,
-				Verbose,
-				TEXT("Found existing compatible Material '%s' on the asset cache with hash '%s'"),
-				*ExistingCompatibleMaterial->GetPathName(),
-				*PrefixedMaterialHash
-			);
-			CompatibleMaterial = ExistingCompatibleMaterial;
-		}
+		const FString PrefixedCompatibleHash = MaterialHashPrefix + Hash.ToString();
 
-		TMap<FString, int32> CompatiblePrimvarToUVIndex;
-		CompatiblePrimvarToUVIndex.Reserve(CompatiblePrimvarAndUVIndexPairs.Num());
-		for (const TPair<FString, int32>& Pair : CompatiblePrimvarAndUVIndexPairs)
-		{
-			CompatiblePrimvarToUVIndex.Add(Pair);
-		}
-
-		// We have to create a brand new compatible material instance
 		bool bCreatedNew = false;
-		if (!CompatibleMaterial)
+#if WITH_EDITOR
+		if (GIsEditor)
 		{
-			const FName NewInstanceName = MakeUniqueObjectName(GetTransientPackage(), UMaterialInstance::StaticClass(), Material.GetFName());
-
-			UE_LOG(
-				LogUsd,
-				Verbose,
-				TEXT("Generating compatible version of Material '%s' (Parameter to primvar: %s, primvar to UV index: %s) with hash '%s'"),
-				*Material.GetPathName(),
-				*UsdUtils::StringifyMap(MaterialAssetUserData->ParameterToPrimvar),
-				*UsdUtils::StringifyMap(CompatiblePrimvarToUVIndex),
-				*PrefixedMaterialHash
+			UMaterialInstanceConstant* CompatibleMIC = AssetCache->GetOrCreateCachedAsset<UMaterialInstanceConstant>(
+				PrefixedCompatibleHash,
+				Material.GetName(),
+				Material.GetFlags(),
+				&bCreatedNew
 			);
 
-#if WITH_EDITOR
-			if (GIsEditor)
+			if (bCreatedNew)
 			{
-				UMaterialInstanceConstant* CompatibleMIC = NewObject<UMaterialInstanceConstant>(
-					GetTransientPackage(),
-					NewInstanceName,
-					Material.GetFlags()
-				);
-
 				CompatibleMIC->SetParentEditorOnly(&Material);
-
-				CompatibleMaterial = CompatibleMIC;
 			}
-			else
+
+			CompatibleMaterial = CompatibleMIC;
+		}
+		else
 #endif	  // WITH_EDITOR
-			{
-				UMaterialInstanceDynamic* CompatibleMID = UMaterialInstanceDynamic::Create(&Material, GetTransientPackage(), NewInstanceName);
+		{
+			// Note how we're requesting just a UMaterialInstance here, instead of spelling out the MID. This because
+			// if we're a runtime we may have a cooked MIC assigned to this hash, and in that case we want to use it
+			// instead of overwriting it with a MID. Our creation func will ensure we create a MID as a fallback anyway
+			UMaterialInstance* CompatibleMI = AssetCache->GetOrCreateCustomCachedAsset<UMaterialInstance>(
+				PrefixedCompatibleHash,
+				Material.GetName(),
+				Material.GetFlags() | RF_Transient,	   // We never want MIDs to become assets in the content browser
+				[&Material](UPackage* Outer, FName SanitizedName, EObjectFlags FlagsToUse)
+				{
+					UMaterialInstanceDynamic* NewMID = UMaterialInstanceDynamic::Create(&Material, Outer, SanitizedName);
+					NewMID->ClearFlags(NewMID->GetFlags());
+					NewMID->SetFlags(FlagsToUse);
+					return NewMID;
+				},
+				&bCreatedNew
+			);
 
-				CompatibleMaterial = CompatibleMID;
-			}
+			CompatibleMaterial = CompatibleMI;
+		}
 
-			bCreatedNew = true;
+		TMap<FString, int32> CompatiblePrimvarToUVIndex;
+		CompatiblePrimvarToUVIndex.Reserve(CompatiblePrimvarAndUVIndexPairs.Num());
+		for (const TPair<FString, int32>& Pair : CompatiblePrimvarAndUVIndexPairs)
+		{
+			CompatiblePrimvarToUVIndex.Add(Pair);
 		}
 
 		// Update the AssetUserData whether we created a new material instance or reused one from the asset cache.
@@ -362,16 +374,26 @@ namespace UE::MeshTranslationImplInternal::Private
 		UUsdMaterialAssetUserData* CompatibleUserData = nullptr;
 		if (CompatibleMaterial)
 		{
-			CompatibleUserData = DuplicateObject(MaterialAssetUserData, CompatibleMaterial, TEXT("USDAssetUserData"));
+			CompatibleUserData = DuplicateObject(MaterialAssetUserData, CompatibleMaterial);
 			CompatibleUserData->PrimvarToUVIndex = CompatiblePrimvarToUVIndex;
 
-			UsdUtils::SetAssetUserData(CompatibleMaterial, CompatibleUserData);
+			UsdUnreal::ObjectUtils::SetAssetUserData(CompatibleMaterial, CompatibleUserData);
 		}
 
 		// Now that the AssetUserData is done, actually set the UV index material parameters with the target indices
 		UMaterialInstance* CompatibleInstance = Cast<UMaterialInstance>(CompatibleMaterial);
 		if (bCreatedNew && CompatibleInstance && CompatibleUserData)
 		{
+			UE_LOG(
+				LogUsd,
+				Verbose,
+				TEXT("Generated compatible version of Material '%s' (Parameter to primvar: %s, primvar to UV index: %s) with hash '%s'"),
+				*Material.GetPathName(),
+				*UsdUtils::StringifyMap(MaterialAssetUserData->ParameterToPrimvar),
+				*UsdUtils::StringifyMap(CompatiblePrimvarToUVIndex),
+				*PrefixedCompatibleHash
+			);
+
 			for (const TPair<FString, FString>& ParameterPair : CompatibleUserData->ParameterToPrimvar)
 			{
 				const FString& Parameter = ParameterPair.Key;
@@ -401,19 +423,11 @@ namespace UE::MeshTranslationImplInternal::Private
 #endif	  // WITH_EDITOR
 		}
 
-		if (CompatibleMaterial && CompatibleMaterial != &Material)
+		if (CompatibleMaterial && CompatibleMaterial != &Material && InfoCache)
 		{
-			if (bMaterialBelongsToAssetCache)
+			for (const UE::FSdfPath& Prim : InfoCache->GetPrimsForAsset(&Material))
 			{
-				AssetCache->CacheAsset(PrefixedMaterialHash, CompatibleMaterial);
-			}
-
-			if (InfoCache)
-			{
-				for (const UE::FSdfPath& Prim : InfoCache->GetPrimsForAsset(&Material))
-				{
-					InfoCache->LinkAssetToPrim(Prim, CompatibleMaterial);
-				}
+				InfoCache->LinkAssetToPrim(Prim, CompatibleMaterial);
 			}
 		}
 
@@ -424,7 +438,7 @@ namespace UE::MeshTranslationImplInternal::Private
 TMap<const UsdUtils::FUsdPrimMaterialSlot*, UMaterialInterface*> MeshTranslationImpl::ResolveMaterialAssignmentInfo(
 	const pxr::UsdPrim& UsdPrim,
 	const TArray<UsdUtils::FUsdPrimMaterialAssignmentInfo>& AssignmentInfo,
-	UUsdAssetCache2& AssetCache,
+	UUsdAssetCache3& AssetCache,
 	FUsdInfoCache& InfoCache,
 	EObjectFlags Flags,
 	bool bReuseIdenticalAssets
@@ -462,62 +476,86 @@ TMap<const UsdUtils::FUsdPrimMaterialSlot*, UMaterialInterface*> MeshTranslation
 						Slot.MaterialSource
 					);
 
-					if (DisplayColorDesc.IsSet())
+					if (!DisplayColorDesc.IsSet())
 					{
-						FString DisplayColorHash;
-						{
-							FSHAHash Hash;
-							FSHA1 SHA1;
-							SHA1.UpdateWithString(*Slot.MaterialSource, Slot.MaterialSource.Len());
+						continue;
+					}
 
-							const FSoftObjectPath* ReferencePath = IUsdClassesModule::GetReferenceMaterialPath(DisplayColorDesc.GetValue());
-							if (ReferencePath)
-							{
-								FString ReferencePathString = ReferencePath->ToString();
-								SHA1.UpdateWithString(*ReferencePathString, ReferencePathString.Len());
-							}
+					const FSoftObjectPath* ReferencePath = IUsdClassesModule::GetReferenceMaterialPath(DisplayColorDesc.GetValue());
+					if (!ReferencePath)
+					{
+						continue;
+					}
 
-							SHA1.Final();
-							SHA1.GetHash(&Hash.Hash[0]);
-							DisplayColorHash = Hash.ToString();
-						}
-						const FString PrefixedHash = UsdUtils::GetAssetHashPrefix(UsdPrim, bReuseIdenticalAssets) + DisplayColorHash;
+					UMaterialInterface* ParentMaterial = Cast<UMaterialInterface>(ReferencePath->TryLoad());
+					if (!ParentMaterial)
+					{
+						continue;
+					}
 
-						// Try reusing an already created DisplayColor material
-						if (UMaterialInterface* ExistingMaterial = Cast<UMaterialInterface>(AssetCache.GetCachedAsset(PrefixedHash)))
-						{
-							Material = ExistingMaterial;
-						}
+					FString DisplayColorHash;
+					{
+						FSHAHash Hash;
+						FSHA1 SHA1;
+						SHA1.UpdateWithString(*Slot.MaterialSource, Slot.MaterialSource.Len());
 
-						// Need to create a new DisplayColor material
-						if (Material == nullptr)
-						{
-							UMaterialInstance* MaterialInstance = nullptr;
+						FString ReferencePathString = ReferencePath->ToString();
+						SHA1.UpdateWithString(*ReferencePathString, ReferencePathString.Len());
 
-							if (GIsEditor)	  // Editor, PIE => true; Standlone, packaged => false
-							{
-								MaterialInstance = IUsdClassesModule::CreateDisplayColorMaterialInstanceConstant(DisplayColorDesc.GetValue());
-							}
-							else
-							{
-								MaterialInstance = IUsdClassesModule::CreateDisplayColorMaterialInstanceDynamic(DisplayColorDesc.GetValue());
-							}
+						SHA1.Final();
+						SHA1.GetHash(&Hash.Hash[0]);
+						DisplayColorHash = Hash.ToString();
+					}
+					const FString PrefixedHash = UsdUtils::GetAssetHashPrefix(UsdPrim, bReuseIdenticalAssets) + DisplayColorHash;
 
-							if (MaterialInstance)
-							{
-								// Leave PrimPath as empty as it likely will be reused by many prims
-								UUsdAssetUserData* UserData = NewObject<UUsdAssetUserData>(MaterialInstance, TEXT("USDAssetUserData"));
-								MaterialInstance->AddAssetUserData(UserData);
-							}
+					FString DisplayColorName = FString::Printf(
+						TEXT("DisplayColor%s%s"),
+						DisplayColorDesc->bHasOpacity ? TEXT("_Translucent") : TEXT(""),
+						DisplayColorDesc->bIsDoubleSided ? TEXT("_TwoSided") : TEXT("")
+					);
 
-							// We can only cache transient assets
-							MaterialInstance->SetFlags(RF_Transient);
+					bool bCreatedNew = false;
+#if WITH_EDITOR
+					if (GIsEditor)
+					{
+						UMaterialInstanceConstant* MaterialInstance = AssetCache.GetOrCreateCachedAsset<UMaterialInstanceConstant>(
+							PrefixedHash,
+							DisplayColorName,
+							Flags,
+							&bCreatedNew
+						);
 
-							AssetCache.CacheAsset(PrefixedHash, MaterialInstance);
-							Material = MaterialInstance;
+						if (bCreatedNew)
+						{
+							MaterialInstance->SetParentEditorOnly(ParentMaterial);
 						}
+
+						Material = MaterialInstance;
+					}
+					else
+#endif	  // WITH_EDITOR
+					{
+						// Note how we're requesting just a UMaterialInstance here, instead of spelling out the MID. This because
+						// if we're a runtime we may have a cooked MIC assigned to this hash, and in that case we want to use it
+						// instead of overwriting it with a MID. Our creation func will ensure we create a MID as a fallback anyway
+						UMaterialInstance* MaterialInstance = AssetCache.GetOrCreateCustomCachedAsset<UMaterialInstance>(
+							PrefixedHash,
+							DisplayColorName,
+							Flags | RF_Transient,	 // We never want MIDs to become assets in the content browser
+							[Material](UPackage* Outer, FName SanitizedName, EObjectFlags FlagsToUse)
+							{
+								UMaterialInstanceDynamic* NewMID = UMaterialInstanceDynamic::Create(Material, Outer, SanitizedName);
+								NewMID->ClearFlags(NewMID->GetFlags());
+								NewMID->SetFlags(FlagsToUse);
+								return NewMID;
+							},
+							&bCreatedNew
+						);
+
+						Material = MaterialInstance;
 					}
 
+					UsdUnreal::ObjectUtils::GetOrCreateAssetUserData(Material);
 					break;
 				}
 				case UsdUtils::EPrimAssignmentType::MaterialPrim:
@@ -574,8 +612,9 @@ TMap<const UsdUtils::FUsdPrimMaterialSlot*, UMaterialInterface*> MeshTranslation
 							UE_LOG(
 								LogUsd,
 								Warning,
-								TEXT("Failed to generate a two-sided material from the material prim at path '%s' as no "
-									 "single-sided material was generated for it."),
+								TEXT(
+									"Failed to generate a two-sided material from the material prim at path '%s' as no single-sided material was generated for it."
+								),
 								*Slot.MaterialSource
 							);
 							continue;
@@ -584,26 +623,34 @@ TMap<const UsdUtils::FUsdPrimMaterialSlot*, UMaterialInterface*> MeshTranslation
 						const FString PrefixedOneSidedHash = AssetCache.GetHashForAsset(OneSidedMat);
 						const FString PrefixedTwoSidedHash = PrefixedOneSidedHash + UnrealIdentifiers::TwoSidedMaterialSuffix;
 
-						// Check if for some reason we already have a two-sided material ready due to a complex scenario
-						// related to the global cache
-						UMaterialInterface* TwoSidedMat = Cast<UMaterialInterface>(AssetCache.GetCachedAsset(PrefixedTwoSidedHash));
-						if (!TwoSidedMat)
-						{
-							TwoSidedMat = UE::MeshTranslationImplInternal::Private::CreateTwoSidedVersionOfMaterial(OneSidedMat);
-						}
+						UMaterialInterface* TwoSidedMat = UE::MeshTranslationImplInternal::Private::GetOrCreateTwoSidedVersionOfMaterial(
+							OneSidedMat,
+							PrefixedTwoSidedHash,
+							AssetCache
+						);
 
 						if (TwoSidedMat)
 						{
 							// Update AssetUserData whether we generated a new material or reused one from the asset cache
 							{
-								UUsdMaterialAssetUserData* OneSidedUserData = OneSidedMat->GetAssetUserData<UUsdMaterialAssetUserData>();
+								UUsdMaterialAssetUserData* OneSidedUserData = UsdUnreal::ObjectUtils::GetAssetUserData<UUsdMaterialAssetUserData>(
+									OneSidedMat
+								);
 								ensure(OneSidedUserData);
 
-								UUsdMaterialAssetUserData* UserData = DuplicateObject(OneSidedUserData, TwoSidedMat, TEXT("USDAssetUserData"));
-								UsdUtils::SetAssetUserData(TwoSidedMat, UserData);
+								UUsdMaterialAssetUserData*
+									TwoSidedUserData = UsdUnreal::ObjectUtils::GetOrCreateAssetUserData<UUsdMaterialAssetUserData>(TwoSidedMat);
+
+								// Copy stuff from OneSidedMat when it makes sense, as it may have been regenerated
+								if (OneSidedUserData && TwoSidedUserData)
+								{
+									TwoSidedUserData->ParameterToPrimvar = OneSidedUserData->ParameterToPrimvar;
+									TwoSidedUserData->PrimvarToUVIndex = OneSidedUserData->PrimvarToUVIndex;
+									TwoSidedUserData->StageIdentifierToMetadata = OneSidedUserData->StageIdentifierToMetadata;
+									TwoSidedUserData->PrimPaths = OneSidedUserData->PrimPaths;
+								}
 							}
 
-							TwoSidedMat->SetFlags(RF_Transient);
 							Material = TwoSidedMat;
 							PrefixedMaterialHash = PrefixedTwoSidedHash;
 						}
@@ -612,8 +659,9 @@ TMap<const UsdUtils::FUsdPrimMaterialSlot*, UMaterialInterface*> MeshTranslation
 							UE_LOG(
 								LogUsd,
 								Warning,
-								TEXT("Failed to generate a two-sided material from the material prim at path '%s'. Falling "
-									 "back to using the single-sided material '%s' instead."),
+								TEXT(
+									"Failed to generate a two-sided material from the material prim at path '%s'. Falling back to using the single-sided material '%s' instead."
+								),
 								*Slot.MaterialSource,
 								*OneSidedMat->GetPathName()
 							);
@@ -624,9 +672,9 @@ TMap<const UsdUtils::FUsdPrimMaterialSlot*, UMaterialInterface*> MeshTranslation
 
 					if (Material)
 					{
-						// Cache the material to "ping it" as active, but also register two sided materials for the
-						// first time
-						AssetCache.CacheAsset(PrefixedMaterialHash, Material);
+						// Ping one-sided material as active
+						AssetCache.TouchAssetPath(Material);
+
 						InfoCache.LinkAssetToPrim(UE::FSdfPath{*Slot.MaterialSource}, Material);
 
 						// Finally, try to make our generated material primvar-compatible. We do this last because this will
@@ -637,7 +685,7 @@ TMap<const UsdUtils::FUsdPrimMaterialSlot*, UMaterialInterface*> MeshTranslation
 						{
 							Material = AlreadyHandledMaterial;
 
-							AssetCache.TouchAsset(Material);
+							AssetCache.TouchAssetPath(Material);
 							InfoCache.LinkAssetToPrim(UE::FSdfPath{*Slot.MaterialSource}, Material);
 						}
 						else
@@ -720,7 +768,7 @@ void MeshTranslationImpl::SetMaterialOverrides(
 	const pxr::UsdPrim& Prim,
 	const TArray<UMaterialInterface*>& ExistingAssignments,
 	UMeshComponent& MeshComponent,
-	UUsdAssetCache2& AssetCache,
+	UUsdAssetCache3& AssetCache,
 	FUsdInfoCache& InfoCache,
 	float Time,
 	EObjectFlags Flags,
diff --git a/Engine/Plugins/Importers/USDImporter/Source/USDSchemas/Private/MeshTranslationImpl.h b/Engine/Plugins/Importers/USDImporter/Source/USDSchemas/Private/MeshTranslationImpl.h
--- a/Engine/Plugins/Importers/USDImporter/Source/USDSchemas/Private/MeshTranslationImpl.h
+++ b/Engine/Plugins/Importers/USDImporter/Source/USDSchemas/Private/MeshTranslationImpl.h
@@ -18,7 +18,7 @@ PXR_NAMESPACE_CLOSE_SCOPE
 class FUsdInfoCache;
 class UMaterialInterface;
 class UMeshComponent;
-class UUsdAssetCache2;
+class UUsdAssetCache3;
 class UUsdMeshAssetUserData;
 namespace UsdUtils
 {
@@ -33,7 +33,7 @@ namespace MeshTranslationImpl
 	TMap<const UsdUtils::FUsdPrimMaterialSlot*, UMaterialInterface*> ResolveMaterialAssignmentInfo(
 		const pxr::UsdPrim& UsdPrim,
 		const TArray<UsdUtils::FUsdPrimMaterialAssignmentInfo>& AssignmentInfo,
-		UUsdAssetCache2& AssetCache,
+		UUsdAssetCache3& AssetCache,
 		FUsdInfoCache& InfoCache,
 		EObjectFlags Flags,
 		bool bReuseIdenticalAssets
@@ -47,7 +47,7 @@ namespace MeshTranslationImpl
 		const pxr::UsdPrim& Prim,
 		const TArray<UMaterialInterface*>& ExistingAssignments,
 		UMeshComponent& MeshComponent,
-		UUsdAssetCache2& AssetCache,
+		UUsdAssetCache3& AssetCache,
 		FUsdInfoCache& InfoCache,
 		float Time,
 		EObjectFlags Flags,
