USD: Implements UsdAssetCache3 which should be a simpler and more flexible way of managing USD assets;

diff --git a/Engine/Plugins/Importers/USDImporter/Source/USDSchemas/Private/USDGeomMeshTranslator.cpp b/Engine/Plugins/Importers/USDImporter/Source/USDSchemas/Private/USDGeomMeshTranslator.cpp
--- a/Engine/Plugins/Importers/USDImporter/Source/USDSchemas/Private/USDGeomMeshTranslator.cpp
+++ b/Engine/Plugins/Importers/USDImporter/Source/USDSchemas/Private/USDGeomMeshTranslator.cpp
@@ -6,6 +6,7 @@
 
 #include "MeshTranslationImpl.h"
 #include "UnrealUSDWrapper.h"
+#include "USDAssetCache3.h"
 #include "USDAssetUserData.h"
 #include "USDClassesModule.h"
 #include "USDConversionUtils.h"
@@ -13,7 +14,9 @@
 #include "USDGeomMeshConversion.h"
 #include "USDInfoCache.h"
 #include "USDLog.h"
+#include "USDObjectUtils.h"
 #include "USDPrimConversion.h"
+#include "USDTranslatorUtils.h"
 #include "USDTypesConversion.h"
 
 #include "UsdWrappers/SdfPath.h"
@@ -251,7 +254,7 @@ namespace UsdGeomMeshTranslatorImpl
 		const pxr::UsdPrim& UsdPrim,
 		const TArray<UsdUtils::FUsdPrimMaterialAssignmentInfo>& LODIndexToMaterialInfo,
 		UStaticMesh& StaticMesh,
-		UUsdAssetCache2& AssetCache,
+		UUsdAssetCache3& AssetCache,
 		FUsdInfoCache* InfoCache,
 		float Time,
 		EObjectFlags Flags,
@@ -708,26 +711,23 @@ namespace UsdGeomMeshTranslatorImpl
 
 		FString PrefixedAssetHash = UsdUtils::GetAssetHashPrefix(Prim, Context.bReuseIdenticalAssets) + AllLODHash.ToString();
 
-		if (Context.AssetCache)
+		FString DesiredName = FPaths::GetBaseFilename(MeshName);
+
+		// If we're just opening the stage we'll ignore the LOD variant sets and just parse the first LOD mesh as a regular Mesh prim,
+		// which would have produced a static mesh just named something like "LOD0" or "LOD1". This should provide a more descriptive name,
+		// like "Cube_LOD0" instead
+		if (!Context.bAllowInterpretingLODs && UsdUtils::IsGeomMeshALOD(Prim))
 		{
-			StaticMesh = Cast<UStaticMesh>(Context.AssetCache->GetCachedAsset(PrefixedAssetHash));
+			DesiredName = FPaths::GetBaseFilename(FPaths::GetPath(MeshName)) + TEXT("_") + DesiredName;
 		}
 
-		if (!StaticMesh && bHasValidMeshDescription)
+		if (Context.UsdAssetCache)
 		{
-			bOutIsNew = true;
-
-			FName AssetName = MakeUniqueObjectName(
-				GetTransientPackage(),
-				UStaticMesh::StaticClass(),
-				*IUsdClassesModule::SanitizeObjectName(FPaths::GetBaseFilename(MeshName))
-			);
-			StaticMesh = NewObject<UStaticMesh>(
-				GetTransientPackage(),
-				AssetName,
-				Context.ObjectFlags | EObjectFlags::RF_Public | EObjectFlags::RF_Transient
-			);
+			StaticMesh = Context.UsdAssetCache->GetOrCreateCachedAsset<UStaticMesh>(PrefixedAssetHash, DesiredName, Context.ObjectFlags, &bOutIsNew);
+		}
 
+		if (StaticMesh && bHasValidMeshDescription && bOutIsNew)
+		{
 #if WITH_EDITOR
 			for (int32 LODIndex = 0; LODIndex < LODIndexToMeshDescription.Num(); ++LODIndex)
 			{
@@ -754,16 +754,6 @@ namespace UsdGeomMeshTranslatorImpl
 #endif	  // WITH_EDITOR
 
 			StaticMesh->SetLightingGuid();
-
-			if (Context.AssetCache)
-			{
-				Context.AssetCache->CacheAsset(PrefixedAssetHash, StaticMesh);
-			}
-		}
-		else
-		{
-			// FPlatformMisc::LowLevelOutputDebugStringf( TEXT("Mesh found in cache %s\n"), *StaticMesh->GetName() );
-			bOutIsNew = false;
 		}
 
 		return StaticMesh;
@@ -851,7 +841,8 @@ namespace UsdGeomMeshTranslatorImpl
 			// Fetch the MeshDescription from the imported LODIndexToMeshDescription as StaticMesh.GetMeshDescription is editor-only
 			StaticMesh.GetRenderData()->Bounds = LODIndexToMeshDescription[0].GetBounds();
 			StaticMesh.CalculateExtendedBounds();
-#endif											   // WITH_EDITOR
+
+#endif	  // WITH_EDITOR
 		}
 	}
 }	 // namespace UsdGeomMeshTranslatorImpl
@@ -1711,13 +1702,7 @@ void FBuildStaticMeshTaskChain::SetupTasks()
 				   Context->InfoCache->LinkAssetToPrim(TargetPath, StaticMesh);
 			   }
 
-#if WITH_EDITOR
-			   StaticMesh->NaniteSettings.bEnabled = bShouldEnableNanite;
-			   StaticMesh->NaniteSettings.NormalPrecision = GNaniteSettingsNormalPrecision;
-			   StaticMesh->NaniteSettings.TangentPrecision = GNaniteSettingsTangentPrecision;
-#endif	  // WITH_EDITOR
-
-			   if (UUsdMeshAssetUserData* UserData = UsdUtils::GetOrCreateAssetUserData<UUsdMeshAssetUserData>(StaticMesh))
+			   if (UUsdMeshAssetUserData* UserData = UsdUnreal::ObjectUtils::GetOrCreateAssetUserData<UUsdMeshAssetUserData>(StaticMesh))
 			   {
 				   UserData->PrimvarToUVIndex = LODIndexToMaterialInfo[0].PrimvarToUVIndex;	   // We use the same primvar mapping for all LODs
 				   UserData->PrimPaths.AddUnique(PrimPathString);
@@ -1756,11 +1741,17 @@ void FBuildStaticMeshTaskChain::SetupTasks()
 			   // components
 			   if (bIsNew)
 			   {
+#if WITH_EDITOR
+				   StaticMesh->NaniteSettings.bEnabled = bShouldEnableNanite;
+				   StaticMesh->NaniteSettings.NormalPrecision = GNaniteSettingsNormalPrecision;
+				   StaticMesh->NaniteSettings.TangentPrecision = GNaniteSettingsTangentPrecision;
+#endif	  // WITH_EDITOR
+
 				   UsdGeomMeshTranslatorImpl::ProcessStaticMeshMaterials(
 					   GetPrim(),
 					   LODIndexToMaterialInfo,
 					   *StaticMesh,
-					   *Context->AssetCache.Get(),
+					   *Context->UsdAssetCache.Get(),
 					   Context->InfoCache.Get(),
 					   Context->Time,
 					   Context->ObjectFlags,
@@ -1781,6 +1772,8 @@ void FBuildStaticMeshTaskChain::SetupTasks()
 			   else
 			   {
 				   // Setup collision on existing mesh in case the collision settings have changed
+				   // TODO: This is recomputing collision and dirtying the StaticMesh every time, even if the mesh
+				   // is reused from the AssetCache. Maybe we can do that only when collision settings actually change?
 				   UE::UsdCollision::Private::SetupSimpleCollision(GetPrim(), *StaticMesh);
 			   }
 		   }
@@ -1854,31 +1847,14 @@ void FBuildStaticMeshTaskChain::SetupTasks()
 			// Build failed, abandon mesh
 			if (!bSuccess)
 			{
-				Context->InfoCache->RemoveAllAssetPrimLinks(StaticMesh);
-
-				FString Hash = Context->AssetCache->GetHashForAsset(StaticMesh);
-				if (!Hash.IsEmpty())
-				{
-					Context->AssetCache->RemoveAssetReference(StaticMesh);
-					UObject* RemovedAsset = Context->AssetCache->RemoveAsset(Hash);
-					if (ensure(RemovedAsset))
-					{
-						const TCHAR* NewName = nullptr;
-						UObject* NewOuter = GetTransientPackage();
-						RemovedAsset->Rename(NewName, NewOuter);
-
-						StaticMesh = nullptr;
-
-						UE_LOG(
-							LogUsd,
-							Warning,
-							TEXT("Discarding StaticMesh generated for prim '%s' as it didn't produce any valid RenderData (likely all triangles "
-								 "were degenerate)"),
-							*PrimPath.GetString()
-						);
-					}
-				}
-
+				UE_LOG(
+					LogUsd,
+					Warning,
+					TEXT("Discarding StaticMesh generated for prim '%s' as it didn't produce any valid RenderData (likely all triangles "
+						 "were degenerate)"),
+					*PrimPath.GetString()
+				);
+				UsdUnreal::TranslatorUtils::AbandonFailedAsset(StaticMesh, Context->UsdAssetCache.Get(), Context->InfoCache.Get());
 				return false;
 			}
 
diff --git a/Engine/Plugins/Importers/USDImporter/Source/USDSchemas/Private/USDGeomPrimitiveTranslator.cpp b/Engine/Plugins/Importers/USDImporter/Source/USDSchemas/Private/USDGeomPrimitiveTranslator.cpp
--- a/Engine/Plugins/Importers/USDImporter/Source/USDSchemas/Private/USDGeomPrimitiveTranslator.cpp
+++ b/Engine/Plugins/Importers/USDImporter/Source/USDSchemas/Private/USDGeomPrimitiveTranslator.cpp
@@ -164,7 +164,7 @@ USceneComponent* FUsdGeomPrimitiveTranslator::CreateComponents()
 					GetPrim(),
 					ExistingAssignments,
 					*StaticMeshComponent,
-					*Context->AssetCache.Get(),
+					*Context->UsdAssetCache.Get(),
 					*Context->InfoCache.Get(),
 					Context->Time,
 					Context->ObjectFlags,
diff --git a/Engine/Plugins/Importers/USDImporter/Source/USDSchemas/Private/USDGeomXformableTranslator.cpp b/Engine/Plugins/Importers/USDImporter/Source/USDSchemas/Private/USDGeomXformableTranslator.cpp
--- a/Engine/Plugins/Importers/USDImporter/Source/USDSchemas/Private/USDGeomXformableTranslator.cpp
+++ b/Engine/Plugins/Importers/USDImporter/Source/USDSchemas/Private/USDGeomXformableTranslator.cpp
@@ -4,8 +4,8 @@
 
 #include "MeshTranslationImpl.h"
 #include "UnrealUSDWrapper.h"
+#include "USDAssetCache3.h"
 #include "USDAssetUserData.h"
-#include "USDClassesModule.h"
 #include "USDConversionUtils.h"
 #include "USDDrawModeComponent.h"
 #include "USDGeomMeshConversion.h"
@@ -13,6 +13,7 @@
 #include "USDIntegrationUtils.h"
 #include "USDLog.h"
 #include "USDMemory.h"
+#include "USDObjectUtils.h"
 #include "USDPrimConversion.h"
 #include "USDSchemasModule.h"
 #include "USDShadeConversion.h"
@@ -30,6 +31,7 @@
 #include "LiveLinkComponentController.h"
 #include "LiveLinkRole.h"
 #include "Misc/App.h"
+#include "Misc/Paths.h"
 #include "Modules/ModuleManager.h"
 #include "PhysicsEngine/BodySetup.h"
 #include "Roles/LiveLinkTransformRole.h"
@@ -332,7 +334,7 @@ USceneComponent* FUsdGeomXformableTranslator::CreateComponents()
 					GetPrim(),
 					ExistingAssignments,
 					*StaticMeshComponent,
-					*Context->AssetCache.Get(),
+					*Context->UsdAssetCache.Get(),
 					*Context->InfoCache.Get(),
 					Context->Time,
 					Context->ObjectFlags,
@@ -445,7 +447,11 @@ USceneComponent* FUsdGeomXformableTranslator::CreateComponentsEx(TOptional<TSubc
 	// the transient flag after spawn to make sure the spawned actor get an external
 	// package if needed.
 	const EObjectFlags PreComponentFlags = Context->ObjectFlags & ~(RF_Standalone | RF_Public | RF_Transient);
-	const EObjectFlags PostComponentFlags = Context->ObjectFlags & RF_Transient;
+	EObjectFlags PostComponentFlags = Context->ObjectFlags & ~(RF_Standalone | RF_Public);
+	if (!Context->bIsImporting)
+	{
+		PostComponentFlags |= RF_Transient;
+	}
 
 	if (bNeedsActor.GetValue())
 	{
@@ -549,7 +555,7 @@ USceneComponent* FUsdGeomXformableTranslator::CreateComponentsEx(TOptional<TSubc
 			const FName ComponentName = MakeUniqueObjectName(
 				ComponentOuter,
 				ComponentType.GetValue(),
-				*IUsdClassesModule::SanitizeObjectName(Prim.GetName().ToString())
+				*UsdUnreal::ObjectUtils::SanitizeObjectName(Prim.GetName().ToString())
 			);
 			SceneComponent = NewObject<USceneComponent>(ComponentOuter, ComponentType.GetValue(), ComponentName, PreComponentFlags);
 			SceneComponent->SetFlags(PostComponentFlags);
@@ -563,7 +569,7 @@ USceneComponent* FUsdGeomXformableTranslator::CreateComponentsEx(TOptional<TSubc
 
 	if (Context->MetadataOptions.bCollectMetadata && Context->MetadataOptions.bCollectOnComponents)
 	{
-		UUsdAssetUserData* UserData = UsdUtils::GetOrCreateAssetUserData(SceneComponent);
+		UUsdAssetUserData* UserData = UsdUnreal::ObjectUtils::GetOrCreateAssetUserData(SceneComponent);
 
 		// It makes sense for asset metadata to "include all prims in the subtree", as when we generate an
 		// asset we don't generate additional separate assets for child prims. This is not the same behavior
@@ -579,7 +585,7 @@ USceneComponent* FUsdGeomXformableTranslator::CreateComponentsEx(TOptional<TSubc
 			bCollectMetadataFromSubtree
 		);
 	}
-	else if (UUsdAssetUserData* UserData = UsdUtils::GetAssetUserData(SceneComponent))
+	else if (UUsdAssetUserData* UserData = UsdUnreal::ObjectUtils::GetAssetUserData(SceneComponent))
 	{
 		// Strip the metadata from this prim, so that if we uncheck "Collect Metadata" it actually disappears on the AssetUserData
 		UserData->StageIdentifierToMetadata.Remove(Prim.GetStage().GetRootLayer().GetIdentifier());
@@ -773,7 +779,7 @@ namespace UE::UsdXformableTranslatorImpl::Private
 	void AssignDrawModeComponentTextures(
 		UE::FUsdPrim Prim,
 		UUsdDrawModeComponent* DrawModeComponent,
-		UUsdAssetCache2& AssetCache,
+		UUsdAssetCache3& AssetCache,
 		FUsdInfoCache& InfoCache
 	)
 	{
@@ -855,7 +861,8 @@ namespace UE::UsdXformableTranslatorImpl::Private
 
 				if (TextureSetter)
 				{
-					std::unordered_map<pxr::SdfPath, TArray<UTexture2D*>, pxr::SdfPath::Hash>::iterator iter = AttrPathToTextures.find(Attr.GetPath());
+					std::unordered_map<pxr::SdfPath, TArray<UTexture2D*>, pxr::SdfPath::Hash>::iterator iter = AttrPathToTextures.find(Attr.GetPath()
+					);
 					if (iter != AttrPathToTextures.end())
 					{
 						const FString TexturePath = UsdUtils::GetResolvedAssetPath(Attr);
@@ -918,14 +925,14 @@ USceneComponent* FUsdGeomXformableTranslator::CreateAlternativeDrawModeComponent
 		case EUsdDrawMode::Cards:
 		{
 			UUsdDrawModeComponent* Component = Cast<UUsdDrawModeComponent>(CreateComponentsEx({UUsdDrawModeComponent::StaticClass()}, bNeedsActor));
-			if (ensure(Component) && Context->AssetCache && Context->InfoCache)
+			if (ensure(Component) && Context->UsdAssetCache && Context->InfoCache)
 			{
 				// For now we only assign textures when creating components, not when updating. Maybe in the future we can
 				// add support for "texture animations"
 				UE::UsdXformableTranslatorImpl::Private::AssignDrawModeComponentTextures(
 					GetPrim(),
 					Component,
-					*Context->AssetCache,
+					*Context->UsdAssetCache,
 					*Context->InfoCache
 				);
 			}
@@ -947,7 +954,7 @@ void FUsdGeomXformableTranslator::CreateAlternativeDrawModeAssets(EUsdDrawMode D
 {
 	// Currently we just use this function to create the textures that we're going to use on the bounds components,
 	// if applicable
-	if (DrawMode != EUsdDrawMode::Cards || !Context->AssetCache || !Context->InfoCache)
+	if (DrawMode != EUsdDrawMode::Cards || !Context->UsdAssetCache || !Context->InfoCache)
 	{
 		return;
 	}
@@ -985,38 +992,31 @@ void FUsdGeomXformableTranslator::CreateAlternativeDrawModeAssets(EUsdDrawMode D
 			{
 				const FString HashPrefix = UsdUtils::GetAssetHashPrefix(GetPrim(), Context->bReuseIdenticalAssets);
 				const FString PrefixedTextureHash = HashPrefix + LexToString(FMD5Hash::HashFile(*ResolvedPath));
-				UTexture2D* Texture = Cast<UTexture2D>(Context->AssetCache->GetCachedAsset(PrefixedTextureHash));
 
-				if (!Texture)
-				{
-					Texture = Cast<UTexture2D>(UsdUtils::CreateTexture(
-						Attr,
-						UsdToUnreal::ConvertPath(Attr.GetPrim().GetPath()),
-						TEXTUREGROUP_World,
-						Context->AssetCache.Get()
-					));
-
-					if (Texture)
-					{
-						Context->AssetCache->CacheAsset(PrefixedTextureHash, Texture);
-					}
-				}
+				const FString& DesiredTextureName = FPaths::GetBaseFilename(ResolvedPath);
 
-				if (Texture)
-				{
-					// We link the textures to the prim, so that if the prim is reloaded the AUsdStageActor knows to potentially
-					// drop the textures. However we put the full attribute path on AssetUserData, so that when we're filling in
-					// our UUsdDrawModeComponent later, we know which texture came from which attribute
+				const EObjectFlags DesiredFlags = Context->ObjectFlags;
 
-					UUsdAssetUserData* TextureUserData = Texture->GetAssetUserData<UUsdAssetUserData>();
-					if (!TextureUserData)
+				bool bCreatedTexture = false;
+				UTexture2D* Texture = Context->UsdAssetCache.Get()->GetOrCreateCustomCachedAsset<UTexture2D>(
+					PrefixedTextureHash,
+					DesiredTextureName,
+					DesiredFlags,
+					[&ResolvedPath](UPackage* Outer, FName SanitizedName, EObjectFlags FlagsToUse)
 					{
-						TextureUserData = NewObject<UUsdAssetUserData>(Texture, TEXT("USDAssetUserData"));
-						Texture->AddAssetUserData(TextureUserData);
-					}
-					TextureUserData->PrimPaths.AddUnique(UsdToUnreal::ConvertPath(Attr.GetPath()));
+						const TextureGroup Group = TextureGroup::TEXTUREGROUP_World;
+						return UsdUtils::CreateTexture(ResolvedPath, SanitizedName, Group, FlagsToUse, Outer);
+					},
+					&bCreatedTexture
+				);
 
-					Context->InfoCache->LinkAssetToPrim(PrimPath, Texture);
+				// We link the textures to the prim, so that if the prim is reloaded the AUsdStageActor knows to potentially
+				// drop the textures. However we put the full attribute path on AssetUserData, so that when we're filling in
+				// our UUsdDrawModeComponent later, we know which texture came from which attribute
+				Context->InfoCache->LinkAssetToPrim(PrimPath, Texture);
+				if (UUsdAssetUserData* TextureUserData = UsdUnreal::ObjectUtils::GetOrCreateAssetUserData(Texture))
+				{
+					TextureUserData->PrimPaths.AddUnique(UsdToUnreal::ConvertPath(Attr.GetPath()));
 				}
 			}
 		}
diff --git a/Engine/Plugins/Importers/USDImporter/Source/USDSchemas/Private/USDGeometryCacheTranslator.cpp b/Engine/Plugins/Importers/USDImporter/Source/USDSchemas/Private/USDGeometryCacheTranslator.cpp
--- a/Engine/Plugins/Importers/USDImporter/Source/USDSchemas/Private/USDGeometryCacheTranslator.cpp
+++ b/Engine/Plugins/Importers/USDImporter/Source/USDSchemas/Private/USDGeometryCacheTranslator.cpp
@@ -6,14 +6,15 @@
 // The GeometryCacheStreamer module is editor-only, so is the translator
 
 #include "MeshTranslationImpl.h"
+#include "USDAssetCache3.h"
 #include "USDAssetUserData.h"
-#include "USDClassesModule.h"
 #include "USDConversionUtils.h"
 #include "USDDrawModeComponent.h"
 #include "USDGroomTranslatorUtils.h"
 #include "USDInfoCache.h"
 #include "USDIntegrationUtils.h"
 #include "USDLog.h"
+#include "USDObjectUtils.h"
 #include "USDPrimConversion.h"
 #include "USDTypesConversion.h"
 
@@ -66,13 +67,22 @@ static FAutoConsoleVariableRef CVarEnableSubdiv(
 		 "StaticMeshes")
 );
 
+static bool GForceImport = true;
+static FAutoConsoleVariableRef CVarForceImport(
+	TEXT("USD.GeometryCache.ForceImport"),
+	GEnableSubdiv,
+	TEXT(
+		"Whether to immediately cache all geometry cache frames into the asset when the stage opens, acting like a fully imported Geometry Cache asset"
+	)
+);
+
 namespace UsdGeometryCacheTranslatorImpl
 {
 	bool ProcessGeometryCacheMaterials(
 		const pxr::UsdPrim& UsdPrim,
 		const TArray<UsdUtils::FUsdPrimMaterialAssignmentInfo>& LODIndexToMaterialInfo,
 		UGeometryCache& GeometryCache,
-		UUsdAssetCache2& AssetCache,
+		UUsdAssetCache3& AssetCache,
 		FUsdInfoCache* InfoCache,
 		float Time,
 		EObjectFlags Flags,
@@ -373,7 +383,7 @@ namespace UsdGeometryCacheTranslatorImpl
 	UGeometryCacheTrackStreamable* CreateStreamableTrack(UGeometryCache* GeometryCache, const FString& PrimPath)
 	{
 		// Create and configure a new StreamableTrack to be added to the GeometryCache
-		FString ObjectName = IUsdClassesModule::SanitizeObjectName(FPaths::GetBaseFilename(PrimPath));
+		FString ObjectName = UsdUnreal::ObjectUtils::SanitizeObjectName(FPaths::GetBaseFilename(PrimPath));
 
 		FName CodecName = MakeUniqueObjectName(GeometryCache, UGeometryCacheCodecV1::StaticClass(), FName(ObjectName + FString(TEXT("_Codec"))));
 		UGeometryCacheCodecV1* Codec = NewObject<UGeometryCacheCodecV1>(GeometryCache, CodecName, RF_Public);
@@ -420,36 +430,32 @@ namespace UsdGeometryCacheTranslatorImpl
 			FramesPerSecond = 1;
 		}
 
+		const bool bIsImporting = GForceImport || Context->bIsImporting;
+
 		// Frame rate must be taken into account as well since different frame rates must produce different sampling in the tracks
 		SHA1.Update(reinterpret_cast<uint8*>(&FramesPerSecond), sizeof(FramesPerSecond));
 
 		// Track type depends on if it's importing or not. Import needs to generate a persistent asset with all the frames already sampled
-		SHA1.Update(reinterpret_cast<uint8*>(&Context->bIsImporting), sizeof(Context->bIsImporting));
+		SHA1.Update(reinterpret_cast<const uint8*>(&bIsImporting), sizeof(bIsImporting));
 		SHA1.Final();
 
 		FSHAHash GeoCacheHash;
 		SHA1.GetHash(&GeoCacheHash.Hash[0]);
 		const FString PrefixedGeoCacheHash = UsdUtils::GetAssetHashPrefix(RootPrim, Context->bReuseIdenticalAssets) + GeoCacheHash.ToString();
 
-		UGeometryCache* GeometryCache = Cast<UGeometryCache>(Context->AssetCache->GetCachedAsset(PrefixedGeoCacheHash));
+		const FString DesiredName = FPaths::GetBaseFilename(RootPrimPath);
 
-		if (!GeometryCache)
-		{
-			bOutIsNew = true;
-
-			const FName AssetName = MakeUniqueObjectName(
-				GetTransientPackage(),
-				UGeometryCache::StaticClass(),
-				*IUsdClassesModule::SanitizeObjectName(FPaths::GetBaseFilename(RootPrimPath))
-			);
-			GeometryCache = NewObject<UGeometryCache>(
-				GetTransientPackage(),
-				AssetName,
-				Context->ObjectFlags | EObjectFlags::RF_Public | RF_Transient
-			);
+		UGeometryCache* GeometryCache = Context->UsdAssetCache->GetOrCreateCachedAsset<UGeometryCache>(
+			PrefixedGeoCacheHash,
+			DesiredName,
+			Context->ObjectFlags,
+			&bOutIsNew
+		);
 
+		if (GeometryCache && bOutIsNew)
+		{
 			FReadMeshDataArgs Args(GetReadMeshDataArgs(Context, RootPrimPath));
-			if (!Context->bIsImporting)
+			if (!bIsImporting)
 			{
 				// StartOffsetTime is the offset applied to the GeometryCache section on the sequencer track, so not relevant when importing
 				StartOffsetTime = static_cast<float>(Args.StartFrame) / Args.FramesPerSecond;
@@ -461,7 +467,7 @@ namespace UsdGeometryCacheTranslatorImpl
 			{
 				const FString& PrimPath = MeshPaths[Index].GetString();
 				UGeometryCacheTrack* Track = nullptr;
-				if (!Context->bIsImporting)
+				if (!bIsImporting)
 				{
 					Track = CreateUsdStreamTrack(GeometryCache, Args, PrimPath, MaterialOffsets[Index]);
 				}
@@ -480,12 +486,6 @@ namespace UsdGeometryCacheTranslatorImpl
 				MatTimes.Add(0.0f);
 				Track->SetMatrixSamples(Mats, MatTimes);
 			}
-
-			Context->AssetCache->CacheAsset(PrefixedGeoCacheHash, GeometryCache);
-		}
-		else
-		{
-			bOutIsNew = false;
 		}
 
 		return GeometryCache;
@@ -885,8 +885,10 @@ void FGeometryCacheCreateAssetsTaskChain::SetupTasks()
 
 			if (GeometryCache)
 			{
-				if (UUsdGeometryCacheAssetUserData* UserData = UsdUtils::GetOrCreateAssetUserData<UUsdGeometryCacheAssetUserData>(GeometryCache.Get()
-					))
+				UUsdGeometryCacheAssetUserData* UserData = UsdUnreal::ObjectUtils::GetOrCreateAssetUserData<UUsdGeometryCacheAssetUserData>(
+					GeometryCache.Get()
+				);
+				if (UserData)
 				{
 					UserData->PrimvarToUVIndex = LODIndexToMaterialInfo[0].PrimvarToUVIndex;	// We use the same primvar mapping for all LODs
 					UserData->LayerStartOffsetSeconds = StartTimeOffset;
@@ -920,7 +922,7 @@ void FGeometryCacheCreateAssetsTaskChain::SetupTasks()
 						GetPrim(),
 						LODIndexToMaterialInfo,
 						*GeometryCache,
-						*Context->AssetCache.Get(),
+						*Context->UsdAssetCache.Get(),
 						Context->InfoCache.Get(),
 						Context->Time,
 						Context->ObjectFlags,
@@ -935,8 +937,10 @@ void FGeometryCacheCreateAssetsTaskChain::SetupTasks()
 				}
 			}
 
+			const bool bIsImporting = GForceImport || Context->bIsImporting;
+
 			// Continue with the import steps
-			return Context->bIsImporting && GeometryCache && bIsNew;
+			return bIsImporting && GeometryCache && bIsNew;
 		}
 	);
 
@@ -1016,10 +1020,8 @@ USceneComponent* FUsdGeometryCacheTranslator::CreateComponents()
 			return nullptr;
 		}
 
-		SceneComponent = CreateComponentsEx(
-			{Context->bIsImporting ? UGeometryCacheComponent::StaticClass() : UGeometryCacheUsdComponent::StaticClass()},
-			{}
-		);
+		const bool bIsImporting = GForceImport || Context->bIsImporting;
+		SceneComponent = CreateComponentsEx({bIsImporting ? UGeometryCacheComponent::StaticClass() : UGeometryCacheUsdComponent::StaticClass()}, {});
 	}
 	else
 	{
@@ -1030,7 +1032,7 @@ USceneComponent* FUsdGeometryCacheTranslator::CreateComponents()
 
 	if (UGeometryCacheComponent* Component = Cast<UGeometryCacheComponent>(SceneComponent))
 	{
-		if (Context->InfoCache && Context->AssetCache)
+		if (Context->InfoCache && Context->UsdAssetCache)
 		{
 			if (UGeometryCache* GeometryCache = Context->InfoCache->GetSingleAssetForPrim<UGeometryCache>(PrimPath))
 			{
@@ -1041,7 +1043,7 @@ USceneComponent* FUsdGeometryCacheTranslator::CreateComponents()
 					GetPrim(),
 					GeometryCache->Materials,
 					*Component,
-					*Context->AssetCache,
+					*Context->UsdAssetCache,
 					*Context->InfoCache,
 					Context->Time,
 					Context->ObjectFlags,
@@ -1056,7 +1058,7 @@ USceneComponent* FUsdGeometryCacheTranslator::CreateComponents()
 				{
 					UsdGroomTranslatorUtils::CreateGroomBindingAsset(
 						GetPrim(),
-						*Context->AssetCache,
+						*Context->UsdAssetCache,
 						*Context->InfoCache,
 						Context->ObjectFlags,
 						Context->bReuseIdenticalAssets
@@ -1181,10 +1183,12 @@ void FUsdGeometryCacheTranslator::UpdateComponents(USceneComponent* SceneCompone
 			UsdGroomTranslatorUtils::SetGroomFromPrim(GetPrim(), *Context->InfoCache, SceneComponent);
 		}
 
+		const bool bIsImporting = GForceImport || Context->bIsImporting;
+
 		// Defer to xformable translator to set our transforms, visibility, etc. but only when opening the stage: This will be baked in for import.
 		// Don't go through FUsdGeomMeshTranslator::UpdateComponents as it will want to create a static mesh if PrimPath is an animated mesh prim
 		// (which is likely, given that we're running this FUsdGeometryCacheTranslator for it)
-		if (!Context->bIsImporting)
+		if (!bIsImporting)
 		{
 			FUsdGeomXformableTranslator::UpdateComponents(GeometryCacheComponent);
 		}
diff --git a/Engine/Plugins/Importers/USDImporter/Source/USDSchemas/Private/USDGroomTranslator.cpp b/Engine/Plugins/Importers/USDImporter/Source/USDSchemas/Private/USDGroomTranslator.cpp
--- a/Engine/Plugins/Importers/USDImporter/Source/USDSchemas/Private/USDGroomTranslator.cpp
+++ b/Engine/Plugins/Importers/USDImporter/Source/USDSchemas/Private/USDGroomTranslator.cpp
@@ -25,7 +25,9 @@
 #include "USDGroomTranslatorUtils.h"
 #include "USDIntegrationUtils.h"
 #include "USDLog.h"
+#include "USDObjectUtils.h"
 #include "USDPrimConversion.h"
+#include "USDTranslatorUtils.h"
 
 #include "USDIncludesStart.h"
 #include "pxr/usd/usd/prim.h"
@@ -131,7 +133,8 @@ protected:
 	FHairDescription HairDescription;
 
 	TStrongObjectPtr<UGroomImportOptions> ImportOptions;
-	TUniquePtr<FGroomCacheProcessor> GroomCacheProcessor;
+	UGroomAsset* GroomAsset = nullptr;
+	UGroomCache* GroomCache = nullptr;
 	FGroomAnimationInfo AnimInfo;
 	FString PrefixedGroomCacheHash;
 
@@ -170,6 +173,7 @@ protected:
 
 			   return bSuccess && HairDescription.IsValid();
 		   });
+
 		// Build groom asset from hair description (Sync)
 		Then(
 			ESchemaTranslationLaunchPolicy::Sync,
@@ -181,57 +185,50 @@ protected:
 				FHairDescriptionGroups GroupsDescription;
 				FGroomBuilder::BuildHairDescriptionGroups(HairDescription, GroupsDescription);
 
-				ImportOptions.Reset(UE::UsdGroomTranslator::Private::CreateGroomImportOptions(GroupsDescription, Context->GroomInterpolationSettings)
+				ImportOptions.Reset(	//
+					UE::UsdGroomTranslator::Private::CreateGroomImportOptions(GroupsDescription, Context->GroomInterpolationSettings)
 				);
 
 				FSHAHash SHAHash = UE::UsdGroomTranslator::Private::ComputeHairDescriptionHash(HairDescription, ImportOptions->InterpolationSettings);
 
-				FString PrefixedAssetHash = UsdUtils::GetAssetHashPrefix(GetPrim(), Context->bReuseIdenticalAssets) + SHAHash.ToString();
+				const FString PrefixedAssetHash = UsdUtils::GetAssetHashPrefix(GetPrim(), Context->bReuseIdenticalAssets) + SHAHash.ToString();
 
 				const FString PrimPathString = PrimPath.GetString();
-				UGroomAsset* GroomAsset = Cast<UGroomAsset>(Context->AssetCache->GetCachedAsset(PrefixedAssetHash));
-				if (!GroomAsset)
-				{
-					FName AssetName = MakeUniqueObjectName(
-						GetTransientPackage(),
-						UGroomAsset::StaticClass(),
-						*IUsdClassesModule::SanitizeObjectName(FPaths::GetBaseFilename(PrimPathString))
-					);
-
-					FHairImportContext HairImportContext(
-						ImportOptions.Get(),
-						GetTransientPackage(),
-						UGroomAsset::StaticClass(),
-						AssetName,
-						Context->ObjectFlags | RF_Public | RF_Transient
-					);
-					UGroomAsset* ExistingAsset = nullptr;
-					GroomAsset = FHairStrandsImporter::ImportHair(HairImportContext, HairDescription, ExistingAsset);
-					if (GroomAsset)
+				const FString DesiredName = FPaths::GetBaseFilename(PrimPathString);
+
+				bool bGroomAssetIsNew = false;
+				GroomAsset = Context->UsdAssetCache->GetOrCreateCustomCachedAsset<UGroomAsset>(
+					PrefixedAssetHash,
+					DesiredName,
+					Context->ObjectFlags,
+					[this](UPackage* Outer, FName SanitizedName, EObjectFlags FlagsToUse)
 					{
-						Context->AssetCache->CacheAsset(PrefixedAssetHash, GroomAsset);
+						FHairImportContext HairImportContext(ImportOptions.Get(), Outer, UGroomAsset::StaticClass(), SanitizedName, FlagsToUse);
+						UGroomAsset* ExistingAsset = nullptr;
+						return FHairStrandsImporter::ImportHair(HairImportContext, HairDescription, ExistingAsset);
+					},
+					&bGroomAssetIsNew
+				);
 
-						if (UUsdAssetUserData* UserData = UsdUtils::GetOrCreateAssetUserData(GroomAsset))
-						{
-							UserData->PrimPaths.AddUnique(PrimPathString);
+				if (UUsdAssetUserData* UserData = UsdUnreal::ObjectUtils::GetOrCreateAssetUserData(GroomAsset))
+				{
+					UserData->PrimPaths.AddUnique(PrimPathString);
 
-							if (Context->MetadataOptions.bCollectMetadata)
-							{
-								UsdToUnreal::ConvertMetadata(
-									GetPrim(),
-									UserData,
-									Context->MetadataOptions.BlockedPrefixFilters,
-									Context->MetadataOptions.bInvertFilters,
-									Context->MetadataOptions.bCollectFromEntireSubtrees
-								);
-							}
-							else
-							{
-								// Strip the metadata from this prim, so that if we uncheck "Collect Metadata" it actually disappears on the
-								// AssetUserData
-								UserData->StageIdentifierToMetadata.Remove(GetPrim().GetStage().GetRootLayer().GetIdentifier());
-							}
-						}
+					if (Context->MetadataOptions.bCollectMetadata)
+					{
+						UsdToUnreal::ConvertMetadata(
+							GetPrim(),
+							UserData,
+							Context->MetadataOptions.BlockedPrefixFilters,
+							Context->MetadataOptions.bInvertFilters,
+							Context->MetadataOptions.bCollectFromEntireSubtrees
+						);
+					}
+					else
+					{
+						// Strip the metadata from this prim, so that if we uncheck "Collect Metadata" it actually disappears on the
+						// AssetUserData
+						UserData->StageIdentifierToMetadata.Remove(GetPrim().GetStage().GetRootLayer().GetIdentifier());
 					}
 				}
 
@@ -244,7 +241,8 @@ protected:
 				return GroomAsset != nullptr && AnimInfo.IsValid();
 			}
 		);
-		// Parse GroomCache data into processor (Async)
+
+		// Compute the GroomCache hash (Async)
 		Then(
 			ESchemaTranslationLaunchPolicy::Async,
 			[this]() -> bool
@@ -254,7 +252,7 @@ protected:
 					return false;
 				}
 
-				UGroomAsset* GroomAsset = Context->InfoCache->GetSingleAssetForPrim<UGroomAsset>(PrimPath);
+				GroomAsset = Context->InfoCache->GetSingleAssetForPrim<UGroomAsset>(PrimPath);
 				if (!GroomAsset)
 				{
 					return false;
@@ -284,173 +282,193 @@ protected:
 				SHA1.GetHash(Hash.Hash);
 
 				PrefixedGroomCacheHash = UsdUtils::GetAssetHashPrefix(GetPrim(), Context->bReuseIdenticalAssets) + Hash.ToString();
+				return true;
+			}
+		);
 
-				bool bSuccess = true;
-				UGroomCache* GroomCache = Cast<UGroomCache>(Context->AssetCache->GetCachedAsset(PrefixedGroomCacheHash));
-				if (!GroomCache)
-				{
-					GroomCacheProcessor = MakeUnique<FGroomCacheProcessor>(EGroomCacheType::Strands, AnimInfo.Attributes);
+		// Get or create GroomCache asset (Sync)
+		Then(
+			ESchemaTranslationLaunchPolicy::Sync,
+			[this]() -> bool
+			{
+				const FString StrandsGroomCachePrimPath = UsdGroomTranslatorUtils::GetStrandsGroomCachePrimPath(PrimPath);
+				const FString DesiredName = FPaths::GetBaseFilename(StrandsGroomCachePrimPath);
+
+				bool bGroomCacheIsNew = false;
+				GroomCache = Context->UsdAssetCache->GetOrCreateCachedAsset<UGroomCache>(	 //
+					PrefixedGroomCacheHash,
+					DesiredName,
+					Context->ObjectFlags,
+					&bGroomCacheIsNew
+				);
 
-					// ref. FGroomCacheImporter::ImportGroomCache
-					const TArray<FHairGroupPlatformData>& GroomHairGroupsData = GroomAsset->GetHairGroupsPlatformData();
+				if (GroomCache && Context->InfoCache)
+				{
+					Context->InfoCache->LinkAssetToPrim(PrimPath, GroomCache);
 
-					// Each frame is translated into a HairDescription and processed into HairGroupData
-					// Sample one extra frame so that we can interpolate between EndFrame - 1 and EndFrame
-					for (int32 FrameIndex = AnimInfo.StartFrame; FrameIndex < AnimInfo.EndFrame + 1; ++FrameIndex)
+					if (UUsdAssetUserData* UserData = UsdUnreal::ObjectUtils::GetOrCreateAssetUserData(GroomCache))
 					{
-						FHairDescription FrameHairDescription;
-						bSuccess = UsdToUnreal::ConvertGroomHierarchy(
-							GetPrim(),
-							pxr::UsdTimeCode(FrameIndex),
-							FTransform::Identity,
-							FrameHairDescription
-						);
+						UserData->PrimPaths.AddUnique(PrimPath.GetString());
 
-						if (!bSuccess)
+						if (Context->MetadataOptions.bCollectMetadata)
 						{
-							break;
+							UsdToUnreal::ConvertMetadata(
+								GetPrim(),
+								UserData,
+								Context->MetadataOptions.BlockedPrefixFilters,
+								Context->MetadataOptions.bInvertFilters,
+								Context->MetadataOptions.bCollectFromEntireSubtrees
+							);
 						}
-
-						FHairDescriptionGroups HairDescriptionGroups;
-						if (!FGroomBuilder::BuildHairDescriptionGroups(FrameHairDescription, HairDescriptionGroups))
+						else
 						{
-							bSuccess = false;
-							break;
+							// Strip the metadata from this prim, so that if we uncheck "Collect Metadata" it actually disappears on the AssetUserData
+							UserData->StageIdentifierToMetadata.Remove(GetPrim().GetStage().GetRootLayer().GetIdentifier());
 						}
+					}
+				}
 
-						const uint32 GroupCount = HairDescriptionGroups.HairGroups.Num();
+				// We only need to continue if we need to finish building this groom cache
+				return GroomCache && bGroomCacheIsNew;
+			}
+		);
 
-						TArray<FHairGroupInfoWithVisibility> HairGroupsInfo = GroomAsset->GetHairGroupsInfo();
-						TArray<FHairDescriptionGroup> HairGroupsData;
-						HairGroupsData.SetNum(GroupCount);
-						for (uint32 GroupIndex = 0; GroupIndex < GroupCount; ++GroupIndex)
-						{
-							const FHairDescriptionGroup& HairGroup = HairDescriptionGroups.HairGroups[GroupIndex];
-							FHairDescriptionGroup& HairGroupData = HairGroupsData[GroupIndex];
-							FGroomBuilder::BuildData(
-								HairGroup,
-								GroomAsset->GetHairGroupsInterpolation()[GroupIndex],
-								HairGroupsInfo[GroupIndex],
-								HairGroupData.Strands,
-								HairGroupData.Guides
-							);
-						}
+		// Parse GroomCache data into processor (Async)
+		Then(
+			ESchemaTranslationLaunchPolicy::Async,
+			[this]() -> bool
+			{
+				if (!GroomAsset || !GroomCache)
+				{
+					return false;
+				}
 
-						// Validate that the GroomCache has the same topology as the static groom
-						if (HairGroupsData.Num() == GroomHairGroupsData.Num())
+				FGroomCacheProcessor GroomCacheProcessor{EGroomCacheType::Strands, AnimInfo.Attributes};
+
+				// ref. FGroomCacheImporter::ImportGroomCache
+				const TArray<FHairGroupPlatformData>& GroomHairGroupsData = GroomAsset->GetHairGroupsPlatformData();
+
+				// Process data and store results into GroomCacheProcessor.
+				// Each frame is translated into a HairDescription and processed into HairGroupData
+				// Sample one extra frame so that we can interpolate between EndFrame - 1 and EndFrame
+				bool bSuccess = true;
+				for (int32 FrameIndex = AnimInfo.StartFrame; FrameIndex < AnimInfo.EndFrame + 1; ++FrameIndex)
+				{
+					FHairDescription FrameHairDescription;
+					bSuccess = UsdToUnreal::ConvertGroomHierarchy(
+						GetPrim(),
+						pxr::UsdTimeCode(FrameIndex),
+						FTransform::Identity,
+						FrameHairDescription
+					);
+
+					if (!bSuccess)
+					{
+						break;
+					}
+
+					FHairDescriptionGroups HairDescriptionGroups;
+					if (!FGroomBuilder::BuildHairDescriptionGroups(FrameHairDescription, HairDescriptionGroups))
+					{
+						bSuccess = false;
+						break;
+					}
+
+					const uint32 GroupCount = HairDescriptionGroups.HairGroups.Num();
+
+					TArray<FHairGroupInfoWithVisibility> HairGroupsInfo = GroomAsset->GetHairGroupsInfo();
+					TArray<FHairDescriptionGroup> HairGroupsData;
+					HairGroupsData.SetNum(GroupCount);
+					for (uint32 GroupIndex = 0; GroupIndex < GroupCount; ++GroupIndex)
+					{
+						const FHairDescriptionGroup& HairGroup = HairDescriptionGroups.HairGroups[GroupIndex];
+						FHairDescriptionGroup& HairGroupData = HairGroupsData[GroupIndex];
+						FGroomBuilder::BuildData(
+							HairGroup,
+							GroomAsset->GetHairGroupsInterpolation()[GroupIndex],
+							HairGroupsInfo[GroupIndex],
+							HairGroupData.Strands,
+							HairGroupData.Guides
+						);
+					}
+
+					// Validate that the GroomCache has the same topology as the static groom
+					if (HairGroupsData.Num() == GroomHairGroupsData.Num())
+					{
+						for (uint32 GroupIndex = 0; GroupIndex < GroupCount; ++GroupIndex)
 						{
-							for (uint32 GroupIndex = 0; GroupIndex < GroupCount; ++GroupIndex)
+							if (HairGroupsData[GroupIndex].Strands.GetNumPoints() != GroomHairGroupsData[GroupIndex].Strands.BulkData.GetNumPoints())
 							{
-								if (HairGroupsData[GroupIndex].Strands.GetNumPoints()
-									!= GroomHairGroupsData[GroupIndex].Strands.BulkData.GetNumPoints())
-								{
-									bSuccess = false;
-									UE_LOG(
-										LogUsd,
-										Warning,
-										TEXT("GroomCache frame %d does not have the same number of vertices as the static groom (%u instead of %u). "
-											 "Aborting GroomCache import."),
-										FrameIndex,
-										HairGroupsData[GroupIndex].Strands.GetNumPoints(),
-										GroomHairGroupsData[GroupIndex].Strands.BulkData.GetNumPoints()
-									);
-									break;
-								}
+								bSuccess = false;
+								UE_LOG(
+									LogUsd,
+									Warning,
+									TEXT("GroomCache frame %d does not have the same number of vertices as the static groom (%u instead of %u). "
+										 "Aborting GroomCache import."),
+									FrameIndex,
+									HairGroupsData[GroupIndex].Strands.GetNumPoints(),
+									GroomHairGroupsData[GroupIndex].Strands.BulkData.GetNumPoints()
+								);
+								break;
 							}
 						}
-						else
-						{
-							bSuccess = false;
-							UE_LOG(
-								LogUsd,
-								Warning,
-								TEXT("GroomCache does not have the same number of groups as the static groom (%d instead of %d). Aborting GroomCache "
-									 "import."),
-								HairGroupsData.Num(),
-								GroomHairGroupsData.Num()
-							);
-						}
-
-						if (!bSuccess)
-						{
-							break;
-						}
+					}
+					else
+					{
+						bSuccess = false;
+						UE_LOG(
+							LogUsd,
+							Warning,
+							TEXT("GroomCache does not have the same number of groups as the static groom (%d instead of %d). Aborting GroomCache "
+								 "import."),
+							HairGroupsData.Num(),
+							GroomHairGroupsData.Num()
+						);
+					}
 
-						// The HairGroupData is converted into animated groom data by the GroomCacheProcessor
-						GroomCacheProcessor->AddGroomSample(MoveTemp(HairGroupsData));
+					if (!bSuccess)
+					{
+						break;
 					}
+
+					// The HairGroupData is converted into animated groom data by the GroomCacheProcessor
+					GroomCacheProcessor.AddGroomSample(MoveTemp(HairGroupsData));
 				}
 
-				if (Context->InfoCache && GroomCache)
+				if (bSuccess)
 				{
-					Context->InfoCache->LinkAssetToPrim(PrimPath, GroomCache);
+					// We used to call FGroomCacheImporter::ProcessToGroomCache here, but all it does is a bunch of package
+					// management for reimport/creation, which the asset cache will already do. So here just copy what it does
+					// internally, which luckily is all public
+					GroomCache->Initialize(GroomCacheProcessor.GetType());
+					GroomCacheProcessor.TransferChunks(GroomCache);
+					GroomCache->SetGroomAnimationInfo(AnimInfo);
 				}
-
-				if (!bSuccess)
+				else
 				{
-					GroomCacheProcessor.Reset();
+					const FString StrandsGroomCachePrimPath = UsdGroomTranslatorUtils::GetStrandsGroomCachePrimPath(PrimPath);
+					UE_LOG(LogUsd, Warning, TEXT("Failed to create GroomCache for prim '%s'"), *StrandsGroomCachePrimPath);
+
+					UsdUnreal::TranslatorUtils::AbandonFailedAsset(GroomCache, Context->UsdAssetCache.Get(), Context->InfoCache.Get());
 				}
 
-				// Go to the next step only if it needs to create a GroomCache asset
-				return GroomCache == nullptr && GroomCacheProcessor;
+				return bSuccess;
 			}
 		);
-		// Create GroomCache asset from processor (Sync)
+
+		// Post import groom cache
 		Then(
 			ESchemaTranslationLaunchPolicy::Sync,
 			[this]() -> bool
 			{
-				// TEMP: This is a small trick to prevent two concurrent task chains from running into a hash collision in the asset cache.
-				// This is enough of a workaround because this is a Sync task, so we can guarantee only one of the competing task chains will be
-				// run at a time. Whichever wins gets to create the GroomCache, and the other will exit through this branch.
-				// It will likely be properly fixed before 5.4 is out, but check UE-201011 for more details.
-				UGroomCache* ExistingGroomCache = Cast<UGroomCache>(Context->AssetCache->GetCachedAsset(PrefixedGroomCacheHash));
-				if (ExistingGroomCache)
-				{
-					Context->InfoCache->LinkAssetToPrim(PrimPath, ExistingGroomCache);
-					return false;
-				}
-
-				const FString StrandsGroomCachePrimPath = UsdGroomTranslatorUtils::GetStrandsGroomCachePrimPath(PrimPath);
-				FHairImportContext
-					HairImportContext(nullptr, GetTransientPackage(), nullptr, FName(), Context->ObjectFlags | RF_Public | RF_Transient);
-				FName UniqueName = MakeUniqueObjectName(
-					GetTransientPackage(),
-					UGroomCache::StaticClass(),
-					*IUsdClassesModule::SanitizeObjectName(FPaths::GetBaseFilename(StrandsGroomCachePrimPath))
-				);
-
-				// Once the processing has completed successfully, the data is transferred to the GroomCache
-				UGroomCache*
-					GroomCache = FGroomCacheImporter::ProcessToGroomCache(*GroomCacheProcessor, AnimInfo, HairImportContext, UniqueName.ToString());
 				if (GroomCache)
 				{
-					if (UUsdAssetUserData* UserData = UsdUtils::GetOrCreateAssetUserData(GroomCache))
-					{
-						UserData->PrimPaths.AddUnique(PrimPath.GetString());
-
-						if (Context->MetadataOptions.bCollectMetadata)
-						{
-							UsdToUnreal::ConvertMetadata(
-								GetPrim(),
-								UserData,
-								Context->MetadataOptions.BlockedPrefixFilters,
-								Context->MetadataOptions.bInvertFilters,
-								Context->MetadataOptions.bCollectFromEntireSubtrees
-							);
-						}
-						else
-						{
-							// Strip the metadata from this prim, so that if we uncheck "Collect Metadata" it actually disappears on the AssetUserData
-							UserData->StageIdentifierToMetadata.Remove(GetPrim().GetStage().GetRootLayer().GetIdentifier());
-						}
-					}
-
-					Context->AssetCache->CacheAsset(PrefixedGroomCacheHash, GroomCache);
-					Context->InfoCache->LinkAssetToPrim(PrimPath, GroomCache);
+					// This needs to be called from the main thread
+					GroomCache->MarkPackageDirty();
+					GroomCache->PostEditChange();
 				}
 
-				return GroomCache != nullptr;
+				return false;	 // No other tasks
 			}
 		);
 	}
diff --git a/Engine/Plugins/Importers/USDImporter/Source/USDSchemas/Private/USDGroomTranslatorUtils.cpp b/Engine/Plugins/Importers/USDImporter/Source/USDSchemas/Private/USDGroomTranslatorUtils.cpp
--- a/Engine/Plugins/Importers/USDImporter/Source/USDSchemas/Private/USDGroomTranslatorUtils.cpp
+++ b/Engine/Plugins/Importers/USDImporter/Source/USDSchemas/Private/USDGroomTranslatorUtils.cpp
@@ -5,7 +5,7 @@
 #if USE_USD_SDK && WITH_EDITOR
 
 #include "UnrealUSDWrapper.h"
-#include "USDAssetCache2.h"
+#include "USDAssetCache3.h"
 #include "USDClassesModule.h"
 #include "USDInfoCache.h"
 #include "USDIntegrationUtils.h"
@@ -91,47 +91,6 @@ namespace UE::UsdGroomTranslatorUtils::Private
 		return SHAHash;
 	}
 
-	UGroomBindingAsset* CreateGroomBindingAsset(FString GroomBindingPath, const FGroomBindingBuildSettings& Settings, EObjectFlags ObjectFlags)
-	{
-		// Need at least the groom and the target mesh to build a GroomBindingAsset; the source mesh is optional
-		if (GroomBindingPath.IsEmpty() || !Settings.Groom || !Settings.TargetMesh)
-		{
-			return nullptr;
-		}
-
-		const FName BindingAssetName = MakeUniqueObjectName(
-			GetTransientPackage(),
-			UGroomBindingAsset::StaticClass(),
-			*IUsdClassesModule::SanitizeObjectName(FPaths::GetBaseFilename(GroomBindingPath))
-		);
-		UGroomBindingAsset* GroomBinding = NewObject<UGroomBindingAsset>(
-			GetTransientPackage(),
-			BindingAssetName,
-			ObjectFlags | RF_Public | RF_Transient
-		);
-		if (GroomBinding)
-		{
-			GroomBinding->SetGroomBindingType(Settings.GroomBindingType);
-			GroomBinding->SetGroom(Settings.Groom);
-			if (GroomBinding->GetGroomBindingType() == EGroomBindingMeshType::SkeletalMesh)
-			{
-				GroomBinding->SetSourceSkeletalMesh(Cast<USkeletalMesh>(Settings.SourceMesh));
-				GroomBinding->SetTargetSkeletalMesh(Cast<USkeletalMesh>(Settings.TargetMesh));
-			}
-			else
-			{
-				GroomBinding->SetSourceGeometryCache(Cast<UGeometryCache>(Settings.SourceMesh));
-				GroomBinding->SetTargetGeometryCache(Cast<UGeometryCache>(Settings.TargetMesh));
-			}
-			GroomBinding->GetHairGroupsPlatformData().Reserve(Settings.Groom->GetHairGroupsPlatformData().Num());
-			GroomBinding->SetNumInterpolationPoints(Settings.NumInterpolationPoints);
-			GroomBinding->SetMatchingSection(Settings.MatchingSection);
-
-			GroomBinding->Build();
-		}
-		return GroomBinding;
-	}
-
 	FString GetGroomPrimPath(const pxr::UsdPrim& Prim)
 	{
 		FScopedUsdAllocs Allocs;
@@ -192,7 +151,7 @@ namespace UsdGroomTranslatorUtils
 
 	void CreateGroomBindingAsset(
 		const pxr::UsdPrim& Prim,
-		UUsdAssetCache2& AssetCache,
+		UUsdAssetCache3& AssetCache,
 		FUsdInfoCache& InfoCache,
 		EObjectFlags ObjectFlags,
 		bool bReuseIdenticalAssets
@@ -237,22 +196,48 @@ namespace UsdGroomTranslatorUtils
 
 		const FString PrefixedBindingHash = UsdUtils::GetAssetHashPrefix(Prim, bReuseIdenticalAssets) + ComputeGroomBindingHash(Settings).ToString();
 
-		// Try to get the GroomBindingAsset from the cache
-		UGroomBindingAsset* GroomBinding = Cast<UGroomBindingAsset>(AssetCache.GetCachedAsset(PrefixedBindingHash));
-
 		const FString GroomBindingPath = FString::Printf(TEXT("%s_groombinding"), *PrimPath);
-		if (!GroomBinding)
+
+		bool bIsNew = false;
+		UGroomBindingAsset* GroomBinding = nullptr;
+
+		const FString DesiredName = FPaths::GetBaseFilename(GroomBindingPath);
+
+		// Prevent the creation of the GroomBinding or it's internal UObjects (if any) from being added to the transaction.
+		// This is important because we call CreateGroomBindingAsset from CreateComponents(), which is outside of the normal
+		// transaction suppression around the CreateAssets() calls that the AUsdStageActor will already do in AUsdStageActor::LoadAssets
+		// and AUsdStageActor::LoadAsset
+		TGuardValue<ITransaction*> SuppressTransaction{GUndo, nullptr};
+
+		// Need at least the groom and the target mesh to build a GroomBindingAsset; the source mesh is optional
+		if (!GroomBindingPath.IsEmpty() && GroomAsset && TargetMesh)
 		{
-			// Create and cache it, if it didn't exist already
-			GroomBinding = CreateGroomBindingAsset(GroomBindingPath, Settings, ObjectFlags);
-			if (GroomBinding)
-			{
-				AssetCache.CacheAsset(PrefixedBindingHash, GroomBinding);
-			}
+			GroomBinding = AssetCache.GetOrCreateCachedAsset<UGroomBindingAsset>(PrefixedBindingHash, DesiredName, ObjectFlags, &bIsNew);
 		}
 
 		if (GroomBinding)
 		{
+			if (bIsNew)
+			{
+				GroomBinding->SetGroomBindingType(Settings.GroomBindingType);
+				GroomBinding->SetGroom(Settings.Groom);
+				if (GroomBinding->GetGroomBindingType() == EGroomBindingMeshType::SkeletalMesh)
+				{
+					GroomBinding->SetSourceSkeletalMesh(Cast<USkeletalMesh>(Settings.SourceMesh));
+					GroomBinding->SetTargetSkeletalMesh(Cast<USkeletalMesh>(Settings.TargetMesh));
+				}
+				else
+				{
+					GroomBinding->SetSourceGeometryCache(Cast<UGeometryCache>(Settings.SourceMesh));
+					GroomBinding->SetTargetGeometryCache(Cast<UGeometryCache>(Settings.TargetMesh));
+				}
+				GroomBinding->GetHairGroupsPlatformData().Reserve(Settings.Groom->GetHairGroupsPlatformData().Num());
+				GroomBinding->SetNumInterpolationPoints(Settings.NumInterpolationPoints);
+				GroomBinding->SetMatchingSection(Settings.MatchingSection);
+
+				GroomBinding->Build();
+			}
+
 			InfoCache.LinkAssetToPrim(UE::FSdfPath{*GroomBindingPath}, GroomBinding);
 		}
 	}
diff --git a/Engine/Plugins/Importers/USDImporter/Source/USDSchemas/Private/USDGroomTranslatorUtils.h b/Engine/Plugins/Importers/USDImporter/Source/USDSchemas/Private/USDGroomTranslatorUtils.h
--- a/Engine/Plugins/Importers/USDImporter/Source/USDSchemas/Private/USDGroomTranslatorUtils.h
+++ b/Engine/Plugins/Importers/USDImporter/Source/USDSchemas/Private/USDGroomTranslatorUtils.h
@@ -17,7 +17,7 @@ PXR_NAMESPACE_CLOSE_SCOPE
 
 class FUsdInfoCache;
 class USceneComponent;
-class UUsdAssetCache2;
+class UUsdAssetCache3;
 
 namespace UE
 {
@@ -29,7 +29,7 @@ namespace UsdGroomTranslatorUtils
 	// Create the groom binding asset for the given Prim with GroomBindingAPI and cache it in the AssetCache
 	void CreateGroomBindingAsset(
 		const pxr::UsdPrim& Prim,
-		UUsdAssetCache2& AssetCache,
+		UUsdAssetCache3& AssetCache,
 		FUsdInfoCache& InfoCache,
 		EObjectFlags ObjectFlags,
 		bool bReuseIdenticalAssets
diff --git a/Engine/Plugins/Importers/USDImporter/Source/USDSchemas/Private/USDInfoCache.cpp b/Engine/Plugins/Importers/USDImporter/Source/USDSchemas/Private/USDInfoCache.cpp
--- a/Engine/Plugins/Importers/USDImporter/Source/USDSchemas/Private/USDInfoCache.cpp
+++ b/Engine/Plugins/Importers/USDImporter/Source/USDSchemas/Private/USDInfoCache.cpp
@@ -1264,8 +1264,11 @@ namespace UE::USDInfoCacheImpl::Private
 			}
 		);
 
+		static IConsoleVariable* ForceImportCvar = IConsoleManager::Get().FindConsoleVariable(TEXT("USD.GeometryCache.ForceImport"));
+		const bool bIsImporting = Context.bIsImporting || (ForceImportCvar && ForceImportCvar->GetBool());
+
 		bool bIsAnimatedMesh = UsdUtils::IsAnimatedMesh(UsdPrim);
-		if (!Context.bIsImporting)
+		if (!bIsImporting)
 		{
 			FWriteScopeLock ScopeLock(Impl.InfoMapLock);
 
@@ -1590,10 +1593,7 @@ void FUsdInfoCache::UnlinkAssetFromPrim(const UE::FSdfPath& Path, UObject* Asset
 	}
 	FWriteScopeLock ScopeLock(ImplPtr->PrimPathToAssetsLock);
 
-	UE_LOG(LogUsd, Verbose, TEXT("Unlinking asset '%s' to prim '%s'"),
-		*Asset->GetPathName(),
-		*Path.GetString()
-	);
+	UE_LOG(LogUsd, Verbose, TEXT("Unlinking asset '%s' to prim '%s'"), *Asset->GetPathName(), *Path.GetString());
 
 	if (TArray<TWeakObjectPtr<UObject>>* FoundAssetsForPrim = ImplPtr->PrimPathToAssets.Find(Path))
 	{
diff --git a/Engine/Plugins/Importers/USDImporter/Source/USDSchemas/Private/USDLuxLightTranslator.cpp b/Engine/Plugins/Importers/USDImporter/Source/USDSchemas/Private/USDLuxLightTranslator.cpp
--- a/Engine/Plugins/Importers/USDImporter/Source/USDSchemas/Private/USDLuxLightTranslator.cpp
+++ b/Engine/Plugins/Importers/USDImporter/Source/USDSchemas/Private/USDLuxLightTranslator.cpp
@@ -2,10 +2,14 @@
 
 #include "USDLuxLightTranslator.h"
 
+#include "USDAssetUserData.h"
 #include "USDConversionUtils.h"
 #include "USDDrawModeComponent.h"
 #include "USDLightConversion.h"
-
+#include "USDLog.h"
+#include "USDObjectUtils.h"
+#include "USDShadeConversion.h"
+#include "USDTypesConversion.h"
 #include "UsdWrappers/UsdPrim.h"
 
 #include "Components/DirectionalLightComponent.h"
@@ -13,15 +17,95 @@
 #include "Components/RectLightComponent.h"
 #include "Components/SkyLightComponent.h"
 #include "Components/SpotLightComponent.h"
+#include "Engine/TextureCube.h"
 
 #if USE_USD_SDK
 
 #include "USDIncludesStart.h"
 #include "pxr/usd/usdLux/diskLight.h"
+#include "pxr/usd/usdLux/domeLight.h"
 #include "pxr/usd/usdLux/lightAPI.h"
 #include "pxr/usd/usdLux/rectLight.h"
 #include "USDIncludesEnd.h"
 
+void FUsdLuxLightTranslator::CreateAssets()
+{
+	TRACE_CPUPROFILER_EVENT_SCOPE(FUsdLuxLightTranslator::CreateAssets);
+
+	if (!Context->UsdAssetCache || !Context->InfoCache)
+	{
+		return;
+	}
+
+	// Don't bother generating assets if we're going to just draw some bounds for this prim instead
+	EUsdDrawMode DrawMode = UsdUtils::GetAppliedDrawMode(GetPrim());
+	if (DrawMode != EUsdDrawMode::Default)
+	{
+		CreateAlternativeDrawModeAssets(DrawMode);
+		return;
+	}
+
+	const FString VolumePrimPathString = PrimPath.GetString();
+	pxr::UsdPrim Prim = GetPrim();
+	pxr::UsdLuxDomeLight DomeLight{Prim};
+	if (!DomeLight)
+	{
+		// Only dome lights make assets for now
+		return;
+	}
+
+	const FString ResolvedDomeTexturePath = UsdUtils::GetResolvedAssetPath(DomeLight.GetTextureFileAttr());
+	if (ResolvedDomeTexturePath.IsEmpty())
+	{
+		FScopedUsdAllocs Allocs;
+
+		pxr::SdfAssetPath TextureAssetPath;
+		DomeLight.GetTextureFileAttr().Get<pxr::SdfAssetPath>(&TextureAssetPath);
+
+		// Show a good warning for this because it's easy to pick some cubemap asset from the engine (that usually don't come with the
+		// source texture) and have the dome light silently not work again
+		FString TargetAssetPath = UsdToUnreal::ConvertString(TextureAssetPath.GetAssetPath());
+		UE_LOG(
+			LogUsd,
+			Warning,
+			TEXT("Failed to find texture '%s' used for UsdLuxDomeLight '%s'!"),
+			*TargetAssetPath,
+			*UsdToUnreal::ConvertPath(DomeLight.GetPrim().GetPath())
+		);
+
+		return;
+	}
+
+	const FString PrefixedTextureHash = UsdUtils::GetAssetHashPrefix(Prim, Context->bReuseIdenticalAssets)
+										+ LexToString(FMD5Hash::HashFile(*ResolvedDomeTexturePath));
+
+	const FString& DesiredTextureName = FPaths::GetBaseFilename(ResolvedDomeTexturePath);
+
+	TextureGroup Group = TextureGroup::TEXTUREGROUP_Skybox;
+
+	bool bCreatedTexture = false;
+	UTextureCube* Texture = Context->UsdAssetCache->GetOrCreateCustomCachedAsset<UTextureCube>(
+		PrefixedTextureHash,
+		DesiredTextureName,
+		Context->ObjectFlags,
+		[&ResolvedDomeTexturePath, Group](UPackage* Outer, FName SanitizedName, EObjectFlags FlagsToUse)
+		{
+			return UsdUtils::CreateTexture(ResolvedDomeTexturePath, SanitizedName, Group, FlagsToUse, Outer);
+		},
+		&bCreatedTexture
+	);
+
+	if (UUsdAssetUserData* TextureUserData = UsdUnreal::ObjectUtils::GetOrCreateAssetUserData(Texture))
+	{
+		TextureUserData->PrimPaths.AddUnique(UsdToUnreal::ConvertPath(DomeLight.GetPrim().GetPath()));
+	}
+
+	if (Texture)
+	{
+		Context->InfoCache->LinkAssetToPrim(PrimPath, Texture);
+	}
+}
+
 USceneComponent* FUsdLuxLightTranslator::CreateComponents()
 {
 	USceneComponent* SceneComponent = nullptr;
@@ -93,7 +177,14 @@ void FUsdLuxLightTranslator::UpdateComponents(USceneComponent* SceneComponent)
 	}
 	else if (USkyLightComponent* SkyLightComponent = Cast<USkyLightComponent>(SceneComponent))
 	{
-		UsdToUnreal::ConvertDomeLight(Prim, *SkyLightComponent, Context->AssetCache.Get(), Context->bReuseIdenticalAssets);
+		SkyLightComponent->Modify();
+
+		if (UTextureCube* TextuxeCube = Context->InfoCache->GetSingleAssetForPrim<UTextureCube>(PrimPath))
+		{
+			SkyLightComponent->Cubemap = TextuxeCube;
+			SkyLightComponent->SourceType = ESkyLightSourceType::SLS_SpecifiedCubemap;
+		}
+
 		SkyLightComponent->Mobility = EComponentMobility::Movable;	  // We won't bake geometry in the sky light so it needs to be movable
 	}
 
diff --git a/Engine/Plugins/Importers/USDImporter/Source/USDSchemas/Private/USDSchemaTranslator.cpp b/Engine/Plugins/Importers/USDImporter/Source/USDSchemas/Private/USDSchemaTranslator.cpp
--- a/Engine/Plugins/Importers/USDImporter/Source/USDSchemas/Private/USDSchemaTranslator.cpp
+++ b/Engine/Plugins/Importers/USDImporter/Source/USDSchemas/Private/USDSchemaTranslator.cpp
@@ -141,7 +141,15 @@ int32 FUsdSchemaTranslatorRegistry::GetExternalSchemaTranslatorCount()
 
 FUsdSchemaTranslationContext::FUsdSchemaTranslationContext(const UE::FUsdStage& InStage, UUsdAssetCache2& InAssetCache)
 	: Stage(InStage)
-	, AssetCache(&InAssetCache)
+	, UsdAssetCache(NewObject<UUsdAssetCache3>())
+{
+	IUsdSchemasModule& UsdSchemasModule = FModuleManager::Get().LoadModuleChecked<IUsdSchemasModule>(TEXT("USDSchemas"));
+	RenderContext = UsdSchemasModule.GetRenderContextRegistry().GetUniversalRenderContext();
+}
+
+FUsdSchemaTranslationContext::FUsdSchemaTranslationContext(const UE::FUsdStage& InStage, UUsdAssetCache3& InAssetCache)
+	: Stage(InStage)
+	, UsdAssetCache(&InAssetCache)
 {
 	IUsdSchemasModule& UsdSchemasModule = FModuleManager::Get().LoadModuleChecked<IUsdSchemasModule>(TEXT("USDSchemas"));
 	RenderContext = UsdSchemasModule.GetRenderContextRegistry().GetUniversalRenderContext();
diff --git a/Engine/Plugins/Importers/USDImporter/Source/USDSchemas/Private/USDShadeMaterialTranslator.cpp b/Engine/Plugins/Importers/USDImporter/Source/USDSchemas/Private/USDShadeMaterialTranslator.cpp
--- a/Engine/Plugins/Importers/USDImporter/Source/USDSchemas/Private/USDShadeMaterialTranslator.cpp
+++ b/Engine/Plugins/Importers/USDImporter/Source/USDSchemas/Private/USDShadeMaterialTranslator.cpp
@@ -6,9 +6,11 @@
 #include "USDAssetUserData.h"
 #include "USDClassesModule.h"
 #include "USDLog.h"
+#include "USDObjectUtils.h"
 #include "USDPrimConversion.h"
 #include "USDProjectSettings.h"
 #include "USDShadeConversion.h"
+#include "USDTranslatorUtils.h"
 #include "USDTypesConversion.h"
 
 #include "UsdWrappers/SdfPath.h"
@@ -38,6 +40,11 @@ namespace UE::UsdShadeTranslator::Private
 		TSet<UMaterialInterface*>& NewMaterials
 	)
 	{
+		if (!Context->UsdAssetCache || !Context->InfoCache)
+		{
+			return;
+		}
+
 		for (UTexture* Texture : TexturesToUpgrade)
 		{
 			if (Texture->VirtualTextureStreaming)
@@ -125,56 +132,58 @@ namespace UE::UsdShadeTranslator::Private
 						}
 						else
 #endif	  // WITH_EDITOR
-							if (UMaterialInstanceDynamic* MID = Cast<UMaterialInstanceDynamic>(UserMaterial))
+		  // Don't spell out MID directly, as at runtime we may be trying to upgrade a packaged MIC
+							if (UMaterialInstance* MI = Cast<UMaterialInstance>(UserMaterial))
 							{
-								if (Context->AssetCache && Context->InfoCache)
-								{
-									UE::FSdfPath PrimPath;
-									for (const UE::FSdfPath& Path : Context->InfoCache->GetPrimsForAsset(MID))
+								TArray<UE::FSdfPath> PrimsForAsset = Context->InfoCache->GetPrimsForAsset(MI);
+								const FString Hash = Context->UsdAssetCache->GetHashForAsset(MI);
+
+								// For MID we can't swap the reference material, so we need to remove the old material from the cache,
+								// create a brand new one, copy the overrides and then add that back in its place
+
+								TStrongObjectPtr<UMaterialInstance> MIDPin{MI};
+								FSoftObjectPath OldMIDPath = Context->UsdAssetCache->StopTrackingAsset(Hash);
+								ensure(OldMIDPath == FSoftObjectPath{MI});
+
+								bool bCreatedNew = false;
+								UMaterialInstanceDynamic* NewMID = Context->UsdAssetCache->GetOrCreateCustomCachedAsset<UMaterialInstanceDynamic>(
+									Hash,
+									MI->GetName(),
+									MI->GetFlags() | RF_Transient,	  // We never want MIDs to become assets in the content browser
+									[ReferenceMaterialVT](UPackage* Outer, FName SanitizedName, EObjectFlags FlagsToUse)
 									{
-										PrimPath = Path;
-										break;
-									}
-									const FString Hash = Context->AssetCache->GetHashForAsset(MID);
-
-									const FName NewInstanceName = MakeUniqueObjectName(
-										GetTransientPackage(),
-										UMaterialInstance::StaticClass(),
-										PrimPath.IsEmpty() ? TEXT("MaterialInstance")
-														   : *IUsdClassesModule::SanitizeObjectName(FPaths::GetBaseFilename(PrimPath.GetString()))
-									);
-
-									// For MID we can't swap the reference material, so we need to create a brand new one and copy
-									// the overrides
-									UMaterialInstanceDynamic* NewMID = UMaterialInstanceDynamic::Create(
-										ReferenceMaterialVT,
-										GetTransientPackage(),
-										NewInstanceName
-									);
-									if (!ensure(NewMID))
-									{
-										continue;
-									}
+										UMaterialInstanceDynamic* CreatedMID = UMaterialInstanceDynamic::Create(
+											ReferenceMaterialVT,
+											Outer,
+											SanitizedName
+										);
+										CreatedMID->ClearFlags(CreatedMID->GetFlags());
+										CreatedMID->SetFlags(FlagsToUse);
+										return CreatedMID;
+									},
+									&bCreatedNew
+								);
+								if (!ensure(bCreatedNew && NewMID))
+								{
+									continue;
+								}
 
-									NewMID->CopyParameterOverrides(MID);
+								NewMID->CopyParameterOverrides(MI);
 
-									UUsdMaterialAssetUserData* OldUserData = UserMaterial->GetAssetUserData<UUsdMaterialAssetUserData>();
-									if (OldUserData)
-									{
-										UUsdMaterialAssetUserData* NewUserData = DuplicateObject(OldUserData, NewMID);
-										NewMID->AddAssetUserData(NewUserData);
-									}
+								UUsdMaterialAssetUserData* OldUserData = UserMaterial->GetAssetUserData<UUsdMaterialAssetUserData>();
+								if (OldUserData)
+								{
+									UUsdMaterialAssetUserData* NewUserData = DuplicateObject(OldUserData, NewMID);
+									NewMID->AddAssetUserData(NewUserData);
+								}
 
-									if (Context->AssetCache->CanRemoveAsset(Hash) && Context->AssetCache->RemoveAsset(Hash))
-									{
-										Context->AssetCache->CacheAsset(Hash, NewMID);
-										if (!PrimPath.IsEmpty())
-										{
-											Context->InfoCache->LinkAssetToPrim(PrimPath, NewMID);
-										}
-										NewMaterials.Add(NewMID);
-									}
+								for (const UE::FSdfPath& PrimPath : PrimsForAsset)
+								{
+									Context->InfoCache->LinkAssetToPrim(PrimPath, NewMID);
 								}
+								NewMaterials.Add(NewMID);
+
+								MI->MarkAsGarbage();
 							}
 							else
 							{
@@ -338,196 +347,182 @@ void FUsdShadeMaterialTranslator::CreateAssets()
 	}
 	const FString PrefixedMaterialHash = UsdUtils::GetAssetHashPrefix(GetPrim(), Context->bReuseIdenticalAssets) + MaterialHash;
 
-	UMaterialInterface* ConvertedMaterial = nullptr;
+	const FString PrimPathString = PrimPath.GetString();
+	const FString DesiredName = FPaths::GetBaseFilename(PrimPathString);
 
-	if (Context->AssetCache)
-	{
-		ConvertedMaterial = Cast<UMaterialInterface>(Context->AssetCache->GetCachedAsset(PrefixedMaterialHash));
-	}
+	const bool bIsMaterialTranslucent = UsdUtils::IsMaterialTranslucent(ShadeMaterial);
 
-	if (!ConvertedMaterial)
+	UMaterialInterface* ConvertedMaterial = nullptr;
+
+	bool bCreatedNew = false;
+#if WITH_EDITOR
+	if (GIsEditor)
 	{
-		const FString PrimPathString = PrimPath.GetString();
-		const bool bIsTranslucent = UsdUtils::IsMaterialTranslucent(ShadeMaterial);
-		const FName InstanceName = MakeUniqueObjectName(
-			GetTransientPackage(),
-			UMaterialInstance::StaticClass(),
-			*IUsdClassesModule::SanitizeObjectName(FPaths::GetBaseFilename(PrimPathString))
+		UMaterialInstanceConstant* MIC = Context->UsdAssetCache->GetOrCreateCachedAsset<UMaterialInstanceConstant>(
+			PrefixedMaterialHash,	 //
+			DesiredName,
+			Context->ObjectFlags,
+			&bCreatedNew
 		);
+		ConvertedMaterial = MIC;
 
-#if WITH_EDITOR
-		if (GIsEditor)	  // Also have to prevent Standalone game from going with MaterialInstanceConstants
+		if (bCreatedNew)
 		{
-			if (UMaterialInstanceConstant* NewMaterial = NewObject<UMaterialInstanceConstant>(
-					GetTransientPackage(),
-					InstanceName,
-					Context->ObjectFlags | EObjectFlags::RF_Transient
-				))
+			const bool bSuccess = UsdToUnreal::ConvertMaterial(
+				ShadeMaterial,
+				*MIC,
+				Context->UsdAssetCache.Get(),
+				*Context->RenderContext.ToString(),
+				Context->bReuseIdenticalAssets
+			);
+			if (!bSuccess)
 			{
-				UUsdMaterialAssetUserData* UserData = NewObject<UUsdMaterialAssetUserData>(NewMaterial, TEXT("USDAssetUserData"));
-				UserData->PrimPaths = {PrimPath.GetString()};
-				NewMaterial->AddAssetUserData(UserData);
-
-				const bool bSuccess = UsdToUnreal::ConvertMaterial(
-					ShadeMaterial,
-					*NewMaterial,
-					Context->AssetCache.Get(),
-					*Context->RenderContext.ToString(),
-					Context->bReuseIdenticalAssets
-				);
-				if (!bSuccess)
-				{
-					NewMaterial->MarkAsGarbage();
-					return;
-				}
+				UsdUnreal::TranslatorUtils::AbandonFailedAsset(MIC, Context->UsdAssetCache.Get(), Context->InfoCache.Get());
+				return;
+			}
 
-				TSet<UTexture*> VTTextures;
-				TSet<UTexture*> NonVTTextures;
-				for (const FTextureParameterValue& TextureValue : NewMaterial->TextureParameterValues)
+			TSet<UTexture*> VTTextures;
+			TSet<UTexture*> NonVTTextures;
+			for (const FTextureParameterValue& TextureValue : MIC->TextureParameterValues)
+			{
+				if (UTexture* Texture = TextureValue.ParameterValue)
 				{
-					if (UTexture* Texture = TextureValue.ParameterValue)
+					if (Texture->VirtualTextureStreaming)
 					{
-						if (Texture->VirtualTextureStreaming)
-						{
-							UsdUtils::NotifyIfVirtualTexturesNeeded(Texture);
-							VTTextures.Add(Texture);
-						}
-						else
-						{
-							NonVTTextures.Add(Texture);
-						}
+						UsdUtils::NotifyIfVirtualTexturesNeeded(Texture);
+						VTTextures.Add(Texture);
+					}
+					else
+					{
+						NonVTTextures.Add(Texture);
 					}
 				}
+			}
+
+			// Our VT material only has VT texture samplers, so *all* of its textures must be VT
+			if (VTTextures.Num() && NonVTTextures.Num())
+			{
+				UE_LOG(LogUsd, Log, TEXT("Upgrading textures used by material instance '%s' to VT as the material must be VT"), *MIC->GetName());
+				UE::UsdShadeTranslator::Private::UpgradeMaterialsAndTexturesToVT(NonVTTextures, Context);
+			}
+
+			MeshTranslationImpl::EUsdReferenceMaterialProperties Properties = MeshTranslationImpl::EUsdReferenceMaterialProperties::None;
+			if (bIsMaterialTranslucent)
+			{
+				Properties |= MeshTranslationImpl::EUsdReferenceMaterialProperties::Translucent;
+			}
+			if (VTTextures.Num() > 0)
+			{
+				Properties |= MeshTranslationImpl::EUsdReferenceMaterialProperties::VT;
+			}
+			UMaterialInterface* ReferenceMaterial = MeshTranslationImpl::GetReferencePreviewSurfaceMaterial(Properties);
+
+			if (ensure(ReferenceMaterial))
+			{
+				MIC->SetParentEditorOnly(ReferenceMaterial);
 
-				// Our VT material only has VT texture samplers, so *all* of its textures must be VT
-				if (VTTextures.Num() && NonVTTextures.Num())
+				// We can't blindly recreate all component render states when a level is being added, because we may end up first creating
+				// render states for some components, and UWorld::AddToWorld calls FScene::AddPrimitive which expects the component to not have
+				// primitives yet
+				FMaterialUpdateContext::EOptions::Type Options = FMaterialUpdateContext::EOptions::Default;
+				if (Context->Level && Context->Level->bIsAssociatingLevel)
 				{
-					UE_LOG(
-						LogUsd,
-						Log,
-						TEXT("Upgrading textures used by material instance '%s' to VT as the material must be VT"),
-						*NewMaterial->GetName()
-					);
-					UE::UsdShadeTranslator::Private::UpgradeMaterialsAndTexturesToVT(NonVTTextures, Context);
+					Options = (FMaterialUpdateContext::EOptions::Type)(Options & ~FMaterialUpdateContext::EOptions::RecreateRenderStates);
 				}
 
-				MeshTranslationImpl::EUsdReferenceMaterialProperties Properties = MeshTranslationImpl::EUsdReferenceMaterialProperties::None;
-				if (bIsTranslucent)
+				FMaterialUpdateContext UpdateContext(Options, GMaxRHIShaderPlatform);
+				UpdateContext.AddMaterialInstance(MIC);
+				MIC->PreEditChange(nullptr);
+				MIC->PostEditChange();
+
+				for (UTexture* Texture : VTTextures.Union(NonVTTextures))
 				{
-					Properties |= MeshTranslationImpl::EUsdReferenceMaterialProperties::Translucent;
+					Context->TextureToUserMaterials.FindOrAdd(Texture).Add(MIC);
 				}
-				if (VTTextures.Num() > 0)
+			}
+		}
+	}
+	else
+#endif	  // WITH_EDITOR
+	{
+		// Note how we're requesting just a UMaterialInstance here, instead of spelling out the MID. This because
+		// if we're a runtime we may have a cooked MIC assigned to this hash, and in that case we want to use it
+		// instead of overwriting it with a MID. Our creation func will ensure we create a MID as a fallback anyway
+		UMaterialInstance* MI = Context->UsdAssetCache->GetOrCreateCustomCachedAsset<UMaterialInstance>(
+			PrefixedMaterialHash,
+			DesiredName,
+			Context->ObjectFlags | RF_Transient,	// We never want MIDs to become assets in the content browser
+			[bIsMaterialTranslucent](UPackage* Outer, FName SanitizedName, EObjectFlags FlagsToUse)
+			{
+				MeshTranslationImpl::EUsdReferenceMaterialProperties Properties = MeshTranslationImpl::EUsdReferenceMaterialProperties::None;
+				if (bIsMaterialTranslucent)
 				{
-					Properties |= MeshTranslationImpl::EUsdReferenceMaterialProperties::VT;
+					Properties |= MeshTranslationImpl::EUsdReferenceMaterialProperties::Translucent;
 				}
 				UMaterialInterface* ReferenceMaterial = MeshTranslationImpl::GetReferencePreviewSurfaceMaterial(Properties);
 
-				if (ensure(ReferenceMaterial))
-				{
-					NewMaterial->SetParentEditorOnly(ReferenceMaterial);
+				UMaterialInstanceDynamic* CreatedMID = UMaterialInstanceDynamic::Create(ReferenceMaterial, Outer, SanitizedName);
+				CreatedMID->ClearFlags(CreatedMID->GetFlags());
+				CreatedMID->SetFlags(FlagsToUse);
+				return CreatedMID;
+			},
+			&bCreatedNew
+		);
+		ConvertedMaterial = MI;
 
-					// We can't blindly recreate all component render states when a level is being added, because we may end up first creating
-					// render states for some components, and UWorld::AddToWorld calls FScene::AddPrimitive which expects the component to not have
-					// primitives yet
-					FMaterialUpdateContext::EOptions::Type Options = FMaterialUpdateContext::EOptions::Default;
-					if (Context->Level && Context->Level->bIsAssociatingLevel)
+		if (bCreatedNew)
+		{
+			const bool bSuccess = UsdToUnreal::ConvertMaterial(
+				ShadeMaterial,
+				*MI,
+				Context->UsdAssetCache.Get(),
+				*Context->RenderContext.ToString(),
+				Context->bReuseIdenticalAssets
+			);
+			if (!bSuccess)
+			{
+				UsdUnreal::TranslatorUtils::AbandonFailedAsset(MI, Context->UsdAssetCache.Get(), Context->InfoCache.Get());
+				return;
+			}
+
+			TSet<UTexture*> VTTextures;
+			TSet<UTexture*> NonVTTextures;
+			for (const FTextureParameterValue& TextureValue : MI->TextureParameterValues)
+			{
+				if (UTexture* Texture = TextureValue.ParameterValue)
+				{
+					if (Texture->VirtualTextureStreaming)
 					{
-						Options = (FMaterialUpdateContext::EOptions::Type)(Options & ~FMaterialUpdateContext::EOptions::RecreateRenderStates);
+						VTTextures.Add(Texture);
 					}
-
-					FMaterialUpdateContext UpdateContext(Options, GMaxRHIShaderPlatform);
-					UpdateContext.AddMaterialInstance(NewMaterial);
-					NewMaterial->PreEditChange(nullptr);
-					NewMaterial->PostEditChange();
-
-					ConvertedMaterial = NewMaterial;
-
-					for (UTexture* Texture : VTTextures.Union(NonVTTextures))
+					else
 					{
-						Context->TextureToUserMaterials.FindOrAdd(Texture).Add(NewMaterial);
+						NonVTTextures.Add(Texture);
 					}
 				}
 			}
-		}
-		else
-#endif	  // WITH_EDITOR
-		{
-			// At runtime we always start with a non-VT reference and if we discover we need one we just create a new MID
-			// using the VT reference and copy the overrides. Not much else we can do as we need a reference to call
-			// UMaterialInstanceDynamic::Create and get our instance, but we an instance to call UsdToUnreal::ConvertMaterial
-			// to create our textures and decide on our reference.
-			MeshTranslationImpl::EUsdReferenceMaterialProperties Properties = MeshTranslationImpl::EUsdReferenceMaterialProperties::None;
-			if (bIsTranslucent)
+
+			// We must stash our material and textures *before* we call UpgradeMaterialsAndTexturesToVT, as that
+			// is what will actually swap our reference with a VT one if needed
+			if (Context->InfoCache)
 			{
-				Properties |= MeshTranslationImpl::EUsdReferenceMaterialProperties::Translucent;
+				Context->InfoCache->LinkAssetToPrim(PrimPath, MI);
 			}
-			UMaterialInterface* ReferenceMaterial = MeshTranslationImpl::GetReferencePreviewSurfaceMaterial(Properties);
-
-			if (ensure(ReferenceMaterial))
+			for (UTexture* Texture : VTTextures.Union(NonVTTextures))
 			{
-				if (UMaterialInstanceDynamic* NewMaterial = UMaterialInstanceDynamic::Create(ReferenceMaterial, GetTransientPackage(), InstanceName))
-				{
-					UUsdMaterialAssetUserData* UserData = NewObject<UUsdMaterialAssetUserData>(NewMaterial, TEXT("USDAssetUserData"));
-					UserData->PrimPaths = {PrimPath.GetString()};
-					NewMaterial->AddAssetUserData(UserData);
-
-					NewMaterial->SetFlags(RF_Transient);
-
-					if (UsdToUnreal::ConvertMaterial(
-							ShadeMaterial,
-							*NewMaterial,
-							Context->AssetCache.Get(),
-							*Context->RenderContext.ToString(),
-							Context->bReuseIdenticalAssets
-						))
-					{
-						TSet<UTexture*> VTTextures;
-						TSet<UTexture*> NonVTTextures;
-						for (const FTextureParameterValue& TextureValue : NewMaterial->TextureParameterValues)
-						{
-							if (UTexture* Texture = TextureValue.ParameterValue)
-							{
-								if (Texture->VirtualTextureStreaming)
-								{
-									VTTextures.Add(Texture);
-								}
-								else
-								{
-									NonVTTextures.Add(Texture);
-								}
-							}
-						}
-
-						// We must stash our material and textures *before* we call UpgradeMaterialsAndTexturesToVT, as that
-						// is what will actually swap our reference with a VT one if needed
-						if (Context->AssetCache && Context->InfoCache)
-						{
-							Context->AssetCache->CacheAsset(PrefixedMaterialHash, NewMaterial);
-							Context->InfoCache->LinkAssetToPrim(PrimPath, NewMaterial);
-						}
-						for (UTexture* Texture : VTTextures.Union(NonVTTextures))
-						{
-							Context->TextureToUserMaterials.FindOrAdd(Texture).Add(NewMaterial);
-						}
-
-						// Our VT material only has VT texture samplers, so *all* of its textures must be VT
-						if (VTTextures.Num() && NonVTTextures.Num())
-						{
-							UE_LOG(
-								LogUsd,
-								Log,
-								TEXT("Upgrading textures used by material instance '%s' to VT as the material must be VT"),
-								*NewMaterial->GetName()
-							);
-							UE::UsdShadeTranslator::Private::UpgradeMaterialsAndTexturesToVT(NonVTTextures, Context);
-						}
+				Context->TextureToUserMaterials.FindOrAdd(Texture).Add(MI);
+			}
 
-						// We must go through the cache to fetch our result material here as UpgradeMaterialsAndTexturesToVT
-						// may have created a new MID for this material with a VT reference
-						ConvertedMaterial = Cast<UMaterialInterface>(Context->AssetCache->GetCachedAsset(PrefixedMaterialHash));
-					}
-				}
+			// Our VT material only has VT texture samplers, so *all* of its textures must be VT
+			if (VTTextures.Num() && NonVTTextures.Num())
+			{
+				UE_LOG(LogUsd, Log, TEXT("Upgrading textures used by material instance '%s' to VT as the material must be VT"), *MI->GetName());
+				UE::UsdShadeTranslator::Private::UpgradeMaterialsAndTexturesToVT(NonVTTextures, Context);
 			}
+
+			// We must go through the cache to fetch our result material here as UpgradeMaterialsAndTexturesToVT
+			// may have created a new MID for this material with a VT reference
+			ConvertedMaterial = Context->UsdAssetCache->GetCachedAsset<UMaterialInterface>(PrefixedMaterialHash);
 		}
 	}
 
@@ -546,12 +541,12 @@ bool FUsdShadeMaterialTranslator::CanBeCollapsed(ECollapsingType CollapsingType)
 
 void FUsdShadeMaterialTranslator::PostImportMaterial(const FString& PrefixedMaterialHash, UMaterialInterface* ImportedMaterial)
 {
-	if (!ImportedMaterial || !Context->InfoCache || !Context->AssetCache)
+	if (!ImportedMaterial || !Context->InfoCache || !Context->UsdAssetCache)
 	{
 		return;
 	}
 
-	if (UUsdMaterialAssetUserData* UserData = UsdUtils::GetOrCreateAssetUserData<UUsdMaterialAssetUserData>(ImportedMaterial))
+	if (UUsdMaterialAssetUserData* UserData = UsdUnreal::ObjectUtils::GetOrCreateAssetUserData<UUsdMaterialAssetUserData>(ImportedMaterial))
 	{
 		UserData->PrimPaths.AddUnique(PrimPath.GetString());
 
@@ -574,10 +569,10 @@ void FUsdShadeMaterialTranslator::PostImportMaterial(const FString& PrefixedMate
 
 	// Note that this needs to run even if we found this material in the asset cache already, otherwise we won't
 	// re-register the prim asset links when we reload a stage
-	Context->AssetCache->CacheAsset(PrefixedMaterialHash, ImportedMaterial);
 	Context->InfoCache->LinkAssetToPrim(PrimPath, ImportedMaterial);
 
-	// Also link the textures to the same material prim.
+	// Also link the textures to the same material prim. Our textures should all come from USDShadeConversion.cpp or
+	// MaterialX or MDL translators, so they should already be tracked by the same asset cache the material is tracked by.
 	// This is important because it lets the stage actor drop its references to old unused textures in the
 	// asset cache if they aren't being used by any other material
 	TSet<UObject*> Dependencies = IUsdClassesModule::GetAssetDependencies(ImportedMaterial);
@@ -585,82 +580,20 @@ void FUsdShadeMaterialTranslator::PostImportMaterial(const FString& PrefixedMate
 	{
 		if (UTexture* Texture = Cast<UTexture>(Object))
 		{
-			// We don't use "GetOutermost()" here because it's also possible to be owned by an asset cache that
-			// itself lives in the transient package... bIsOwnedByTransientPackage should be true just for new textures
-			// dumped on the transient package
-			const bool bIsOwnedByTransientPackage = Texture->GetOuter()->GetPackage() == GetTransientPackage();
-			const bool bIsOwnedByCache = Context->AssetCache->IsAssetOwnedByCache(Texture->GetPathName());
-
-			// Texture is already owned by the cache: Just touch it without recomputing its hash as that's expensive
-			if (bIsOwnedByCache)
-			{
-				Context->AssetCache->TouchAsset(Texture);
-			}
-			// Texture is owned by the transient package, but not cached yet: Let's take it
-			else if (bIsOwnedByTransientPackage)
+			// We may be reusing a material from the asset cache that has textures fully unrelated to USD, which
+			// we shouldn't interact with
+			const bool bIsTrackedByCache = Context->UsdAssetCache->IsAssetTrackedByCache(Texture->GetPathName());
+			if (bIsTrackedByCache)
 			{
-				FString FilePath;
-#if WITH_EDITOR
-				if (UAssetImportData* TextureImportData = Texture->AssetImportData.Get())
-				{
-					FilePath = TextureImportData->GetFirstFilename();
-				}
-				else
-#endif	  // WITH_EDITOR
-				{
-					FilePath = Texture->GetName();
-				}
-
-				const FString HashPrefix = UsdUtils::GetAssetHashPrefix(GetPrim(), Context->bReuseIdenticalAssets);
-				const FString PrefixedTextureHash = HashPrefix
-													+ UsdUtils::GetTextureHash(
-														FilePath,
-														Texture->SRGB,
-														Texture->CompressionSettings,
-														Texture->GetTextureAddressX(),
-														Texture->GetTextureAddressY()
-													);
-
-				// Some translators like FMaterialXUsdShadeMaterialTranslator will import many materials and textures
-				// at once and preload them all in the cache. In some complex scenarios when some materials are updated
-				// it is possible to arrive at a situation where the we need to reparse the source file again and
-				// regenerate a bunch of materials and textures without access to the asset cache. In those cases we'll
-				// unfortunately recreate identical textures, and if we try to store them in here we'll run into a hash
-				// collision. These should be rare, however, and require repeatedly e.g. updating the Material prims
-				// generated by MaterialX or updating the MaterialX import options
-				UTexture* ExistingTexture = Cast<UTexture>(Context->AssetCache->GetCachedAsset(PrefixedTextureHash));
-				if (!ExistingTexture)
-				{
-					Texture->SetFlags(RF_Transient);
-					Context->AssetCache->CacheAsset(PrefixedTextureHash, Texture);
-				}
-			}
+				Context->UsdAssetCache->TouchAssetPath(Texture);
+				Context->InfoCache->LinkAssetToPrim(PrimPath, Texture);
 
-			if (bIsOwnedByCache || bIsOwnedByTransientPackage)
-			{
-				UUsdAssetUserData* TextureUserData = Texture->GetAssetUserData<UUsdAssetUserData>();
-				if (!TextureUserData)
+				if (UUsdAssetUserData* TextureUserData = UsdUnreal::ObjectUtils::GetOrCreateAssetUserData(Texture))
 				{
-					TextureUserData = NewObject<UUsdAssetUserData>(Texture, TEXT("USDAssetUserData"));
-					Texture->AddAssetUserData(TextureUserData);
+					TextureUserData->PrimPaths.AddUnique(PrimPath.GetString());
 				}
-				TextureUserData->PrimPaths.AddUnique(PrimPath.GetString());
-
-				Context->InfoCache->LinkAssetToPrim(PrimPath, Texture);
 			}
 		}
-		else if (UMaterialInterface* ReferenceMaterial = Cast<UMaterialInterface>(Object))
-		{
-			// Some scenarios can generate reference/instance material pairs, and reference materials are dependencies.
-			// We won't handle these dependencies recursively though, the caller is responsible for calling this for
-			// all individual materials as they need to also provide the hash to use for each
-			UMaterialInstance* Instance = Cast<UMaterialInstance>(ImportedMaterial);
-			ensure(Instance && Instance->Parent.Get() == ReferenceMaterial);
-		}
-		else
-		{
-			ensureMsgf(false, TEXT("Asset type unsupported!"));
-		}
 	}
 }
 
diff --git a/Engine/Plugins/Importers/USDImporter/Source/USDSchemas/Private/USDSkelRootTranslator.cpp b/Engine/Plugins/Importers/USDImporter/Source/USDSchemas/Private/USDSkelRootTranslator.cpp
--- a/Engine/Plugins/Importers/USDImporter/Source/USDSchemas/Private/USDSkelRootTranslator.cpp
+++ b/Engine/Plugins/Importers/USDImporter/Source/USDSchemas/Private/USDSkelRootTranslator.cpp
@@ -5,8 +5,8 @@
 #if USE_USD_SDK
 
 #include "MeshTranslationImpl.h"
+#include "USDAssetCache3.h"
 #include "USDAssetUserData.h"
-#include "USDClassesModule.h"
 #include "USDConversionUtils.h"
 #include "USDDrawModeComponent.h"
 #include "USDErrorUtils.h"
@@ -17,6 +17,7 @@
 #include "USDLayerUtils.h"
 #include "USDLog.h"
 #include "USDMemory.h"
+#include "USDObjectUtils.h"
 #include "USDSkeletalDataConversion.h"
 #include "USDTypesConversion.h"
 
@@ -69,7 +70,7 @@ namespace UsdSkelRootTranslatorImpl
 		const pxr::UsdPrim& UsdPrim,
 		TArray<UsdUtils::FUsdPrimMaterialAssignmentInfo>& LODIndexToMaterialInfo,
 		USkeletalMesh* SkeletalMesh,
-		UUsdAssetCache2& AssetCache,
+		UUsdAssetCache3& AssetCache,
 		FUsdInfoCache& InfoCache,
 		float Time,
 		EObjectFlags Flags,
@@ -428,7 +429,7 @@ namespace UsdSkelRootTranslatorImpl
 		FName& OutSkeletonName,
 		UsdUtils::FBlendShapeMap* OutBlendShapes,
 		TSet<FString>& InOutUsedMorphTargetNames,
-		const UUsdAssetCache2& AssetCache,
+		const UUsdAssetCache3& AssetCache,
 		bool bInInterpretLODs,
 		const UsdToUnreal::FUsdMeshConversionOptions& Options
 	)
@@ -657,7 +658,7 @@ namespace UsdSkelRootTranslatorImpl
 		const pxr::UsdPrim& SkelRootPrim,
 		const TArray<UMaterialInterface*>& ExistingAssignments,
 		UMeshComponent& MeshComponent,
-		UUsdAssetCache2& AssetCache,
+		UUsdAssetCache3& AssetCache,
 		FUsdInfoCache& InfoCache,
 		float Time,
 		EObjectFlags Flags,
@@ -884,14 +885,10 @@ namespace UsdSkelRootTranslatorImpl
 		FName AssetName = MakeUniqueObjectName(
 			GetTransientPackage(),
 			UPhysicsAsset::StaticClass(),
-			*IUsdClassesModule::SanitizeObjectName(FPaths::GetBaseFilename(TEXT("PHYS_") + SkeletalMesh->GetName()))
+			*UsdUnreal::ObjectUtils::SanitizeObjectName(FPaths::GetBaseFilename(TEXT("PHYS_") + SkeletalMesh->GetName()))
 		);
 
-		UPhysicsAsset* Result = NewObject<UPhysicsAsset>(
-			GetTransientPackage(),
-			AssetName,
-			Flags | EObjectFlags::RF_Public | EObjectFlags::RF_Transient
-		);
+		UPhysicsAsset* Result = NewObject<UPhysicsAsset>(GetTransientPackage(), AssetName, Flags | EObjectFlags::RF_Public);
 
 		FPhysAssetCreateParams NewBodyData;
 		FText CreationErrorMessage;
@@ -924,7 +921,7 @@ namespace UsdSkelRootTranslatorImpl
 		bool* bOutNeedsRecompile = nullptr	  // Whether this function wants the returned AnimBP to be recompiled
 	)
 	{
-		if (!Prim || !Context.InfoCache || Context.AssetCache)
+		if (!Prim || !Context.InfoCache || Context.UsdAssetCache)
 		{
 			return nullptr;
 		}
@@ -990,7 +987,7 @@ namespace UsdSkelRootTranslatorImpl
 
 			// Check if we can find an AnimBP for this prim in the asset cache (useful when doing Action->Import)
 			bool bReusedAnimBP = false;
-			if (UAnimBlueprint* CachedAnimBP = Cast<UAnimBlueprint>(Context.AssetCache->GetCachedAsset(PrefixedAnimBPHash)))
+			if (UAnimBlueprint* CachedAnimBP = Cast<UAnimBlueprint>(Context.UsdAssetCache->GetCachedAsset(PrefixedAnimBPHash)))
 			{
 				if (CachedAnimBP->TargetSkeleton == Skeleton)
 				{
@@ -1005,7 +1002,7 @@ namespace UsdSkelRootTranslatorImpl
 				FName UniqueName = MakeUniqueObjectName(
 					GetTransientPackage(),
 					UAnimBlueprint::StaticClass(),
-					*IUsdClassesModule::SanitizeObjectName(PrimName + TEXT("_DefaultAnimBlueprint"))
+					*UsdUnreal::ObjectUtils::SanitizeObjectName(PrimName + TEXT("_DefaultAnimBlueprint"))
 				);
 
 				// Duplicate and never reuse these so that they can be assigned independent subject names if desired.
@@ -1023,7 +1020,7 @@ namespace UsdSkelRootTranslatorImpl
 
 				bNeedRecompile = true;
 
-				Context.AssetCache->CacheAsset(PrefixedAnimBPHash, AnimBP);
+				Context.UsdAssetCache->CacheAsset(PrefixedAnimBPHash, AnimBP);
 			}
 		}
 		// Path is pointing to an existing, persistent AnimBP
@@ -1339,7 +1336,7 @@ namespace UsdSkelRootTranslatorImpl
 				   SkeletonName,
 				   OutBlendShapes,
 				   UsedMorphTargetNames,
-				   *Context->AssetCache.Get(),
+				   *Context->UsdAssetCache.Get(),
 				   Context->bAllowInterpretingLODs,
 				   Options
 			   );
@@ -1352,7 +1349,7 @@ namespace UsdSkelRootTranslatorImpl
 			ESchemaTranslationLaunchPolicy::Sync,
 			[this]()
 			{
-				if (!Context->InfoCache || !Context->AssetCache)
+				if (!Context->InfoCache || !Context->UsdAssetCache)
 				{
 					return false;
 				}
@@ -1363,11 +1360,12 @@ namespace UsdSkelRootTranslatorImpl
 				FSHAHash SkeletalMeshHash = UsdSkelRootTranslatorImpl::ComputeSHAHash(LODIndexToSkeletalMeshImportData, SkeletonBones, BlendShapes);
 				PrefixedSkelMeshHash = UsdUtils::GetAssetHashPrefix(GetPrim(), Context->bReuseIdenticalAssets) + SkeletalMeshHash.ToString();
 
-				USkeletalMesh* SkeletalMesh = Cast<USkeletalMesh>(Context->AssetCache->GetCachedAsset(PrefixedSkelMeshHash));
+				USkeletalMesh* SkeletalMesh = Cast<USkeletalMesh>(Context->UsdAssetCache->GetCachedAsset(PrefixedSkelMeshHash));
 
 				bool bIsNew = false;
 				if (!SkeletalMesh)
 				{
+					PRAGMA_DISABLE_DEPRECATION_WARNINGS
 					bIsNew = true;
 					SkeletalMesh = UsdToUnreal::GetSkeletalMeshFromImportData(
 						LODIndexToSkeletalMeshImportData,
@@ -1377,6 +1375,7 @@ namespace UsdSkelRootTranslatorImpl
 						*FPaths::GetBaseFilename(SkelRootPath),
 						SkeletonName
 					);
+					PRAGMA_ENABLE_DEPRECATION_WARNINGS
 				}
 
 				if (SkeletalMesh)
@@ -1398,7 +1397,7 @@ namespace UsdSkelRootTranslatorImpl
 							GetPrim(),
 							LODIndexToMaterialInfo,
 							SkeletalMesh,
-							*Context->AssetCache.Get(),
+							*Context->UsdAssetCache.Get(),
 							*Context->InfoCache.Get(),
 							Context->Time,
 							Context->ObjectFlags,
@@ -1412,8 +1411,8 @@ namespace UsdSkelRootTranslatorImpl
 							SkeletalMesh->UpdateUVChannelData(bRebuildAll);
 						}
 
-						Context->AssetCache->CacheAsset(PrefixedSkelMeshHash, SkeletalMesh);
-						Context->AssetCache->CacheAsset(PrefixedSkelMeshHash + TEXT("_Skeleton"), SkeletalMesh->GetSkeleton());
+						Context->UsdAssetCache->CacheAsset(PrefixedSkelMeshHash, SkeletalMesh);
+						Context->UsdAssetCache->CacheAsset(PrefixedSkelMeshHash + TEXT("_Skeleton"), SkeletalMesh->GetSkeleton());
 
 						// The PreviewSkeletalMesh property on the Skeleton is a soft object path, that is set within
 						// UsdToUnreal::GetSkeletalMeshFromImportData before the SkeletalMesh is part of the cache. When we cache the
@@ -1436,7 +1435,7 @@ namespace UsdSkelRootTranslatorImpl
 
 							if (PhysicsAsset)
 							{
-								Context->AssetCache->CacheAsset(PrefixedSkelMeshHash + TEXT("_PhysicsAsset"), PhysicsAsset);
+								Context->UsdAssetCache->CacheAsset(PrefixedSkelMeshHash + TEXT("_PhysicsAsset"), PhysicsAsset);
 							}
 						}
 
@@ -1588,7 +1587,7 @@ namespace UsdSkelRootTranslatorImpl
 						FSHAHash Hash = UsdSkelRootTranslatorImpl::ComputeSHAHash(SkelQuery, RootMotionPrim, PrefixedSkelMeshHash);
 						FString PrefixedSkelAnimHash = UsdUtils::GetAssetHashPrefix(SkelAnimationPrim, Context->bReuseIdenticalAssets)
 													   + Hash.ToString();
-						UAnimSequence* AnimSequence = Cast<UAnimSequence>(Context->AssetCache->GetCachedAsset(PrefixedSkelAnimHash));
+						UAnimSequence* AnimSequence = Cast<UAnimSequence>(Context->UsdAssetCache->GetCachedAsset(PrefixedSkelAnimHash));
 
 						if (!AnimSequence || AnimSequence->GetSkeleton() != SkeletalMesh->GetSkeleton())
 						{
@@ -1597,7 +1596,7 @@ namespace UsdSkelRootTranslatorImpl
 							FName AnimSequenceName = MakeUniqueObjectName(
 								GetTransientPackage(),
 								UAnimSequence::StaticClass(),
-								*IUsdClassesModule::SanitizeObjectName(UsdToUnreal::ConvertToken(SkelAnimationPrim.GetName()))
+								*UsdUnreal::ObjectUtils::SanitizeObjectName(UsdToUnreal::ConvertToken(SkelAnimationPrim.GetName()))
 							);
 
 							// The UAnimSequence can't be created with the RF_Transactional flag, or else it will be serialized without
@@ -1607,7 +1606,7 @@ namespace UsdSkelRootTranslatorImpl
 							AnimSequence = NewObject<UAnimSequence>(
 								GetTransientPackage(),
 								AnimSequenceName,
-								(Context->ObjectFlags & ~EObjectFlags::RF_Transactional) | EObjectFlags::RF_Transient
+								(Context->ObjectFlags & ~EObjectFlags::RF_Transactional)
 							);
 							AnimSequence->SetSkeleton(SkeletalMesh->GetSkeleton());
 
@@ -1638,7 +1637,7 @@ namespace UsdSkelRootTranslatorImpl
 
 								AnimSequence->AddAssetUserData(UserData);
 
-								Context->AssetCache->CacheAsset(PrefixedSkelAnimHash, AnimSequence);
+								Context->UsdAssetCache->CacheAsset(PrefixedSkelAnimHash, AnimSequence);
 							}
 							else
 							{
@@ -1699,11 +1698,11 @@ USceneComponent* FUsdSkelRootTranslator::CreateComponents()
 #if WITH_EDITOR
 	// Check if the prim has the GroomBinding schema and setup the component and assets necessary to bind the groom to the SkeletalMesh
 	if (SceneComponent && SceneComponent->IsA<USkeletalMeshComponent>() && UsdUtils::PrimHasSchema(GetPrim(), UnrealIdentifiers::GroomBindingAPI)
-		&& Context->AssetCache && Context->InfoCache && Context->bAllowParsingGroomAssets)
+		&& Context->UsdAssetCache && Context->InfoCache && Context->bAllowParsingGroomAssets)
 	{
 		UsdGroomTranslatorUtils::CreateGroomBindingAsset(
 			GetPrim(),
-			*Context->AssetCache,
+			*Context->UsdAssetCache,
 			*Context->InfoCache,
 			Context->ObjectFlags,
 			Context->bReuseIdenticalAssets
@@ -1800,7 +1799,7 @@ void FUsdSkelRootTranslator::UpdateComponents(USceneComponent* SceneComponent)
 				Prim,
 				ExistingAssignments,
 				*SkeletalMeshComponent,
-				*Context->AssetCache.Get(),
+				*Context->UsdAssetCache.Get(),
 				*Context->InfoCache.Get(),
 				Context->Time,
 				Context->ObjectFlags,
