USD: Implements UsdAssetCache3 which should be a simpler and more flexible way of managing USD assets;

diff --git a/Engine/Plugins/Importers/USDImporter/Source/USDSchemas/Private/USDSkelSkeletonTranslator.cpp b/Engine/Plugins/Importers/USDImporter/Source/USDSchemas/Private/USDSkelSkeletonTranslator.cpp
--- a/Engine/Plugins/Importers/USDImporter/Source/USDSchemas/Private/USDSkelSkeletonTranslator.cpp
+++ b/Engine/Plugins/Importers/USDImporter/Source/USDSchemas/Private/USDSkelSkeletonTranslator.cpp
@@ -5,6 +5,7 @@
 #if USE_USD_SDK
 
 #include "MeshTranslationImpl.h"
+#include "USDAssetCache3.h"
 #include "USDAssetUserData.h"
 #include "USDClassesModule.h"
 #include "USDConversionUtils.h"
@@ -16,8 +17,10 @@
 #include "USDLayerUtils.h"
 #include "USDLog.h"
 #include "USDMemory.h"
+#include "USDObjectUtils.h"
 #include "USDPrimConversion.h"
 #include "USDSkeletalDataConversion.h"
+#include "USDTranslatorUtils.h"
 #include "USDTypesConversion.h"
 
 #include "UsdWrappers/SdfLayer.h"
@@ -77,7 +80,7 @@ namespace UsdSkelSkeletonTranslatorImpl
 		const pxr::UsdPrim& UsdPrim,
 		TArray<UsdUtils::FUsdPrimMaterialAssignmentInfo>& LODIndexToMaterialInfo,
 		USkeletalMesh* SkeletalMesh,
-		UUsdAssetCache2& AssetCache,
+		UUsdAssetCache3& AssetCache,
 		FUsdInfoCache& InfoCache,
 		float Time,
 		EObjectFlags Flags,
@@ -743,7 +746,7 @@ namespace UsdSkelSkeletonTranslatorImpl
 		const pxr::UsdSkelBinding& SkeletonBinding,
 		const TArray<UMaterialInterface*>& ExistingAssignments,
 		UMeshComponent& MeshComponent,
-		UUsdAssetCache2& AssetCache,
+		UUsdAssetCache3& AssetCache,
 		FUsdInfoCache& InfoCache,
 		float Time,
 		EObjectFlags Flags,
@@ -964,45 +967,6 @@ namespace UsdSkelSkeletonTranslatorImpl
 		return false;
 	}
 
-	UPhysicsAsset* GenerateAndAssignPhysicsAsset(USkeletalMesh* SkeletalMesh, EObjectFlags Flags)
-	{
-#if WITH_EDITOR
-		FName AssetName = MakeUniqueObjectName(
-			GetTransientPackage(),
-			UPhysicsAsset::StaticClass(),
-			*IUsdClassesModule::SanitizeObjectName(FPaths::GetBaseFilename(TEXT("PHYS_") + SkeletalMesh->GetName()))
-		);
-
-		UPhysicsAsset* Result = NewObject<UPhysicsAsset>(
-			GetTransientPackage(),
-			AssetName,
-			Flags | EObjectFlags::RF_Public | EObjectFlags::RF_Transient
-		);
-
-		FPhysAssetCreateParams NewBodyData;
-		FText CreationErrorMessage;
-		bool bSuccess = FPhysicsAssetUtils::CreateFromSkeletalMesh(Result, SkeletalMesh, NewBodyData, CreationErrorMessage);
-		if (!bSuccess)
-		{
-			UE_LOG(
-				LogUsd,
-				Error,
-				TEXT("Couldn't create PhysicsAsset for SkeletalMesh '%s': %s"),
-				*SkeletalMesh->GetName(),
-				*CreationErrorMessage.ToString()
-			);
-
-			ensure(ObjectTools::DeleteSingleObject(Result));
-
-			Result = nullptr;
-		}
-
-		return Result;
-#else
-		return nullptr;
-#endif	  // WITH_EDITOR
-	}
-
 	UAnimBlueprint* CreateAnimBlueprint(
 		const FUsdSchemaTranslationContext& Context,
 		const pxr::UsdPrim& SkeletonPrim,
@@ -1010,7 +974,7 @@ namespace UsdSkelSkeletonTranslatorImpl
 		bool* bOutNeedsRecompile = nullptr	  // Whether this function wants the returned AnimBP to be recompiled
 	)
 	{
-		if (!SkeletonPrim || !Context.InfoCache || !Context.AssetCache)
+		if (!SkeletonPrim || !Context.InfoCache || !Context.UsdAssetCache)
 		{
 			return nullptr;
 		}
@@ -1086,9 +1050,8 @@ namespace UsdSkelSkeletonTranslatorImpl
 			SHA1.GetHash(&Hash.Hash[0]);
 			const FString PrefixedAnimBPHash = UsdUtils::GetAssetHashPrefix(SkeletonPrim, Context.bReuseIdenticalAssets) + Hash.ToString();
 
-			// Check if we can find an AnimBP for this prim in the asset cache (useful when doing Action->Import)
 			bool bReusedAnimBP = false;
-			if (UAnimBlueprint* CachedAnimBP = Cast<UAnimBlueprint>(Context.AssetCache->GetCachedAsset(PrefixedAnimBPHash)))
+			if (UAnimBlueprint* CachedAnimBP = Context.UsdAssetCache->GetCachedAsset<UAnimBlueprint>(PrefixedAnimBPHash))
 			{
 				if (CachedAnimBP->TargetSkeleton == Skeleton)
 				{
@@ -1100,28 +1063,32 @@ namespace UsdSkelSkeletonTranslatorImpl
 			// We have to generate a new transient AnimBP
 			if (!bReusedAnimBP)
 			{
-				FName UniqueName = MakeUniqueObjectName(
-					GetTransientPackage(),
-					UAnimBlueprint::StaticClass(),
-					*IUsdClassesModule::SanitizeObjectName(PrimName + TEXT("_DefaultAnimBlueprint"))
-				);
-
 				// Duplicate and never reuse these so that they can be assigned independent subject names if desired.
 				// Its not as if scenes will have thousands of these anyway.
-				AnimBP = DuplicateObject(AnimBP, GetTransientPackage(), UniqueName);
-
-				// Patch up the flags here or else the rest of the engine code (and our plugin code) will get confused
-				// as to what a non-transient asset in the transient package even means. It can lead to some crashes
-				// when saving after import if we don't do this
-				AnimBP->ClearFlags(AnimBP->GetFlags());
-				AnimBP->SetFlags(Skeleton->GetFlags());
-
-				AnimBP->TargetSkeleton = Skeleton;
-				AnimBP->bIsTemplate = false;
+				bool bCreatedAsset = false;
+				const EObjectFlags DesiredFlags = Skeleton->GetFlags();
+				const FString& DesiredAnimBPName = PrimName + TEXT("_DefaultAnimBlueprint");
+				AnimBP = Context.UsdAssetCache->GetOrCreateCustomCachedAsset<UAnimBlueprint>(
+					PrefixedAnimBPHash,
+					DesiredAnimBPName,
+					DesiredFlags,
+					[AnimBP](UPackage* Outer, FName SanitizedName, EObjectFlags FlagsToUse)
+					{
+						UObject* Duplicated = DuplicateObject(AnimBP, Outer, SanitizedName);
+						Duplicated->ClearFlags(Duplicated->GetFlags());
+						Duplicated->SetFlags(FlagsToUse);
+						return Duplicated;
+					},
+					&bCreatedAsset
+				);
 
-				bNeedRecompile = true;
+				if (AnimBP && bCreatedAsset)
+				{
+					AnimBP->TargetSkeleton = Skeleton;
+					AnimBP->bIsTemplate = false;
 
-				Context.AssetCache->CacheAsset(PrefixedAnimBPHash, AnimBP);
+					bNeedRecompile = true;
+				}
 			}
 		}
 		// Path is pointing to an existing, persistent AnimBP
@@ -1534,7 +1501,7 @@ namespace UsdSkelSkeletonTranslatorImpl
 			ESchemaTranslationLaunchPolicy::Sync,
 			[this]()
 			{
-				if (!Context->InfoCache || !Context->AssetCache)
+				if (!Context->InfoCache || !Context->UsdAssetCache)
 				{
 					return false;
 				}
@@ -1559,37 +1526,75 @@ namespace UsdSkelSkeletonTranslatorImpl
 				);
 				PrefixedSkelMeshHash = UsdUtils::GetAssetHashPrefix(GetSkeletonPrim(), Context->bReuseIdenticalAssets) + SkeletalMeshHash.ToString();
 
-				USkeletalMesh* SkeletalMesh = Cast<USkeletalMesh>(Context->AssetCache->GetCachedAsset(PrefixedSkelMeshHash));
+				const FString DesiredSkeletalMeshName = UsdToUnreal::ConvertString(ClosestParentSkelRoot.Get().GetPrim().GetName());
 
+				// Even though we're translating SkeletalMeshes from Skeleton prims now, keep using the SkelRoot
+				// prim name as the SkeletalMesh asset name. This because in the general case the entire SkelRoot will
+				// represent a character and be named after it, and the Skeleton will be named just "Skeleton" or "Skel"
 				bool bIsNew = false;
-				if (!SkeletalMesh)
+				USkeletalMesh* SkeletalMesh = Context->UsdAssetCache->GetOrCreateCachedAsset<USkeletalMesh>(
+					PrefixedSkelMeshHash,
+					DesiredSkeletalMeshName,
+					Context->ObjectFlags,
+					&bIsNew
+				);
+
+				if (bIsNew)
 				{
-					// Even though we're translating SkeletalMeshes from Skeleton prims now, keep using the SkelRoot
-					// prim name as the SkeletalMesh asset name. This because in the general case the entire SkelRoot will
-					// represent a character and be named after it, and the Skeleton will be named just "Skeleton" or "Skel"
-					FString MeshName;
-					{
-						FScopedUsdAllocs Allocs;
-						MeshName = UsdToUnreal::ConvertString(ClosestParentSkelRoot.Get().GetPrim().GetName());
-					}
+					const FString PrefixedSkeletonHash = PrefixedSkelMeshHash + TEXT("_Skeleton");
+					const FString DesiredSkeletonName = SkeletonName.ToString();
+					bool bSkeletonIsNew = false;
+					USkeleton* Skeleton = Context->UsdAssetCache->GetOrCreateCachedAsset<USkeleton>(
+						PrefixedSkeletonHash,
+						DesiredSkeletonName,
+						Context->ObjectFlags,
+						&bSkeletonIsNew
+					);
+
+					Skeleton->SetPreviewMesh(SkeletalMesh);
+					SkeletalMesh->SetSkeleton(Skeleton);
 
-					bIsNew = true;
-					SkeletalMesh = UsdToUnreal::GetSkeletalMeshFromImportData(
+					bool bSuccess = UsdToUnreal::ConvertSkeletalImportData(
 						LODIndexToSkeletalMeshImportData,
 						SkeletonBones,
 						NewBlendShapes,
-						Context->ObjectFlags,
-						*MeshName,
-						SkeletonName
+						SkeletalMesh	// Already has Skeleton
 					);
+
+					if (!bSuccess)
+					{
+						UE_LOG(
+							LogUsd,
+							Warning,
+							TEXT("Failed to create SkeletalMesh for prim '%s'"),
+							*UsdToUnreal::ConvertPath(ClosestParentSkelRoot.Get().GetPrim().GetPrimPath())
+						);
+						UsdUnreal::TranslatorUtils::AbandonFailedAsset(SkeletalMesh, Context->UsdAssetCache.Get(), Context->InfoCache.Get());
+
+						if (bSkeletonIsNew)
+						{
+							UE_LOG(
+								LogUsd,
+								Warning,
+								TEXT("Failed to create Skeleton '%s' for prim '%s'"),
+								*SkeletonName.ToString(),
+								*UsdToUnreal::ConvertPath(ClosestParentSkelRoot.Get().GetPrim().GetPrimPath())
+							);
+							UsdUnreal::TranslatorUtils::AbandonFailedAsset(Skeleton, Context->UsdAssetCache.Get(), Context->InfoCache.Get());
+						}
+					}
 				}
 
 				if (SkeletalMesh)
 				{
 					// Handle the SkeletalMesh AssetUserData
-					if (UUsdMeshAssetUserData* UserData = UsdUtils::GetOrCreateAssetUserData<UUsdMeshAssetUserData>(SkeletalMesh))
+					if (UUsdMeshAssetUserData* UserData = UsdUnreal::ObjectUtils::GetOrCreateAssetUserData<UUsdMeshAssetUserData>(SkeletalMesh))
 					{
-						UserData->PrimvarToUVIndex = LODIndexToMaterialInfo[0].PrimvarToUVIndex;	// We use the same primvar mapping for all LODs
+						if (LODIndexToMaterialInfo.Num() > 0)
+						{
+							// We use the same primvar mapping for all LODs
+							UserData->PrimvarToUVIndex = LODIndexToMaterialInfo[0].PrimvarToUVIndex;
+						}
 						UserData->PrimPaths.AddUnique(SkeletonPrimPath.GetString());
 
 						// For the skel task chain we always collect skeletal mesh metadata when first parsing the prims directly, as it
@@ -1610,7 +1615,7 @@ namespace UsdSkelSkeletonTranslatorImpl
 					// Handle the Skeleton AssetUserData
 					if (USkeleton* Skeleton = SkeletalMesh->GetSkeleton())
 					{
-						if (UUsdAssetUserData* UserData = UsdUtils::GetOrCreateAssetUserData(Skeleton))
+						if (UUsdAssetUserData* UserData = UsdUnreal::ObjectUtils::GetOrCreateAssetUserData(Skeleton))
 						{
 							UserData->PrimPaths.AddUnique(SkeletonPrimPath.GetString());
 
@@ -1644,7 +1649,7 @@ namespace UsdSkelSkeletonTranslatorImpl
 							GetSkeletonPrim(),
 							LODIndexToMaterialInfo,
 							SkeletalMesh,
-							*Context->AssetCache.Get(),
+							*Context->UsdAssetCache.Get(),
 							*Context->InfoCache.Get(),
 							Context->Time,
 							Context->ObjectFlags,
@@ -1657,35 +1662,62 @@ namespace UsdSkelSkeletonTranslatorImpl
 							const bool bRebuildAll = true;
 							SkeletalMesh->UpdateUVChannelData(bRebuildAll);
 						}
-
-						Context->AssetCache->CacheAsset(PrefixedSkelMeshHash, SkeletalMesh);
-						Context->AssetCache->CacheAsset(PrefixedSkelMeshHash + TEXT("_Skeleton"), SkeletalMesh->GetSkeleton());
-
-						// The PreviewSkeletalMesh property on the Skeleton is a soft object path, that is set within
-						// UsdToUnreal::GetSkeletalMeshFromImportData before the SkeletalMesh is part of the cache. When we cache the
-						// SkeletalMesh, it is renamed into the AssetCache package, and so PreviewSkeletalMesh is no longer pointing at it.
-						// Ideally we'd create all our asset cache assets directly outered to the AssetCache to avoid this issue, but
-						// for now I think it's only a problem for SkeletalMesh/Skeleton pairs (since they are created at the same time
-						// already knowing each other), so it's simpler to just do this
-						SkeletalMesh->GetSkeleton()->SetPreviewMesh(SkeletalMesh);
 					}
 
 					if (bGeneratePhysicsAssets)
 					{
 						UPhysicsAsset* PhysicsAsset = SkeletalMesh->GetPhysicsAsset();
+
+#if WITH_EDITOR
 						if (!PhysicsAsset)
 						{
-							PhysicsAsset = UsdSkelSkeletonTranslatorImpl::GenerateAndAssignPhysicsAsset(SkeletalMesh, Context->ObjectFlags);
-
-							if (PhysicsAsset)
+							const FString PhysicsAssetHash = PrefixedSkelMeshHash + TEXT("_PhysicsAsset");
+							const FString DesiredPhysicsAssetName = FPaths::GetBaseFilename(TEXT("PHYS_") + SkeletalMesh->GetName());
+							bool bCreatedPhysicsAsset = false;
+							PhysicsAsset = Context->UsdAssetCache->GetOrCreateCachedAsset<UPhysicsAsset>(
+								PhysicsAssetHash,
+								DesiredPhysicsAssetName,
+								Context->ObjectFlags,
+								&bCreatedPhysicsAsset
+							);
+
+							if (bCreatedPhysicsAsset)
 							{
-								Context->AssetCache->CacheAsset(PrefixedSkelMeshHash + TEXT("_PhysicsAsset"), PhysicsAsset);
+								FPhysAssetCreateParams NewBodyData;
+								FText CreationErrorMessage;
+
+								bool bSuccess = false;
+								if (SkeletalMesh && SkeletalMesh->GetResourceForRendering())
+								{
+									bSuccess = FPhysicsAssetUtils::CreateFromSkeletalMesh(
+										PhysicsAsset,
+										SkeletalMesh,
+										NewBodyData,
+										CreationErrorMessage
+									);
+								}
+								if (!bSuccess)
+								{
+									UE_LOG(
+										LogUsd,
+										Warning,
+										TEXT("Failed to create PhysicsAsset for skeletal mesh '%s'"),
+										*SkeletalMesh->GetPathName()
+									);
+									UsdUnreal::TranslatorUtils::AbandonFailedAsset(
+										PhysicsAsset,
+										Context->UsdAssetCache.Get(),
+										Context->InfoCache.Get()
+									);
+								}
 							}
 						}
+#endif	  // WITH_EDITOR
 
 						if (PhysicsAsset)
 						{
 							Context->InfoCache->LinkAssetToPrim(SkeletonPrimPath, PhysicsAsset);
+							Context->UsdAssetCache->TouchAssetPath(PhysicsAsset);
 						}
 					}
 					else
@@ -1700,6 +1732,7 @@ namespace UsdSkelSkeletonTranslatorImpl
 					if (USkeleton* Skeleton = SkeletalMesh->GetSkeleton())
 					{
 						Context->InfoCache->LinkAssetToPrim(SkeletonPrimPath, Skeleton);
+						Context->UsdAssetCache->TouchAssetPath(Skeleton);
 					}
 
 					// We may be reusing a skeletal mesh we got in the cache, but we always need the BlendShapesByPath stored on the
@@ -1833,37 +1866,30 @@ namespace UsdSkelSkeletonTranslatorImpl
 
 				FSHAHash Hash = UsdSkelSkeletonTranslatorImpl::ComputeSHAHash(SkeletonQuery.Get(), RootMotionPrim, PrefixedSkelMeshHash);
 				FString PrefixedSkelAnimHash = UsdUtils::GetAssetHashPrefix(SkelAnimationPrim, Context->bReuseIdenticalAssets) + Hash.ToString();
-				UAnimSequence* AnimSequence = Cast<UAnimSequence>(Context->AssetCache->GetCachedAsset(PrefixedSkelAnimHash));
+
+				const FString DesiredName = UsdToUnreal::ConvertToken(SkelAnimationPrim.GetName());
+
+				bool bAnimSequenceIsNew = false;
+				UAnimSequence* AnimSequence = Context->UsdAssetCache->GetOrCreateCachedAsset<UAnimSequence>(
+					PrefixedSkelAnimHash,
+					DesiredName,
+					Context->ObjectFlags,
+					&bAnimSequenceIsNew
+				);
 
 				TOptional<float> LayerStartOffsetSeconds;
-				if (!AnimSequence || AnimSequence->GetSkeleton() != SkeletalMesh->GetSkeleton())
+				if (bAnimSequenceIsNew || AnimSequence->GetSkeleton() != SkeletalMesh->GetSkeleton())
 				{
 					FScopedUnrealAllocs UEAllocs;
 
-					FName AnimSequenceName = MakeUniqueObjectName(
-						GetTransientPackage(),
-						UAnimSequence::StaticClass(),
-						*IUsdClassesModule::SanitizeObjectName(UsdToUnreal::ConvertToken(SkelAnimationPrim.GetName()))
-					);
-
-					// The UAnimSequence can't be created with the RF_Transactional flag, or else it will be serialized without
-					// Bone/CurveCompressionSettings. Undoing that transaction would call UAnimSequence::Serialize with nullptr values for
-					// both, which crashes. Besides, this particular asset type is only ever created when we import to content folder assets
-					// (so never for realtime), and in that case we don't need it to be transactional anyway
-					AnimSequence = NewObject<UAnimSequence>(
-						GetTransientPackage(),
-						AnimSequenceName,
-						(Context->ObjectFlags & ~EObjectFlags::RF_Transactional) | EObjectFlags::RF_Transient
-					);
-					AnimSequence->SetSkeleton(SkeletalMesh->GetSkeleton());
-
 					// This is read back in the USDImporter, so that if we ever import this AnimSequence we will always also import the
 					// SkeletalMesh for it
 					AnimSequence->SetPreviewMesh(SkeletalMesh);
+					AnimSequence->SetSkeleton(SkeletalMesh->GetSkeleton());
 
 					TUsdStore<pxr::VtArray<pxr::UsdSkelSkinningQuery>> SkinningTargets = SkeletonBinding.Get().GetSkinningTargets();
 					LayerStartOffsetSeconds = 0.0f;
-					const bool bSuccess = UsdToUnreal::ConvertSkelAnim(
+					bool bSuccess = UsdToUnreal::ConvertSkelAnim(
 						SkeletonQuery.Get(),
 						&SkinningTargets.Get(),
 						&NewBlendShapes,
@@ -1873,19 +1899,18 @@ namespace UsdSkelSkeletonTranslatorImpl
 						&LayerStartOffsetSeconds.GetValue()
 					);
 
-					if (bSuccess
-						&& (AnimSequence->GetDataModel()->GetNumBoneTracks() != 0 || AnimSequence->GetDataModel()->GetNumberOfFloatCurves() != 0))
+					if (!bSuccess
+						|| (AnimSequence->GetDataModel()->GetNumBoneTracks() == 0 && AnimSequence->GetDataModel()->GetNumberOfFloatCurves() == 0))
 					{
-						Context->AssetCache->CacheAsset(PrefixedSkelAnimHash, AnimSequence);
-					}
-					else
-					{
-						AnimSequence->MarkAsGarbage();
+						UE_LOG(LogUsd, Warning, TEXT("Failed to create AnimSequence for prim '%s'"), *SkelAnimationPrimPath);
+						UsdUnreal::TranslatorUtils::AbandonFailedAsset(AnimSequence, Context->UsdAssetCache.Get(), Context->InfoCache.Get());
 						AnimSequence = nullptr;
 					}
 				}
 
-				if (UUsdAnimSequenceAssetUserData* UserData = UsdUtils::GetOrCreateAssetUserData<UUsdAnimSequenceAssetUserData>(AnimSequence))
+				if (UUsdAnimSequenceAssetUserData* UserData = UsdUnreal::ObjectUtils::GetOrCreateAssetUserData<UUsdAnimSequenceAssetUserData>(
+						AnimSequence
+					))
 				{
 					UserData->PrimPaths.AddUnique(SkelAnimationPrimPath);
 
@@ -1959,8 +1984,6 @@ USceneComponent* FUsdSkelSkeletonTranslator::CreateComponents()
 
 	USceneComponent* SceneComponent = FUsdGeomXformableTranslator::CreateComponents();
 
-	// TODO: DrawMode stuff here too?
-
 #if WITH_EDITOR
 	UE::FUsdPrim SkeletonPrim = GetPrim();
 
@@ -1977,7 +2000,7 @@ USceneComponent* FUsdSkelSkeletonTranslator::CreateComponents()
 	}
 
 	// Check if the prim has the GroomBinding schema and setup the component and assets necessary to bind the groom to the SkeletalMesh
-	if (Context->bAllowParsingGroomAssets && Context->AssetCache && Context->InfoCache)
+	if (Context->bAllowParsingGroomAssets && Context->UsdAssetCache && Context->InfoCache)
 	{
 		UE::FUsdPrim PrimWithSchema;
 		if (UsdUtils::PrimHasSchema(SkeletonPrim, UnrealIdentifiers::GroomBindingAPI))
@@ -1994,7 +2017,7 @@ USceneComponent* FUsdSkelSkeletonTranslator::CreateComponents()
 		{
 			UsdGroomTranslatorUtils::CreateGroomBindingAsset(
 				PrimWithSchema,
-				*Context->AssetCache,
+				*Context->UsdAssetCache,
 				*Context->InfoCache,
 				Context->ObjectFlags,
 				Context->bReuseIdenticalAssets
@@ -2138,7 +2161,7 @@ void FUsdSkelSkeletonTranslator::UpdateComponents(USceneComponent* SceneComponen
 				SkeletonBindingForPrim.Get(),
 				ExistingAssignments,
 				*SkeletalMeshComponent,
-				*Context->AssetCache.Get(),
+				*Context->UsdAssetCache.Get(),
 				*Context->InfoCache.Get(),
 				Context->Time,
 				Context->ObjectFlags,
diff --git a/Engine/Plugins/Importers/USDImporter/Source/USDSchemas/Private/USDTranslatorUtils.h b/Engine/Plugins/Importers/USDImporter/Source/USDSchemas/Private/USDTranslatorUtils.h
new file mode 100644
--- /dev/null
+++ b/Engine/Plugins/Importers/USDImporter/Source/USDSchemas/Private/USDTranslatorUtils.h
@@ -0,0 +1,15 @@
+// Copyright Epic Games, Inc. All Rights Reserved.
+
+#pragma once
+
+#include "CoreMinimal.h"
+#include "UObject/ObjectMacros.h"
+
+class UUsdAssetCache3;
+class FUsdInfoCache;
+
+namespace UsdUnreal::TranslatorUtils
+{
+	// Properly deletes the asset and removes it from the asset and info caches, if provided
+	void AbandonFailedAsset(UObject* Asset, UUsdAssetCache3* AssetCache, FUsdInfoCache* InfoCache);
+}
diff --git a/Engine/Plugins/Importers/USDImporter/Source/USDSchemas/Private/USDVolVolumeTranslator.cpp b/Engine/Plugins/Importers/USDImporter/Source/USDSchemas/Private/USDVolVolumeTranslator.cpp
--- a/Engine/Plugins/Importers/USDImporter/Source/USDSchemas/Private/USDVolVolumeTranslator.cpp
+++ b/Engine/Plugins/Importers/USDImporter/Source/USDSchemas/Private/USDVolVolumeTranslator.cpp
@@ -9,6 +9,7 @@
 #include "USDIntegrationUtils.h"
 #include "USDLayerUtils.h"
 #include "USDLog.h"
+#include "USDObjectUtils.h"
 #include "USDPrimConversion.h"
 #include "USDProjectSettings.h"
 #include "USDShadeConversion.h"
@@ -905,71 +906,6 @@ namespace UE::UsdVolVolumeTranslator::Private
 		return ResultParameterToInfo;
 	}
 
-	UMaterialInstance* InstantiateMaterial(
-		FName InstanceName,
-		UMaterialInterface* ReferenceMaterial,
-		const FUsdSchemaTranslationContext& TranslationContext
-	)
-	{
-		if (!ReferenceMaterial)
-		{
-			return nullptr;
-		}
-
-		// Create an UMaterialInstanceConstant
-		if (GIsEditor)	  // Also have to prevent Standalone game from going with MaterialInstanceConstants
-		{
-			UMaterialInstanceConstant* NewMaterial = NewObject<UMaterialInstanceConstant>(
-				GetTransientPackage(),
-				InstanceName,
-				TranslationContext.ObjectFlags | EObjectFlags::RF_Transient
-			);
-
-			if (NewMaterial)
-			{
-				if (ensure(ReferenceMaterial))
-				{
-					// We can't blindly recreate all component render states when a level is being added, because we may end up first creating
-					// render states for some components, and UWorld::AddToWorld calls FScene::AddPrimitive which expects the component to not
-					// have primitives yet
-					FMaterialUpdateContext::EOptions::Type Options = FMaterialUpdateContext::EOptions::Default;
-					if (TranslationContext.Level && TranslationContext.Level->bIsAssociatingLevel)
-					{
-						Options = (FMaterialUpdateContext::EOptions::Type)(Options & ~FMaterialUpdateContext::EOptions::RecreateRenderStates);
-					}
-					FMaterialUpdateContext UpdateContext(Options, GMaxRHIShaderPlatform);
-					UpdateContext.AddMaterialInstance(NewMaterial);
-
-					NewMaterial->SetParentEditorOnly(ReferenceMaterial);
-
-					NewMaterial->PreEditChange(nullptr);
-					NewMaterial->PostEditChange();
-
-					return NewMaterial;
-				}
-			}
-		}
-		else
-		// Create an UMaterialInstanceDynamic
-		{
-			// SparseVolumeTextures can't be created at runtime so this branch should never really be taken for now, but anyway...
-			// Note: Some code in FNiagaraBakerRenderer::RenderSparseVolumeTexture suggests that this workflow wouldn't really work
-			// because the HeterogeneousVolumeComponent always creates its own MID from the material we give it, and creating a MID
-			// from another MID doesn't really work
-			if (ensure(ReferenceMaterial))
-			{
-				UMaterialInstanceDynamic* NewMaterial = UMaterialInstanceDynamic::Create(ReferenceMaterial, GetTransientPackage(), InstanceName);
-				if (NewMaterial)
-				{
-					NewMaterial->SetFlags(RF_Transient);
-				}
-				return NewMaterial;
-			}
-		}
-
-		return nullptr;
-	}
-
 	void AssignMaterialParameters(UMaterialInstance* MaterialInstance, const TMap<FString, FSparseVolumeTextureInfo*>& ParameterToTexture)
 	{
 		// Now that we finally have the parameter assignment for each SVT, assign them to the materials
@@ -997,7 +933,7 @@ void FUsdVolVolumeTranslator::CreateAssets()
 
 	using namespace UE::UsdVolVolumeTranslator::Private;
 
-	if (!Context->AssetCache || !Context->InfoCache)
+	if (!Context->UsdAssetCache || !Context->InfoCache)
 	{
 		return;
 	}
@@ -1072,50 +1008,53 @@ void FUsdVolVolumeTranslator::CreateAssets()
 		}
 		ParsedTexture.PrefixedAssetHash = VolumePrimHashPrefix + VDBAndAssignmentHash.ToString();
 
-		USparseVolumeTexture* SparseVolumeTexture = Cast<USparseVolumeTexture>(Context->AssetCache->GetCachedAsset(ParsedTexture.PrefixedAssetHash));
+		// File path instead of prim path in case we have  multiple .vdb files in the same Volume prim
+		const FString& DesiredName = FPaths::GetBaseFilename(VDBFilePath);
 
-		// Need to create a brand new asset
+		bool bCreatedNew = false;
+		USparseVolumeTexture* SparseVolumeTexture = Context->UsdAssetCache->GetOrCreateCustomCachedAsset<USparseVolumeTexture>(
+			ParsedTexture.PrefixedAssetHash,
+			DesiredName,
+			Context->ObjectFlags,
+			[&VDBFilePath, &ImportOptions](UPackage* Outer, FName SanitizedName, EObjectFlags FlagsToUse)
+			{
+				TStrongObjectPtr<USparseVolumeTextureFactory> SparseVolumeTextureFactory{NewObject<USparseVolumeTextureFactory>()};
+
+				// We use the asset import task to indicate it's an automated import, and also to transmit our import options
+				TStrongObjectPtr<UAssetImportTask> AssetImportTask{NewObject<UAssetImportTask>()};
+				AssetImportTask->Filename = VDBFilePath;
+				AssetImportTask->bAutomated = true;
+				AssetImportTask->bSave = false;
+				AssetImportTask->Options = ImportOptions.Get();
+				AssetImportTask->Factory = SparseVolumeTextureFactory.Get();
+				SparseVolumeTextureFactory->SetAssetImportTask(AssetImportTask.Get());
+
+				// Call FactoryCreateFile directly here or else the usual AssetToolsModule.Get().ImportAssetTasks()
+				// workflow would end up creating a package for every asset, which we don't care about since
+				// the asset cache will do that anyway
+				const TCHAR* Parms = nullptr;
+				bool bOperationCanceled = false;
+				return Cast<USparseVolumeTexture>(SparseVolumeTextureFactory->FactoryCreateFile(
+					USparseVolumeTexture::StaticClass(),
+					Outer,
+					SanitizedName,
+					FlagsToUse,
+					VDBFilePath,
+					Parms,
+					GWarn,
+					bOperationCanceled
+				));
+			},
+			&bCreatedNew
+		);
 		if (!SparseVolumeTexture)
 		{
-			TStrongObjectPtr<USparseVolumeTextureFactory> SparseVolumeTextureFactory{NewObject<USparseVolumeTextureFactory>()};
-
-			// We use the asset import task to indicate it's an automated import, and also to transmit our import options
-			TStrongObjectPtr<UAssetImportTask> AssetImportTask{NewObject<UAssetImportTask>()};
-			AssetImportTask->Filename = VDBFilePath;
-			AssetImportTask->bAutomated = true;
-			AssetImportTask->bSave = false;
-			AssetImportTask->Options = ImportOptions.Get();
-			AssetImportTask->Factory = SparseVolumeTextureFactory.Get();
-
-			bool bOperationCanceled = false;
-			const TCHAR* Parms = nullptr;
-			const FName AssetName = MakeUniqueObjectName(
-				GetTransientPackage(),
-				USparseVolumeTexture::StaticClass(),
-				*IUsdClassesModule::SanitizeObjectName(FPaths::GetBaseFilename(VDBFilePath))	// File path instead of prim path in case we have
-																								// multiple .vdb files in the same Volume prim
-			);
-
-			// Call FactoryCreateFile directly here or else the usual AssetToolsModule.Get().ImportAssetTasks()
-			// workflow would end up creating a package for every asset, which we don't care about since
-			// we'll rename them into the asset cache anyway
-			SparseVolumeTextureFactory->SetAssetImportTask(AssetImportTask.Get());
-			SparseVolumeTexture = Cast<USparseVolumeTexture>(SparseVolumeTextureFactory->FactoryCreateFile(
-				USparseVolumeTexture::StaticClass(),
-				GetTransientPackage(),
-				AssetName,
-				Context->ObjectFlags | EObjectFlags::RF_Public | EObjectFlags::RF_Transient,
-				VDBFilePath,
-				Parms,
-				GWarn,
-				bOperationCanceled
-			));
-			if (!SparseVolumeTexture)
-			{
-				UE_LOG(LogUsd, Error, TEXT("Failed to generate Sparse Volume Texture from OpenVDB file '%s'"), *VDBFilePath);
-				return;
-			}
+			UE_LOG(LogUsd, Error, TEXT("Failed to generate Sparse Volume Texture from OpenVDB file '%s'"), *VDBFilePath);
+			return;
+		}
 
+		if (bCreatedNew && SparseVolumeTexture)
+		{
 			SparseVolumeTexture->PostEditChange();
 
 			if (UStreamableSparseVolumeTexture* StreamableTexture = Cast<UStreamableSparseVolumeTexture>(SparseVolumeTexture))
@@ -1127,17 +1066,14 @@ void FUsdVolVolumeTranslator::CreateAssets()
 
 				StreamableTexture->AssetImportData = ImportData;
 			}
-
-			Context->AssetCache->CacheAsset(ParsedTexture.PrefixedAssetHash, SparseVolumeTexture);
 		}
 
 		if (SparseVolumeTexture)
 		{
 			Context->InfoCache->LinkAssetToPrim(PrimPath, SparseVolumeTexture);
 
-			if (UUsdSparseVolumeTextureAssetUserData* UserData = UsdUtils::GetOrCreateAssetUserData<UUsdSparseVolumeTextureAssetUserData>(
-					SparseVolumeTexture
-				))
+			if (UUsdSparseVolumeTextureAssetUserData* UserData = UsdUnreal::ObjectUtils::GetOrCreateAssetUserData<
+					UUsdSparseVolumeTextureAssetUserData>(SparseVolumeTexture))
 			{
 				UserData->PrimPaths.AddUnique(VolumePrimPathString);
 				UserData->SourceOpenVDBAssetPrimPaths = ParsedTexture.SourceOpenVDBAssetPrimPaths;
@@ -1260,41 +1196,72 @@ void FUsdVolVolumeTranslator::CreateAssets()
 	}
 	const FString PrefixedMaterialHash = VolumePrimHashPrefix + MaterialHash.ToString();
 
-	UMaterialInstance* MaterialInstance = nullptr;
+	const FString DesiredName = FPaths::GetBaseFilename(VolumePrimPathString);
 
-	if (Context->AssetCache)
+	bool bIsNew = false;
+	UMaterialInstance* MaterialInstance = nullptr;
+	if (GIsEditor)
 	{
-		MaterialInstance = Cast<UMaterialInstance>(Context->AssetCache->GetCachedAsset(PrefixedMaterialHash));
-	}
+		// Create an UMaterialInstanceConstant
 
-	// Create new material instance
-	if (!MaterialInstance)
-	{
-		const FName InstanceName = MakeUniqueObjectName(
-			GetTransientPackage(),
-			UMaterialInstance::StaticClass(),
-			*IUsdClassesModule::SanitizeObjectName(FPaths::GetBaseFilename(VolumePrimPathString))
+		UMaterialInstanceConstant* MIC = Context->UsdAssetCache->GetOrCreateCachedAsset<UMaterialInstanceConstant>(
+			PrefixedMaterialHash,
+			DesiredName,
+			Context->ObjectFlags,
+			&bIsNew
 		);
 
-		MaterialInstance = InstantiateMaterial(InstanceName, ReferenceMaterial, Context.Get());
-
-		if (MaterialInstance)
+		FMaterialUpdateContext::EOptions::Type Options = FMaterialUpdateContext::EOptions::Default;
+		if (Context->Level && Context->Level->bIsAssociatingLevel)
 		{
-			UUsdAssetUserData* UserData = NewObject<UUsdAssetUserData>(MaterialInstance, TEXT("USDAssetUserData"));
-			UserData->PrimPaths = {PrimPath.GetString()};
-			MaterialInstance->AddAssetUserData(UserData);
+			Options = (FMaterialUpdateContext::EOptions::Type)(Options & ~FMaterialUpdateContext::EOptions::RecreateRenderStates);
+		}
+		FMaterialUpdateContext UpdateContext(Options, GMaxRHIShaderPlatform);
+		UpdateContext.AddMaterialInstance(MIC);
+		MIC->SetParentEditorOnly(ReferenceMaterial);
+		MIC->PreEditChange(nullptr);
+		MIC->PostEditChange();
+
+		MaterialInstance = MIC;
+	}
+	else
+	{
+		// Create a material instance for the volume component.
+		// SparseVolumeTextures can't be created at runtime so this branch should never really be taken for now, but anyway...
+		// Note: Some code in FNiagaraBakerRenderer::RenderSparseVolumeTexture suggests that this workflow wouldn't really work
+		// because the HeterogeneousVolumeComponent always creates its own MID from the material we give it, and creating a MID
+		// from another MID doesn't really work
+
+		bool bCreatedAsset = false;
+		UMaterialInstance* MI = Context->UsdAssetCache->GetOrCreateCustomCachedAsset<UMaterialInstance>(
+			PrefixedMaterialHash,
+			DesiredName,
+			Context->ObjectFlags | RF_Transient,	// We never want MIDs to become assets in the content browser
+			[ReferenceMaterial](UPackage* Outer, FName SanitizedName, EObjectFlags FlagsToUse)
+			{
+				UMaterialInstanceDynamic* NewMID = UMaterialInstanceDynamic::Create(ReferenceMaterial, Outer, SanitizedName);
+				NewMID->ClearFlags(NewMID->GetFlags());
+				NewMID->SetFlags(FlagsToUse);
+				return NewMID;
+			},
+			&bCreatedAsset
+		);
 
-			AssignMaterialParameters(MaterialInstance, MaterialParameterToTexture);
+		MI->PreEditChange(nullptr);
+		MI->PostEditChange();
+	}
 
-			Context->AssetCache->CacheAsset(PrefixedMaterialHash, MaterialInstance);
-		}
+	// Create new material instance
+	if (bIsNew && MaterialInstance)
+	{
+		AssignMaterialParameters(MaterialInstance, MaterialParameterToTexture);
 	}
 
 	if (MaterialInstance)
 	{
 		Context->InfoCache->LinkAssetToPrim(PrimPath, MaterialInstance);
 
-		if (UUsdAssetUserData* UserData = UsdUtils::GetOrCreateAssetUserData<UUsdAssetUserData>(MaterialInstance))
+		if (UUsdAssetUserData* UserData = UsdUnreal::ObjectUtils::GetOrCreateAssetUserData<UUsdAssetUserData>(MaterialInstance))
 		{
 			UserData->PrimPaths.AddUnique(VolumePrimPathString);
 
@@ -1381,7 +1348,7 @@ void FUsdVolVolumeTranslator::UpdateComponents(USceneComponent* SceneComponent)
 					if (SparseVolumeTexture->GetNumFrames() > 1)
 					{
 						if (UUsdSparseVolumeTextureAssetUserData* UserData = Cast<UUsdSparseVolumeTextureAssetUserData>(
-								UsdUtils::GetAssetUserData(SparseVolumeTexture)
+								UsdUnreal::ObjectUtils::GetAssetUserData(SparseVolumeTexture)
 							))
 						{
 							UE::FUsdPrim VolumePrim = GetPrim();
diff --git a/Engine/Plugins/Importers/USDImporter/Source/USDSchemas/Public/USDLuxLightTranslator.h b/Engine/Plugins/Importers/USDImporter/Source/USDSchemas/Public/USDLuxLightTranslator.h
--- a/Engine/Plugins/Importers/USDImporter/Source/USDSchemas/Public/USDLuxLightTranslator.h
+++ b/Engine/Plugins/Importers/USDImporter/Source/USDSchemas/Public/USDLuxLightTranslator.h
@@ -16,6 +16,7 @@ public:
 
 	using FUsdGeomXformableTranslator::FUsdGeomXformableTranslator;
 
+	virtual void CreateAssets() override;
 	virtual USceneComponent* CreateComponents() override;
 	virtual void UpdateComponents(USceneComponent* SceneComponent);
 
diff --git a/Engine/Plugins/Importers/USDImporter/Source/USDStage/Private/USDLevelSequenceHelper.cpp b/Engine/Plugins/Importers/USDImporter/Source/USDStage/Private/USDLevelSequenceHelper.cpp
--- a/Engine/Plugins/Importers/USDImporter/Source/USDStage/Private/USDLevelSequenceHelper.cpp
+++ b/Engine/Plugins/Importers/USDImporter/Source/USDStage/Private/USDLevelSequenceHelper.cpp
@@ -4,7 +4,6 @@
 
 #include "USDAssetUserData.h"
 #include "USDAttributeUtils.h"
-#include "USDClassesModule.h"
 #include "USDConversionUtils.h"
 #include "USDDrawModeComponent.h"
 #include "USDInfoCache.h"
@@ -12,6 +11,7 @@
 #include "USDLayerUtils.h"
 #include "USDListener.h"
 #include "USDLog.h"
+#include "USDObjectUtils.h"
 #include "USDPrimConversion.h"
 #include "USDPrimTwin.h"
 #include "USDProjectSettings.h"
@@ -335,12 +335,7 @@ namespace UsdLevelSequenceHelperImpl
 			AESP.RootToLocalTransform = RootToLocalTransform;
 			AESP.MovieSceneSequence = LevelSequence;
 			AESP.RootMovieSceneSequence = LevelSequence;
-			bResult = MovieSceneToolHelpers::ExportToAnimSequence(
-				AnimSequence,
-				ExportOptions,
-				AESP,
-				SkeletalMeshComp
-			);
+			bResult = MovieSceneToolHelpers::ExportToAnimSequence(AnimSequence, ExportOptions, AESP, SkeletalMeshComp);
 			if (!bResult)
 			{
 				goto cleanup;
@@ -956,6 +951,22 @@ bool FUsdLevelSequenceHelperImpl::HasData() const
 
 void FUsdLevelSequenceHelperImpl::Clear()
 {
+	// Mark old sequences as garbage so that the track references' to the assets can be immediately ignored
+	for (const TPair<FString, TObjectPtr<ULevelSequence>>& IdentifierAndSeq : LevelSequencesByIdentifier)
+	{
+		if (ULevelSequence* Seq = IdentifierAndSeq.Value.Get())
+		{
+#if WITH_EDITOR
+			if (!IsEngineExitRequested())
+			{
+				GEditor->GetEditorSubsystem<UAssetEditorSubsystem>()->CloseAllEditorsForAsset(Seq);
+			}
+#endif	  // WITH_EDITOR
+
+			Seq->MarkAsGarbage();
+		}
+	}
+
 	MainLevelSequence = nullptr;
 	LevelSequencesByIdentifier.Empty();
 	IdentifierByLevelSequence.Empty();
@@ -1244,11 +1255,11 @@ ULevelSequence* FUsdLevelSequenceHelperImpl::FindOrAddSequenceForLayer(
 		// only get an InfoCache when importing (from UUsdStageImporter::ImportFromFile) or when BindToUsdStageActor is called, which also gives us a
 		// stage actor. So if we don't have an actor but have a cache, we're importing
 		const bool bIsImporting = StageActor.IsExplicitlyNull() && InfoCache;
-		FName UniqueSequenceName = bIsImporting ? *IUsdClassesModule::SanitizeObjectName(FPaths::GetBaseFilename(SequenceDisplayName))
+		FName UniqueSequenceName = bIsImporting ? *UsdUnreal::ObjectUtils::SanitizeObjectName(FPaths::GetBaseFilename(SequenceDisplayName))
 												: MakeUniqueObjectName(
 													GetTransientPackage(),
 													ULevelSequence::StaticClass(),
-													*IUsdClassesModule::SanitizeObjectName(FPaths::GetBaseFilename(SequenceDisplayName))
+													*UsdUnreal::ObjectUtils::SanitizeObjectName(FPaths::GetBaseFilename(SequenceDisplayName))
 												);
 
 		Sequence = NewObject<ULevelSequence>(GetTransientPackage(), UniqueSequenceName, FUsdLevelSequenceHelperImpl::DefaultObjFlags);
@@ -2534,7 +2545,7 @@ void FUsdLevelSequenceHelperImpl::AddVolumeTracks(const UUsdPrimTwin& PrimTwin,
 			{
 				if (SparseVolumeTexture->GetNumFrames() > 1)
 				{
-					UserData = Cast<UUsdSparseVolumeTextureAssetUserData>(UsdUtils::GetAssetUserData(SparseVolumeTexture));
+					UserData = Cast<UUsdSparseVolumeTextureAssetUserData>(UsdUnreal::ObjectUtils::GetAssetUserData(SparseVolumeTexture));
 				}
 			}
 
diff --git a/Engine/Plugins/Importers/USDImporter/Source/USDStage/Private/USDPrimTwin.cpp b/Engine/Plugins/Importers/USDImporter/Source/USDStage/Private/USDPrimTwin.cpp
--- a/Engine/Plugins/Importers/USDImporter/Source/USDStage/Private/USDPrimTwin.cpp
+++ b/Engine/Plugins/Importers/USDImporter/Source/USDStage/Private/USDPrimTwin.cpp
@@ -20,7 +20,7 @@ UUsdPrimTwin& UUsdPrimTwin::AddChild(const FString& InPrimPath)
 	TObjectPtr<UUsdPrimTwin>& ChildPrim = Children.Add(ChildPrimName);
 
 	// Needs public because this will mostly live on the transient package (c.f. AUsdStageActor::GetRootPrimTwin())
-	ChildPrim = NewObject<UUsdPrimTwin>(this, NAME_None, RF_Transient | RF_Transactional | RF_Public);
+	ChildPrim = NewObject<UUsdPrimTwin>(this, NAME_None, GetFlags() | RF_Public);
 	ChildPrim->PrimPath = InPrimPath;
 
 	ChildPrim->Parent = this;
diff --git a/Engine/Plugins/Importers/USDImporter/Source/USDStage/Private/USDStageActor.cpp b/Engine/Plugins/Importers/USDImporter/Source/USDStage/Private/USDStageActor.cpp
--- a/Engine/Plugins/Importers/USDImporter/Source/USDStage/Private/USDStageActor.cpp
+++ b/Engine/Plugins/Importers/USDImporter/Source/USDStage/Private/USDStageActor.cpp
@@ -4,7 +4,8 @@
 
 #include "UnrealUSDWrapper.h"
 #include "UObject/Package.h"
-#include "USDAssetCache.h"
+#include "USDAssetCache2.h"
+#include "USDAssetCache3.h"
 #include "USDAssetUserData.h"
 #include "USDClassesModule.h"
 #include "USDConversionUtils.h"
@@ -19,6 +20,7 @@
 #include "USDLightConversion.h"
 #include "USDListener.h"
 #include "USDLog.h"
+#include "USDObjectUtils.h"
 #include "USDPrimConversion.h"
 #include "USDPrimTwin.h"
 #include "USDProjectSettings.h"
@@ -53,6 +55,7 @@
 #include "Engine/Engine.h"
 #include "Engine/Level.h"
 #include "Engine/Light.h"
+#include "Engine/StaticMesh.h"
 #include "Engine/World.h"
 #include "EngineAnalytics.h"
 #include "Framework/Notifications/NotificationManager.h"
@@ -115,7 +118,7 @@ static FAutoConsoleVariableRef CVarTranslateOnlyUsedMaterialsWhenOpeningStage(
 		 "translated into Unreal material assets.")
 );
 
-static const EObjectFlags DefaultObjFlag = EObjectFlags::RF_Transactional | EObjectFlags::RF_Transient;
+static const EObjectFlags DefaultObjFlag = EObjectFlags::RF_Transactional | EObjectFlags::RF_Public | EObjectFlags::RF_Standalone;
 
 AUsdStageActor::FOnActorLoaded AUsdStageActor::OnActorLoaded;
 AUsdStageActor::FOnOpenStageEditorClicked AUsdStageActor::OnOpenStageEditorClicked;
@@ -126,7 +129,7 @@ struct FUsdStageActorImpl
 	{
 		TSharedRef<FUsdSchemaTranslationContext> TranslationContext = MakeShared<FUsdSchemaTranslationContext>(
 			StageActor->GetOrOpenUsdStage(),
-			*StageActor->UsdAssetCache
+			*StageActor->AssetCache
 		);
 
 		TranslationContext->Level = StageActor->GetLevel();
@@ -1154,7 +1157,7 @@ AUsdStageActor::AUsdStageActor()
 						// other client, and were replicated here with our RootLayer path already filled out, meaning we should just load that stage
 						// Note that now our UUsdTransactor may have already caused the stage itself to be loaded, but we may still need to call
 						// LoadUsdStage on our end.
-						else if (ObjectsToWatch.Num() == 0 && (!UsdAssetCache || UsdAssetCache->GetNumAssets() == 0))
+						else if (ObjectsToWatch.Num() == 0 && (!AssetCache || AssetCache->GetNumAssets() == 0))
 						{
 							this->LoadUsdStage();
 							AUsdStageActor::OnActorLoaded.Broadcast(this);
@@ -1505,24 +1508,17 @@ void AUsdStageActor::OnUsdObjectsChanged(const UsdUtils::FObjectChangesByPath& I
 		}
 	}
 
-	SetupAssetCacheIfNeeded();
-	if (!UsdAssetCache)
+	if (!AssetCache)
 	{
-		// If the user canceled out of providing an asset cache and we should be loading assets, we have no choice but to either change the stage
-		// state or close the stage, otherwise we'd be left in an invalid state. Closing the stage should be more visible though, so let's do that
-		UE_LOG(
-			LogUsd,
-			Warning,
-			TEXT("Closing the stage '%s' as no asset cache was provided, but the AUsdStageActor '%s' was set to open the stage and load assets. "
-				 "Either provide an asset cache or switch the stage actor to the 'Opened' state"),
-			*RootLayer.FilePath,
-			*GetPathName()
-		);
+		SetupAssetCacheIfNeeded();
+	}
+	if (!ensure(AssetCache))
+	{
+		// This should never happen now: We should always be able to find an asset cache
 		CloseUsdStage();
-
 		return;
 	}
-	FUsdScopedAssetCacheReferencer ScopedReferencer{UsdAssetCache, this};
+	UUsdAssetCache3::FUsdScopedReferencer ScopedReferencer{AssetCache, this};
 
 	// Mark the level as dirty since we received a notice about our stage having changed in some way.
 	// The main goal of this is to trigger the "save layers" dialog if we then save the UE level
@@ -2006,11 +2002,6 @@ void AUsdStageActor::OnUsdObjectsChanged(const UsdUtils::FObjectChangesByPath& I
 		FUsdStageActorImpl::RepairExternalSequencerBindings();
 	}
 
-	if (bHasLoadedOrAbandonedAssets)
-	{
-		UsdAssetCache->RefreshStorage();
-	}
-
 #if WITH_EDITOR
 	if (GIsEditor && GEditor && !IsGarbageCollecting())	   // Make sure we're not in standalone either
 	{
@@ -2032,8 +2023,8 @@ USDSTAGE_API void AUsdStageActor::Reset()
 
 	Super::Reset();
 
-	UnloadUsdStage();
-	CloseUsdStage();
+	bool bUnloadIfNeeded = true;
+	CloseUsdStage(bUnloadIfNeeded);
 
 	Time = 0.f;
 	RootLayer.FilePath.Empty();
@@ -2221,7 +2212,7 @@ UUsdPrimTwin* AUsdStageActor::ExpandPrim(const UE::FUsdPrim& Prim, bool bResync,
 					continue;
 				}
 
-				UUsdAssetUserData* UserData = UsdUtils::GetAssetUserData(HISMMesh);
+				UUsdAssetUserData* UserData = UsdUnreal::ObjectUtils::GetAssetUserData(HISMMesh);
 				if (!UserData)
 				{
 					continue;
@@ -2538,6 +2529,7 @@ void AUsdStageActor::SetStageState(EUsdStageState NewState)
 	}
 }
 
+PRAGMA_DISABLE_DEPRECATION_WARNINGS
 void AUsdStageActor::SetAssetCache(UUsdAssetCache2* NewCache)
 {
 	if (NewCache == UsdAssetCache)
@@ -2587,6 +2579,52 @@ void AUsdStageActor::SetAssetCache(UUsdAssetCache2* NewCache)
 	ResyncChanges.Add({TEXT("/"), {}});
 	OnUsdObjectsChanged(InfoChanges, ResyncChanges);
 }
+PRAGMA_ENABLE_DEPRECATION_WARNINGS
+
+void AUsdStageActor::SetUsdAssetCache(UUsdAssetCache3* NewCache)
+{
+	if (NewCache == AssetCache)
+	{
+		return;
+	}
+
+	const bool bMarkDirty = false;
+	Modify(bMarkDirty);
+
+	if (AssetCache)
+	{
+		AssetCache->RemoveAllReferencerAssets(this);
+		AssetCache->RequestDelayedAssetAutoCleanup();
+	}
+
+	AssetCache = NewCache;
+
+	// We can't have no cache while we have a stage loaded, so fetch a new one from the IUsdClassesModule
+	// instead, as the intent may have been to just have the actor not point at the previous cache anymore.
+	if (!AssetCache && UsdStage)
+	{
+		FNotificationInfo Toast(LOCTEXT("MustHaveCache", "Must have an Asset Cache"));
+		Toast.SubText = LOCTEXT(
+			"MustHaveCache_Subtext",
+			"The Stage Actor must always have an Asset Cache while a stage is loaded, so a temporary cache will be created.\n\nClose the stage "
+			"before clearing the cache if you wish to clear this property."
+		);
+		Toast.Image = FCoreStyle::Get().GetBrush(TEXT("MessageLog.Warning"));
+		Toast.bUseLargeFont = false;
+		Toast.bFireAndForget = true;
+		Toast.FadeOutDuration = 1.0f;
+		Toast.ExpireDuration = 12.0f;
+		Toast.bUseThrobber = false;
+		Toast.bUseSuccessFailIcons = false;
+		FSlateNotificationManager::Get().AddNotification(Toast);
+
+		SetupAssetCacheIfNeeded();
+	}
+
+	// We used to trigger a reload here, forcing new assets to cached into the new cache.
+	// There's no longer any real need for it? If the assets are used by components the object
+	// references will be enough to prevent the old assets from being deleted anyway
+}
 
 void AUsdStageActor::SetInitialLoadSet(EUsdInitialLoadSet NewLoadSet)
 {
@@ -3142,7 +3180,7 @@ void AUsdStageActor::OnObjectsReplaced(const TMap<UObject*, UObject*>& ObjectRep
 			NewActor->OnPreStageChanged = OnPreStageChanged;
 			NewActor->OnPrimChanged = OnPrimChanged;
 
-			NewActor->UsdAssetCache = UsdAssetCache;
+			NewActor->AssetCache = AssetCache;
 
 			NewActor->InfoCache = InfoCache;
 			InfoCache = nullptr;
@@ -3185,26 +3223,13 @@ void AUsdStageActor::LoadUsdStage(bool bOpenIfNeeded)
 {
 	TRACE_CPUPROFILER_EVENT_SCOPE(AUsdStageActor::LoadUsdStage);
 
-	// Ensure we have an asset cache before we open/load anything.
-	// In theory we don't need an asset cache for just opening the stage (especially if we have StageState that is not OpenAndLoaded),
-	// but that may not be immediately obvious to the user anyway, and it seems that the expected behavior is to abort loading the stage
-	// in case the user cancels out of the "missing asset cache" dialog
+	// We should always have an asset cache when opening a stage now, as IUsdClassesModule::GetAssetCacheForProject
+	// should never fail
 	if (!RootLayer.FilePath.IsEmpty())
 	{
 		SetupAssetCacheIfNeeded();
-
-		// The only way we're coming out of SetupAssetCacheIfNeeded without an asset cache is if the user canceled the
-		// asset cache dialog. In that case let's just abort loading this stage entirely and revert our RootLayer
-		// to empty
-		if (!UsdAssetCache)
+		if (!ensure(AssetCache))
 		{
-			UE_LOG(
-				LogUsd,
-				Log,
-				TEXT("Cancelling the loading of USD Stage '%s' as the Asset Cache setup was explicitly canceled."),
-				*RootLayer.FilePath
-			);
-
 			RootLayer.FilePath.Empty();
 			return;
 		}
@@ -3244,11 +3269,7 @@ void AUsdStageActor::LoadUsdStage(bool bOpenIfNeeded)
 
 	FScopedUsdMessageLog ScopedMessageLog;
 
-	TOptional<FUsdScopedAssetCacheReferencer> ScopedReferencer;
-	if (UsdAssetCache)
-	{
-		ScopedReferencer.Emplace(UsdAssetCache, this);
-	}
+	UUsdAssetCache3::FUsdScopedReferencer ScopedReferencer{AssetCache, this};
 
 	UE::FUsdStage StageToLoad = GetUsdStage();
 
@@ -3307,11 +3328,6 @@ void AUsdStageActor::LoadUsdStage(bool bOpenIfNeeded)
 		}
 	}
 
-	if (bLoadedOrAbandonedAssets && UsdAssetCache)
-	{
-		UsdAssetCache->RefreshStorage();
-	}
-
 	FUsdStageActorImpl::RepairExternalSequencerBindings();
 
 #if WITH_EDITOR
@@ -3374,16 +3390,26 @@ void AUsdStageActor::UnloadUsdStage()
 			// destroying the Sequencer -> Crash).
 			// Note that an AsyncTask has it run on the same tick and so within the same frame, while FTSTicker would end
 			// up on the next tick, showing a frame of the broken sequence bindings on the UI, which doesn't look great...
-			ULevelSequence* LevelSequencePtr = LevelSequence;
+			TWeakObjectPtr<ULevelSequence> LevelSequencePtr = LevelSequence;
 			AsyncTask(
 				ENamedThreads::GameThread,
 				[LevelSequencePtr]()
 				{
-					GEditor->GetEditorSubsystem<UAssetEditorSubsystem>()->CloseAllEditorsForAsset(LevelSequencePtr);
+					if (ULevelSequence* ValidSequence = LevelSequencePtr.Get())
+					{
+						GEditor->GetEditorSubsystem<UAssetEditorSubsystem>()->CloseAllEditorsForAsset(ValidSequence);
+
+						// It's important to mark as garbage so that DeleteUnreferencedAssets can delete assets
+						// that are referenced by LevelSequence tracks, like AnimSequences and geometry caches.
+						// DeleteUnreferencedAssets will run on the *next* tick, so doing this on the AsyncTask is still
+						// going to take place before it
+						ValidSequence->MarkAsGarbage();
+					}
 				}
 			);
 		}
 #endif	  // WITH_EDITOR
+
 		LevelSequence = nullptr;
 	}
 	LevelSequenceHelper.Clear();
@@ -3411,10 +3437,10 @@ void AUsdStageActor::UnloadUsdStage()
 	}
 #endif	  // WITH_EDITOR
 
-	if (UsdAssetCache)
+	if (AssetCache)
 	{
-		UsdAssetCache->RemoveAllAssetReferences(this);
-		UsdAssetCache->RefreshStorage();
+		AssetCache->RemoveAllReferencerAssets(this);
+		AssetCache->RequestDelayedAssetAutoCleanup();
 	}
 
 	if (InfoCache)
@@ -3432,84 +3458,10 @@ void AUsdStageActor::UnloadUsdStage()
 
 void AUsdStageActor::SetupAssetCacheIfNeeded()
 {
-	if (!UsdAssetCache)
+	if (!AssetCache)
 	{
-		if (const UUsdProjectSettings* ProjectSettings = GetDefault<UUsdProjectSettings>())
-		{
-			if (UUsdAssetCache2* DefaultCache = Cast<UUsdAssetCache2>(ProjectSettings->DefaultAssetCache.TryLoad()))
-			{
-				UE_LOG(
-					LogUsd,
-					Log,
-					TEXT("USD Stage Actor '%s' had no previous USD Asset Cache, so it will use the default cache at '%s'. This can be configured on "
-						 "the project settings."),
-					*GetPathName(),
-					*DefaultCache->GetPathName()
-				);
-
-				UsdAssetCache = DefaultCache;
-			}
-		}
-	}
-
-#if WITH_EDITOR
-	// Show a dialog to let the user create a new default asset cache somewhere
-	if (!UsdAssetCache && GIsEditor && !IsRunningCommandlet() && !IsTemplate())
-	{
-		if (UUsdProjectSettings* ProjectSettings = GetMutableDefault<UUsdProjectSettings>())
-		{
-			if (ProjectSettings->bShowCreateDefaultAssetCacheDialog)
-			{
-				UUsdAssetCache2* NewCache = nullptr;
-
-				// Keep showing a dialog so that if the user clicks pick existing/create new and cancels, he can return to the dialog to pick
-				// something else
-				EDefaultAssetCacheDialogOption Outcome = EDefaultAssetCacheDialogOption::PickExisting;
-				while ((Outcome == EDefaultAssetCacheDialogOption::PickExisting || Outcome == EDefaultAssetCacheDialogOption::CreateNew) && !NewCache)
-				{
-					Outcome = IUsdClassesEditorModule::ShowMissingDefaultAssetCacheDialog(NewCache);
-				}
-
-				// We have an asset cache in some way: All good
-				if (NewCache)
-				{
-					ProjectSettings->DefaultAssetCache = NewCache;
-					ProjectSettings->SaveConfig();
-
-					UsdAssetCache = NewCache;
-
-					UE_LOG(
-						LogUsd,
-						Log,
-						TEXT("USD Stage Actor '%s' will use newly created, default USD Asset Cache at '%s'. This can be configured on the project "
-							 "settings."),
-						*GetPathName(),
-						*UsdAssetCache->GetPathName()
-					);
-				}
-				// Return right now if the user canceled, so we don't fall back to the case below of creating
-				// a temp/transient asset cache
-				else if (Outcome == EDefaultAssetCacheDialogOption::Cancel)
-				{
-					return;
-				}
-			}
-		}
-	}
-#endif
-
-	if (!UsdAssetCache)
-	{
-		UE_LOG(
-			LogUsd,
-			Warning,
-			TEXT("USD Stage Actor '%s' had no previous USD Asset Cache and no default cache is specified on the project settings, so a temporary "
-				 "cache will be generated. For better performance, create a persistent USD Asset Cache asset and point to it with this actor's "
-				 "UsdAssetCache property."),
-			*GetPathName()
-		);
-
-		UsdAssetCache = NewObject<UUsdAssetCache2>(GetTransientPackage(), NAME_None, GetMaskedFlags(RF_PropagateToSubObjects));
+		TGuardValue<ITransaction*> SuppressTransaction{GUndo, nullptr};
+		AssetCache = IUsdClassesModule::GetAssetCacheForProject();
 	}
 }
 
@@ -3534,7 +3486,7 @@ UUsdPrimTwin* AUsdStageActor::GetRootPrimTwin()
 		// Be careful not to give it a name, as there could be multiple of these on the transient package.
 		// It needs to be public or else FArchiveReplaceOrClearExternalReferences will reset our property
 		// whenever it is used from UEngine::CopyPropertiesForUnrelatedObjects for blueprint recompilation (if we're a blueprint class)
-		RootUsdTwin = NewObject<UUsdPrimTwin>(GetTransientPackage(), NAME_None, DefaultObjFlag | RF_Public);
+		RootUsdTwin = NewObject<UUsdPrimTwin>(GetTransientPackage(), NAME_None, (DefaultObjFlag | RF_Public) & ~RF_Standalone);
 	}
 
 	return RootUsdTwin;
@@ -3909,12 +3861,6 @@ void AUsdStageActor::PostDuplicate(bool bDuplicateForPIE)
 	}
 	else
 	{
-		// Temporary asset caches aren't meant to be shared
-		if (UsdAssetCache && UsdAssetCache->GetOutermost() == GetTransientPackage())
-		{
-			UsdAssetCache = NewObject<UUsdAssetCache2>(GetTransientPackage(), NAME_None, GetMaskedFlags(RF_PropagateToSubObjects));
-		}
-
 		LoadUsdStage();
 	}
 }
@@ -3945,15 +3891,6 @@ void AUsdStageActor::Serialize(FArchive& Ar)
 		InfoCache->Serialize(Ar);
 		LevelSequenceHelper.Serialize(Ar);
 	}
-
-	// Make sure we clear the old deprecated cache if we still have one, as old scenes may be persisting textures
-	// in there
-	PRAGMA_DISABLE_DEPRECATION_WARNINGS
-	if (AssetCache && Ar.IsLoading())
-	{
-		AssetCache->Reset();
-	}
-	PRAGMA_ENABLE_DEPRECATION_WARNINGS
 }
 
 void AUsdStageActor::Destroyed()
@@ -4048,8 +3985,7 @@ void AUsdStageActor::PostRegisterAllComponents()
 	// This may fail if our stage happened to not spawn any components, actors or assets, but by that
 	// point "being loaded" doesn't really mean anything anyway
 	const bool bStageIsLoaded = GetBaseUsdStage()
-								&& ((RootUsdTwin && RootUsdTwin->GetSceneComponent() != nullptr)
-									|| (UsdAssetCache && UsdAssetCache->GetNumAssets() > 0));
+								&& ((RootUsdTwin && RootUsdTwin->GetSceneComponent() != nullptr) || (AssetCache && AssetCache->GetNumAssets() > 0));
 
 	// Blocks loading stage when going into PIE, if we already have something loaded (we'll want to duplicate stuff instead).
 	// We need to allow loading when going into PIE when we have nothing loaded yet because the MovieRenderQueue (or other callers)
@@ -4143,8 +4079,7 @@ void AUsdStageActor::UnregisterAllComponents(bool bForReregister)
 #endif	  // WITH_EDITOR
 
 	const bool bStageIsLoaded = GetBaseUsdStage()
-								&& ((RootUsdTwin && RootUsdTwin->GetSceneComponent() != nullptr)
-									|| (UsdAssetCache && UsdAssetCache->GetNumAssets() > 0));
+								&& ((RootUsdTwin && RootUsdTwin->GetSceneComponent() != nullptr) || (AssetCache && AssetCache->GetNumAssets() > 0));
 
 	UWorld* World = GetWorld();
 	if (bIsTransitioningIntoPIE && bStageIsLoaded && (!World || World->WorldType == EWorldType::PIE))
@@ -4450,7 +4385,7 @@ void AUsdStageActor::OnObjectPropertyChanged(UObject* ObjectBeingModified, FProp
 	// out to the relevant prims
 	if (UUsdAssetUserData* UserData = Cast<UUsdAssetUserData>(ObjectBeingModified))
 	{
-		if (UsdAssetCache->IsAssetOwnedByCache(ObjectBeingModified->GetOuter()->GetPathName()))
+		if (AssetCache->IsAssetTrackedByCache(ObjectBeingModified->GetOuter()->GetPathName()))
 		{
 			FScopedBlockNoticeListening BlockNoticeListening(this);
 			FUsdStageActorImpl::WriteOutAssetMetadataChange(this, UserData, PropertyChangedEvent);
@@ -4905,12 +4840,21 @@ void AUsdStageActor::HandlePropertyChangedEvent(FPropertyChangedEvent& PropertyC
 		MetadataOptions.bInvertFilters = !bCorrectInvertValue;
 		SetInvertFilters(bCorrectInvertValue);
 	}
+	PRAGMA_DISABLE_DEPRECATION_WARNINGS
 	else if (PropertyName == GET_MEMBER_NAME_CHECKED(AUsdStageActor, UsdAssetCache))
 	{
-		UUsdAssetCache2* const CorrectCache = UsdAssetCache.Get();
+		UUsdAssetCache2* CorrectCache = UsdAssetCache.Get();
 		UsdAssetCache = UsdAssetCache ? nullptr : NewObject<UUsdAssetCache2>();
 		SetAssetCache(CorrectCache);
 	}
+	PRAGMA_ENABLE_DEPRECATION_WARNINGS
+	else if (PropertyName == GET_MEMBER_NAME_CHECKED(AUsdStageActor, AssetCache))
+	{
+		UUsdAssetCache3* CorrectCache = AssetCache.Get();
+		AssetCache = AssetCache ? nullptr : NewObject<UUsdAssetCache3>();
+
+		SetUsdAssetCache(CorrectCache);
+	}
 
 	bIsModifyingAProperty = false;
 }
@@ -4992,7 +4936,7 @@ bool AUsdStageActor::UnloadAssets(const UE::FSdfPath& StartPrimPath, bool bForEn
 	// be tracked via the info cache asset prim links, so if at any time they resync the old assets will still be
 	// found below when iterating the prim links, and we will discard them either way.
 
-	if (!UsdAssetCache || !InfoCache)
+	if (!AssetCache || !InfoCache)
 	{
 		return false;
 	}
@@ -5037,10 +4981,16 @@ bool AUsdStageActor::UnloadAssets(const UE::FSdfPath& StartPrimPath, bool bForEn
 				continue;
 			}
 
-			bAssetsDiscarded |= UsdAssetCache->RemoveAssetReference(OldAsset.Get(), this);
+			// If we're going to delete, just remove our reference but keep tracking: We need to be tracking an
+			// unreferenced asset in order to be able to delete it.
+			// Note: We could make it so that we fully stop tracking the asset if we're not going to delete it,
+			// although that doesn't really get us anything
+			bAssetsDiscarded |= AssetCache->RemoveAssetReferencer(OldAsset.Get(), this);
 		}
 	}
 
+	AssetCache->RequestDelayedAssetAutoCleanup();
+
 	return bAssetsDiscarded;
 }
 
@@ -5048,42 +4998,46 @@ bool AUsdStageActor::LoadAsset(FUsdSchemaTranslationContext& TranslationContext,
 {
 	TRACE_CPUPROFILER_EVENT_SCOPE(AUsdStageActor::LoadAsset);
 
-	if (!UsdAssetCache || !InfoCache)
+	if (!AssetCache || !InfoCache)
 	{
 		return false;
 	}
 
-	// Mark the assets as non transactional so that they don't get serialized in the transaction buffer
-	TGuardValue<EObjectFlags> ContextFlagsGuard(TranslationContext.ObjectFlags, TranslationContext.ObjectFlags & ~RF_Transactional);
-
-	int32 StartNumAssets = UsdAssetCache->GetNumAssets();
+	int32 StartNumAssets = AssetCache->GetNumAssets();
 
-	IUsdSchemasModule& UsdSchemasModule = FModuleManager::Get().LoadModuleChecked<IUsdSchemasModule>(TEXT("USDSchemas"));
-	if (TSharedPtr<FUsdSchemaTranslator> SchemaTranslator = UsdSchemasModule.GetTranslatorRegistry()
-																.CreateTranslatorForSchema(TranslationContext.AsShared(), UE::FUsdTyped(Prim)))
+	// Suppress transaction while we're creating assets.
+	// c.f. the big comment on the analogous position within AUsdStageActor::LoadAssets
+	Modify();
+	AssetCache->Modify();
 	{
-		TRACE_CPUPROFILER_EVENT_SCOPE(AUsdStageActor::CreateAssetsForPrim);
-		SchemaTranslator->CreateAssets();
-	}
+		TGuardValue<ITransaction*> SuppressTransaction{GUndo, nullptr};
+
+		IUsdSchemasModule& UsdSchemasModule = FModuleManager::Get().LoadModuleChecked<IUsdSchemasModule>(TEXT("USDSchemas"));
+		if (TSharedPtr<FUsdSchemaTranslator> SchemaTranslator = UsdSchemasModule.GetTranslatorRegistry().CreateTranslatorForSchema(	   //
+				TranslationContext.AsShared(),
+				UE::FUsdTyped(Prim)
+			))
+		{
+			TRACE_CPUPROFILER_EVENT_SCOPE(AUsdStageActor::CreateAssetsForPrim);
+			SchemaTranslator->CreateAssets();
+		}
 
-	TranslationContext.CompleteTasks();	   // Finish the asset tasks before moving on
+		TranslationContext.CompleteTasks();	   // Finish the asset tasks before moving on
+	}
 
-	return UsdAssetCache->GetNumAssets() != StartNumAssets;
+	return AssetCache->GetNumAssets() != StartNumAssets;
 }
 
 bool AUsdStageActor::LoadAssets(FUsdSchemaTranslationContext& TranslationContext, const UE::FUsdPrim& StartPrim)
 {
 	TRACE_CPUPROFILER_EVENT_SCOPE(AUsdStageActor::LoadAssets);
 
-	if (!UsdAssetCache || !InfoCache)
+	if (!AssetCache || !InfoCache)
 	{
 		return false;
 	}
 
-	// Mark the assets as non transactional so that they don't get serialized in the transaction buffer
-	TGuardValue<EObjectFlags> ContextFlagsGuard(TranslationContext.ObjectFlags, TranslationContext.ObjectFlags & ~RF_Transactional);
-
-	int32 StartNumAssets = UsdAssetCache->GetNumAssets();
+	int32 StartNumAssets = AssetCache->GetNumAssets();
 
 	IUsdSchemasModule& UsdSchemasModule = FModuleManager::Get().LoadModuleChecked<IUsdSchemasModule>(TEXT("USDSchemas"));
 
@@ -5119,21 +5073,38 @@ bool AUsdStageActor::LoadAssets(FUsdSchemaTranslationContext& TranslationContext
 		return false;
 	};
 
-	// Load materials first since meshes are referencing them
-	TArray<UE::FUsdPrim> AllPrimAssets = UsdUtils::GetAllPrimsOfType(StartPrim, TEXT("UsdShadeMaterial"));
+	// Suppress current transaction, as we never want assets to be put into the transaction buffer. This because these
+	// will be exposed to the content browser now, so that "opening the stage" essentially acts as a full import.
+	// We don't want to rip these assets up when pressing undo after they've been created. The engine should act
+	// essentially as if they've always been there.
+	//
+	// Note that we have tried achieving that by just creating these assets without the RF_Transactional flag, but that is not enough:
+	// Some assets create subobjects that are transactional anyway (StaticMeshes), and some other assets have much
+	// more complicated logic that can even spawn some transient Worlds, actors and components, and can put them all
+	// into the transaction buffer (Skeletal assets), causing havoc if we try to make sense of object referencers
+	// when it's time to clean up the asset.
+	//
+	Modify();	 // Mainly to dirty the info cache
+	AssetCache->Modify();
 	{
-		FScopedSlowTask MaterialsProgress(AllPrimAssets.Num(), LOCTEXT("CreateMaterials", "Creating materials"));
-		CreateAssetsForPrims(AllPrimAssets, MaterialsProgress);
-	}
+		TGuardValue<ITransaction*> SuppressTransaction{GUndo, nullptr};
 
-	// Load everything else (including meshes)
-	AllPrimAssets = UsdUtils::GetAllPrimsOfType(StartPrim, TEXT("UsdSchemaBase"), PruneChildren, {TEXT("UsdShadeMaterial")});
-	{
-		FScopedSlowTask AssetsProgress(AllPrimAssets.Num(), LOCTEXT("CreateAssets", "Creating assets"));
-		CreateAssetsForPrims(AllPrimAssets, AssetsProgress);
+		// Load materials first since meshes are referencing them
+		TArray<UE::FUsdPrim> AllPrimAssets = UsdUtils::GetAllPrimsOfType(StartPrim, TEXT("UsdShadeMaterial"));
+		{
+			FScopedSlowTask MaterialsProgress(AllPrimAssets.Num(), LOCTEXT("CreateMaterials", "Creating materials"));
+			CreateAssetsForPrims(AllPrimAssets, MaterialsProgress);
+		}
+
+		// Load everything else (including meshes)
+		AllPrimAssets = UsdUtils::GetAllPrimsOfType(StartPrim, TEXT("UsdSchemaBase"), PruneChildren, {TEXT("UsdShadeMaterial")});
+		{
+			FScopedSlowTask AssetsProgress(AllPrimAssets.Num(), LOCTEXT("CreateAssets", "Creating assets"));
+			CreateAssetsForPrims(AllPrimAssets, AssetsProgress);
+		}
 	}
 
-	return UsdAssetCache->GetNumAssets() != StartNumAssets;
+	return AssetCache->GetNumAssets() != StartNumAssets;
 }
 
 void AUsdStageActor::AnimatePrims()
diff --git a/Engine/Plugins/Importers/USDImporter/Source/USDStage/Public/USDStageActor.h b/Engine/Plugins/Importers/USDImporter/Source/USDStage/Public/USDStageActor.h
--- a/Engine/Plugins/Importers/USDImporter/Source/USDStage/Public/USDStageActor.h
+++ b/Engine/Plugins/Importers/USDImporter/Source/USDStage/Public/USDStageActor.h
@@ -17,8 +17,8 @@
 class FUsdInfoCache;
 class ISequencer;
 class ULevelSequence;
-class UUsdAssetCache;
 class UUsdAssetCache2;
+class UUsdAssetCache3;
 class UUsdPrimTwin;
 class UUsdTransactor;
 enum class EMovieSceneDataChangeType;
@@ -53,6 +53,9 @@ public:
 	EUsdStageState StageState;
 
 	UPROPERTY(EditAnywhere, BlueprintReadOnly, Category = "USD")
+	TObjectPtr<UUsdAssetCache3> AssetCache;
+
+	UPROPERTY()
 	TObjectPtr<UUsdAssetCache2> UsdAssetCache;
 
 	// These properties are configs so that spawned actors read them from the CDO when spawned.
@@ -154,9 +157,13 @@ public:
 	UFUNCTION(BlueprintCallable, Category = "USD", meta = (CallInEditor = "true"))
 	USDSTAGE_API void SetStageState(EUsdStageState NewStageState);
 
+	UE_DEPRECATED(5.5, "The UUsdAssetCache2 class has been deprecated. Use UUsdAssetCache3 and call SetUsdAssetCache instead")
 	UFUNCTION(BlueprintCallable, Category = "USD", meta = (CallInEditor = "true"))
 	USDSTAGE_API void SetAssetCache(UUsdAssetCache2* NewCache);
 
+	UFUNCTION(BlueprintCallable, Category = "USD", meta = (CallInEditor = "true"))
+	USDSTAGE_API void SetUsdAssetCache(UUsdAssetCache3* NewCache);
+
 	UFUNCTION(BlueprintCallable, Category = "USD", meta = (CallInEditor = "true"))
 	USDSTAGE_API void SetInitialLoadSet(EUsdInitialLoadSet NewLoadSet);
 
@@ -411,13 +418,6 @@ protected:
 	UPROPERTY(Transient)
 	TMap<TObjectPtr<UObject>, FString> ObjectsToWatch;
 
-	UE_DEPRECATED(5.2, "Use the new AssetCache property instead, that uses UUsdAssetCache2 objects")
-	UPROPERTY(
-		AdvancedDisplay,
-		meta = (DeprecatedProperty, DeprecationMessage = "Use the new AssetCache property instead, that uses UUsdAssetCache2 objects")
-	)
-	TObjectPtr<UUsdAssetCache> AssetCache;
-
 	UPROPERTY(Transient)
 	TObjectPtr<UUsdTransactor> Transactor;
 
diff --git a/Engine/Plugins/Importers/USDImporter/Source/USDStageEditor/Private/USDStageEditorModule.cpp b/Engine/Plugins/Importers/USDImporter/Source/USDStageEditor/Private/USDStageEditorModule.cpp
--- a/Engine/Plugins/Importers/USDImporter/Source/USDStageEditor/Private/USDStageEditorModule.cpp
+++ b/Engine/Plugins/Importers/USDImporter/Source/USDStageEditor/Private/USDStageEditorModule.cpp
@@ -350,9 +350,11 @@ public:
 					}
 				}
 
+				TWeakObjectPtr<AUsdStageActor> WeakActor = StageActor;
 				StageActor->OnPreStageChanged.AddLambda(
-					[this, StageActor]()
+					[this, WeakActor]()
 					{
+						AUsdStageActor* StageActor = WeakActor.Get();
 						if (!bUndoRedoing && StageActor && static_cast<const AUsdStageActor*>(StageActor)->GetUsdStage())
 						{
 							const bool bForClosing = true;
diff --git a/Engine/Plugins/Importers/USDImporter/Source/USDStageEditorViewModels/Private/USDStageViewModel.cpp b/Engine/Plugins/Importers/USDImporter/Source/USDStageEditorViewModels/Private/USDStageViewModel.cpp
--- a/Engine/Plugins/Importers/USDImporter/Source/USDStageEditorViewModels/Private/USDStageViewModel.cpp
+++ b/Engine/Plugins/Importers/USDImporter/Source/USDStageEditorViewModels/Private/USDStageViewModel.cpp
@@ -3,7 +3,7 @@
 #include "USDStageViewModel.h"
 
 #include "UnrealUSDWrapper.h"
-#include "USDAssetCache2.h"
+#include "USDAssetCache3.h"
 #include "USDClassesModule.h"
 #include "USDConversionUtils.h"
 #include "USDErrorUtils.h"
@@ -314,6 +314,8 @@ void FUsdStageViewModel::ImportStage(const TCHAR* TargetContentFolder, UUsdStage
 		FUsdStageImportContext ImportContext;
 
 		// Preload some settings according to USDStage options. These will overwrite whatever is loaded from config
+		ImportContext.ImportOptions->ExistingAssetCache = StageActor->AssetCache;
+		ImportContext.ImportOptions->bUseExistingAssetCache = true;
 		ImportContext.ImportOptions->PurposesToImport = StageActor->PurposesToLoad;
 		ImportContext.ImportOptions->RenderContextToImport = StageActor->RenderContext;
 		ImportContext.ImportOptions->MaterialPurpose = StageActor->MaterialPurpose;
@@ -355,21 +357,16 @@ void FUsdStageViewModel::ImportStage(const TCHAR* TargetContentFolder, UUsdStage
 				ImportContext.ImportOptions = Options;
 			}
 
-			// Let the importer reuse our assets, but force it to spawn new actors and components always
-			// This allows a different setting for asset/component collapsing, and doesn't require modifying the PrimTwins
-			ImportContext.AssetCache = StageActor->UsdAssetCache;
+			// Pick the asset cache that the user potentially changed on the import options dialog. The USDStageImporter will sort itself out
+			// if that happens to be nullptr/invalid
+			ImportContext.UsdAssetCache = ImportContext.ImportOptions->bUseExistingAssetCache
+											  ? Cast<UUsdAssetCache3>(ImportContext.ImportOptions->ExistingAssetCache.TryLoad())
+											  : nullptr;
 			ImportContext.BBoxCache = StageActor->GetBBoxCache();
 
 			ImportContext.TargetSceneActorAttachParent = StageActor->GetRootComponent()->GetAttachParent();
 			ImportContext.TargetSceneActorTargetTransform = StageActor->GetActorTransform();
 
-			// Preemptively remove the stage actor as a user of the assets on the cache so that the stage importer
-			// can just take the assets from the cache directly. Otherwise it will be forced to duplicate them
-			if (ImportContext.AssetCache)
-			{
-				ImportContext.AssetCache->RemoveAllAssetReferences(StageActor);
-			}
-
 			UUsdStageImporter* USDImporter = IUsdStageImporterModule::Get().GetImporter();
 			USDImporter->ImportFromFile(ImportContext);
 
diff --git a/Engine/Plugins/Importers/USDImporter/Source/USDStageImporter/Private/USDStageAssetImportFactory.cpp b/Engine/Plugins/Importers/USDImporter/Source/USDStageImporter/Private/USDStageAssetImportFactory.cpp
--- a/Engine/Plugins/Importers/USDImporter/Source/USDStageImporter/Private/USDStageAssetImportFactory.cpp
+++ b/Engine/Plugins/Importers/USDImporter/Source/USDStageImporter/Private/USDStageAssetImportFactory.cpp
@@ -6,6 +6,7 @@
 #include "USDAssetUserData.h"
 #include "USDConversionUtils.h"
 #include "USDErrorUtils.h"
+#include "USDObjectUtils.h"
 #include "USDStageImporter.h"
 #include "USDStageImporterModule.h"
 #include "USDStageImportOptions.h"
@@ -115,7 +116,7 @@ void UUsdStageAssetImportFactory::CleanUp()
 
 bool UUsdStageAssetImportFactory::CanReimport(UObject* Obj, TArray<FString>& OutFilenames)
 {
-	if (UAssetImportData* ImportData = UsdUtils::GetAssetImportData(Obj))
+	if (UAssetImportData* ImportData = UsdUnreal::ObjectUtils::GetAssetImportData(Obj))
 	{
 		const FString FileName = ImportData->GetFirstFilename();
 		const FString FileExtension = FPaths::GetExtension(FileName);
@@ -143,7 +144,7 @@ void UUsdStageAssetImportFactory::SetReimportPaths(UObject* Obj, const TArray<FS
 		return;
 	}
 
-	if (UUsdAssetImportData* ImportData = UsdUtils::GetAssetImportData(Obj))
+	if (UUsdAssetImportData* ImportData = UsdUnreal::ObjectUtils::GetAssetImportData(Obj))
 	{
 		ImportData->UpdateFilenameOnly(NewReimportPaths[0]);
 	}
@@ -172,7 +173,7 @@ EReimportResult::Type UUsdStageAssetImportFactory::Reimport(UObject* Obj)
 		}
 	}
 
-	if (UUsdAssetImportData* ImportData = UsdUtils::GetAssetImportData(Obj))
+	if (UUsdAssetImportData* ImportData = UsdUnreal::ObjectUtils::GetAssetImportData(Obj))
 	{
 		ReimportFilePath = ImportData->GetFirstFilename();
 		ReimportOptions = ImportData->ImportOptions;
diff --git a/Engine/Plugins/Importers/USDImporter/Source/USDStageImporter/Private/USDStageImportContext.cpp b/Engine/Plugins/Importers/USDImporter/Source/USDStageImporter/Private/USDStageImportContext.cpp
--- a/Engine/Plugins/Importers/USDImporter/Source/USDStageImporter/Private/USDStageImportContext.cpp
+++ b/Engine/Plugins/Importers/USDImporter/Source/USDStageImporter/Private/USDStageImportContext.cpp
@@ -27,7 +27,7 @@ bool FUsdStageImportContext::Init(
 	ObjectName = InName;
 	FilePath = InFilePath;
 	bIsAutomated = bInIsAutomated;
-	ImportObjectFlags = InFlags | RF_Transactional;
+	ImportObjectFlags = InFlags;
 	World = GEditor->GetEditorWorldContext().World();
 	PackagePath = InInitialPackagePath;
 
@@ -92,7 +92,10 @@ void FUsdStageImportContext::Reset()
 	ImportOptions = NewObject<UUsdStageImportOptions>();
 	ImportedAsset = nullptr;
 	LevelSequenceHelper.Clear();
+	UsdAssetCache = nullptr;
+	PRAGMA_DISABLE_DEPRECATION_WARNINGS
 	AssetCache = nullptr;
+	PRAGMA_ENABLE_DEPRECATION_WARNINGS
 	BBoxCache = nullptr;
 	MaterialToPrimvarToUVIndex.Empty();
 	Stage = UE::FUsdStage{};
diff --git a/Engine/Plugins/Importers/USDImporter/Source/USDStageImporter/Private/USDStageImportOptions.cpp b/Engine/Plugins/Importers/USDImporter/Source/USDStageImporter/Private/USDStageImportOptions.cpp
--- a/Engine/Plugins/Importers/USDImporter/Source/USDStageImporter/Private/USDStageImportOptions.cpp
+++ b/Engine/Plugins/Importers/USDImporter/Source/USDStageImporter/Private/USDStageImportOptions.cpp
@@ -21,6 +21,8 @@ UUsdStageImportOptions::UUsdStageImportOptions(const FObjectInitializer& ObjectI
 	bImportSparseVolumeTextures = true;
 	bImportOnlyUsedMaterials = false;
 
+	bUseExistingAssetCache = false;
+	ExistingAssetCache = nullptr;
 	PurposesToImport = (int32)(EUsdPurpose::Default | EUsdPurpose::Proxy | EUsdPurpose::Render | EUsdPurpose::Guide);
 	NaniteTriangleThreshold = INT32_MAX;
 	IUsdSchemasModule& UsdSchemasModule = FModuleManager::Get().LoadModuleChecked<IUsdSchemasModule>(TEXT("USDSchemas"));
@@ -100,6 +102,8 @@ void UsdUtils::AddAnalyticsAttributes(const UUsdStageImportOptions& Options, TAr
 	{
 		InOutAttributes.Emplace(TEXT("NumPrimsToImport"), LexToString(Options.PrimsToImport.Num()));
 	}
+
+	InOutAttributes.Emplace(TEXT("bUseExistingAssetCache"), Options.bUseExistingAssetCache);
 	InOutAttributes.Emplace(TEXT("PurposesToImport"), LexToString(Options.PurposesToImport));
 	InOutAttributes.Emplace(TEXT("NaniteTriangleThreshold"), LexToString(Options.NaniteTriangleThreshold));
 	InOutAttributes.Emplace(TEXT("RenderContextToImport"), Options.RenderContextToImport.ToString());
