USD: Implements UsdAssetCache3 which should be a simpler and more flexible way of managing USD assets;

diff --git a/Engine/Plugins/Importers/USDImporter/Source/USDStageImporter/Private/USDStageImporter.cpp b/Engine/Plugins/Importers/USDImporter/Source/USDStageImporter/Private/USDStageImporter.cpp
--- a/Engine/Plugins/Importers/USDImporter/Source/USDStageImporter/Private/USDStageImporter.cpp
+++ b/Engine/Plugins/Importers/USDImporter/Source/USDStageImporter/Private/USDStageImporter.cpp
@@ -2,7 +2,7 @@
 
 #include "USDStageImporter.h"
 
-#include "USDAssetCache2.h"
+#include "USDAssetCache3.h"
 #include "USDAssetImportData.h"
 #include "USDAssetUserData.h"
 #include "USDClassesModule.h"
@@ -11,6 +11,7 @@
 #include "USDErrorUtils.h"
 #include "USDGeomMeshConversion.h"
 #include "USDLog.h"
+#include "USDObjectUtils.h"
 #include "USDPrimTwin.h"
 #include "USDSchemasModule.h"
 #include "USDSchemaTranslator.h"
@@ -50,6 +51,7 @@
 #include "Materials/MaterialInstance.h"
 #include "Materials/MaterialInterface.h"
 #include "Misc/Paths.h"
+#include "Misc/RedirectCollector.h"
 #include "Misc/ScopedSlowTask.h"
 #include "Modules/ModuleManager.h"
 #include "PackageTools.h"
@@ -198,7 +200,7 @@ namespace UE::USDStageImporter::Private
 		// We always spawn another scene actor regardless of collision or whether the level already has one,
 		// so that we can fully build our hierarchy separately before resolving collisions according to ExistingActorPolicy
 		AActor* Actor = ImportContext.World->SpawnActor(AActor::StaticClass(), nullptr, SpawnParameters);
-		Actor->SetActorLabel(IUsdClassesModule::SanitizeObjectName(ImportContext.ObjectName));
+		Actor->SetActorLabel(UsdUnreal::ObjectUtils::SanitizeObjectName(ImportContext.ObjectName));
 
 		USceneComponent* RootComponent = Actor->GetRootComponent();
 		if (!RootComponent)
@@ -229,7 +231,7 @@ namespace UE::USDStageImporter::Private
 	AActor* GetExistingSceneActor(FUsdStageImportContext& ImportContext)
 	{
 		// We always reuse the existing scene actor for a scene, regardless of ReplacePolicy
-		FString TargetActorLabel = IUsdClassesModule::SanitizeObjectName(ImportContext.ObjectName);
+		FString TargetActorLabel = UsdUnreal::ObjectUtils::SanitizeObjectName(ImportContext.ObjectName);
 		for (TActorIterator<AActor> ActorItr(ImportContext.World); ActorItr; ++ActorItr)
 		{
 			AActor* ThisActor = *ActorItr;
@@ -461,9 +463,6 @@ namespace UE::USDStageImporter::Private
 		ImportActor(ImportContext, RootPrim, bForceVisibilityAnimationTracks, TranslationContext);
 	}
 
-	// Assets coming out of USDSchemas module have default names, so here we do our best to provide them with
-	// names based on the source prims. This is likely a temporary solution, as it may be interesting to do this in the
-	// USDSchemas module itself
 	FString GetUserFriendlyName(UObject* Asset, TSet<FString>& UniqueAssetNames)
 	{
 		if (!Asset)
@@ -471,9 +470,7 @@ namespace UE::USDStageImporter::Private
 			return {};
 		}
 
-		FString AssetPrefix;
-
-		FString AssetPath = Asset->GetName();
+		FString AssetName = Asset->GetFName().GetPlainNameString();
 
 		FString PrimPath;
 		if (IInterface_AssetUserData* UserDataInterface = Cast<IInterface_AssetUserData>(Asset))
@@ -483,154 +480,35 @@ namespace UE::USDStageImporter::Private
 				if (!UserData->PrimPaths.IsEmpty())
 				{
 					PrimPath = UserData->PrimPaths[0];
-					AssetPath = PrimPath;
 				}
 			}
 		}
 
-		FString AssetName = FPaths::GetBaseFilename(AssetPath);
-
 		if (UStaticMesh* Mesh = Cast<UStaticMesh>(Asset))
 		{
-			AssetPrefix = TEXT("SM_");
-
 			// If we have multiple LODs here we must have parsed the LOD variant set pattern. If our prims were named
 			// with the LOD pattern, go from e.g. '/Root/MyMesh/LOD0' to '/Root/MyMesh', or else every single LOD mesh
 			// will be named "SM_LOD0_X". We'll actually check though because if the user set a custom name for their
 			// prim other than LOD0 then we'll keep that
 			if (Mesh->GetNumLODs() > 1)
 			{
-				FString PrimName = FPaths::GetBaseFilename(AssetPath);
+				FString PrimName = FPaths::GetBaseFilename(PrimPath);
 				if (PrimName.RemoveFromStart(TEXT("LOD"), ESearchCase::CaseSensitive))
 				{
 					if (PrimName.IsNumeric())
 					{
-						AssetPath = FPaths::GetPath(AssetPath);
-						AssetName = FPaths::GetBaseFilename(AssetPath);
+						AssetName = FPaths::GetBaseFilename(FPaths::GetPath(PrimPath));
 					}
 				}
 			}
 		}
-		else if (Asset->IsA<UGroomAsset>() || Asset->IsA<UGroomCache>() || Asset->IsA<UGroomBindingAsset>())
-		{
-			// Keep the groom assets named like they originally are because they are not just named after their
-			// prim path, but may have additional suffixes like "_stands_cache" and "_groombinding" that we
-			// can't get from the prim path
-			AssetName = Asset->GetFName().GetPlainNameString();
-		}
-		else if (USkeletalMesh* SkMesh = Cast<USkeletalMesh>(Asset))
-		{
-			AssetPrefix = TEXT("SK_");
-
-			// Our SkeletalMesh asset is now assigned to the Skeleton prim path, but we don't want our actual USkeletalMesh
-			// to be named "Skeleton" or something like that. Let's keep naming it after the SkelRoot instead (which the
-			// UAsset will already be named after anyway)
-			AssetName = Asset->GetFName().GetPlainNameString();
-		}
-		else if (USkeleton* Skeleton = Cast<USkeleton>(Asset))
-		{
-			// Skeletons don't have asset import data, so we can't store the prim that originated them, but we do
-			// name the actual assets after the prims, so unlike for the other asset types we want to use the actual
-			// asset name here
-			AssetName = Asset->GetFName().GetPlainNameString();
-			AssetPrefix = TEXT("SKEL_");
-		}
-		else if (UPhysicsAsset* PhysicsAsset = Cast<UPhysicsAsset>(Asset))
-		{
-			// See comments above on the case for USkeleton
-			FString TempName = Asset->GetFName().GetPlainNameString();
-
-			// The asset is named after the SkelRoot prim. If we're importing back a scene that was originally exported,
-			// we should clean up these prefixes or else we may end up with something like "PHYS_SK_PrimName"
-			TempName.RemoveFromStart(TEXT("PHYS_"), ESearchCase::CaseSensitive);
-			TempName.RemoveFromStart(TEXT("SK_"), ESearchCase::CaseSensitive);
-			if (!TempName.IsEmpty())
-			{
-				AssetName = TempName;
-			}
-
-			AssetPrefix = TEXT("PHYS_");
-		}
-		else if (UAnimSequence* AnimSequence = Cast<UAnimSequence>(Asset))
-		{
-			AssetPrefix = TEXT("AS_");
-		}
-		else if (UMaterialInterface* Material = Cast<UMaterialInterface>(Asset))
-		{
-			if (Material->IsA<UMaterialInstance>())
-			{
-				AssetPrefix = TEXT("MI_");
-			}
-			else
-			{
-				AssetPrefix = TEXT("M_");
-			}
-
-			// The only materials with no prim path are our auto-generated displayColor materials
-			if (PrimPath.IsEmpty())
-			{
-				AssetPath = TEXT("DisplayColor");
-				AssetName = FPaths::GetBaseFilename(AssetPath);
-			}
-			else
-			{
-				AssetPath = PrimPath;
-				AssetName = FPaths::GetBaseFilename(AssetPath);
-
-				// If we have a preview surface two-sided material we'll also have a one-sided with the same name,
-				// so add a suffix here so we can clearly tell which is which
-				if (Material->IsTwoSided())
-				{
-					AssetName += UnrealIdentifiers::TwoSidedMaterialSuffix;
-				}
-			}
-		}
-		else if (UTexture* Texture = Cast<UTexture>(Asset))
-		{
-			AssetPrefix = TEXT("T_");
-
-			// Use the actual asset name as it matches the original texture filename, which is more useful
-			// than ending up with "T_<MaterialName>" instead (the textures receive the material prim paths).
-			// We could also fetch it from AssetImportData but then we'd have to process the path a bit in case of
-			// textures inside USDZ files. It's just easier to take from the asset name itself instead
-			AssetName = Texture->GetFName().GetPlainNameString();
-		}
-		else if (ULevelSequence* LevelSequence = Cast<ULevelSequence>(Asset))
-		{
-			AssetPrefix = TEXT("LS_");
-		}
-		else if (UAnimBlueprint* AnimBP = Cast<UAnimBlueprint>(Asset))
-		{
-			FString TempName = AssetName;
-
-			// The asset is named after the SkelRoot prim. If we're importing back a scene that was originally exported,
-			// we should clean up these prefixes or else we may end up with something like "ABP_SK_PrimName"
-			TempName.RemoveFromStart(TEXT("ABP_"), ESearchCase::CaseSensitive);
-			TempName.RemoveFromStart(TEXT("SK_"), ESearchCase::CaseSensitive);
-			if (!TempName.IsEmpty())
-			{
-				AssetName = TempName;
-			}
-
-			AssetPrefix = TEXT("ABP_");
-		}
-		else if (USparseVolumeTexture* SparseVolumeTexture = Cast<USparseVolumeTexture>(Asset))
-		{
-			AssetPrefix = TEXT("SVT_");
 
-			// Same situation as the UTexture case
-			AssetName = SparseVolumeTexture->GetFName().GetPlainNameString();
-		}
-
-		if (!AssetName.StartsWith(AssetPrefix))
-		{
-			AssetName = AssetPrefix + AssetName;
-		}
+		AssetName = UsdUnreal::ObjectUtils::GetPrefixedAssetName(AssetName, Asset->GetClass());
 
 		// We don't care if our assets overwrite something in the final destination package (that conflict will be
 		// handled according to EReplaceAssetPolicy). But we do want these assets to have unique names amongst themselves
 		// or else they will overwrite each other when publishing
-		AssetName = UsdUtils::GetUniqueName(IUsdClassesModule::SanitizeObjectName(AssetName), UniqueAssetNames);
+		AssetName = UsdUnreal::ObjectUtils::GetUniqueName(UsdUnreal::ObjectUtils::SanitizeObjectName(AssetName), UniqueAssetNames);
 		UniqueAssetNames.Add(AssetName);
 
 		return AssetName;
@@ -643,11 +521,11 @@ namespace UE::USDStageImporter::Private
 			return;
 		}
 
-		UUsdAssetImportData* ImportData = UsdUtils::GetAssetImportData(Asset);
+		UUsdAssetImportData* ImportData = UsdUnreal::ObjectUtils::GetAssetImportData(Asset);
 		if (!ImportData)
 		{
 			ImportData = NewObject<UUsdAssetImportData>(Asset);
-			UsdUtils::SetAssetImportData(Asset, ImportData);
+			UsdUnreal::ObjectUtils::SetAssetImportData(Asset, ImportData);
 		}
 
 		if (ImportData)
@@ -728,7 +606,7 @@ namespace UE::USDStageImporter::Private
 		// Close editors opened on existing asset if applicable
 		bool bAssetWasOpen = false;
 		UAssetEditorSubsystem* AssetEditorSubsystem = GEditor->GetEditorSubsystem<UAssetEditorSubsystem>();
-		if (ExistingAsset && AssetEditorSubsystem->FindEditorForAsset(ExistingAsset, false) != nullptr)
+		if (ExistingAsset && AssetEditorSubsystem->FindEditorForAsset(ExistingAsset, false) != nullptr && !IsEngineExitRequested())
 		{
 			AssetEditorSubsystem->CloseAllEditorsForAsset(ExistingAsset);
 			bAssetWasOpen = true;
@@ -751,41 +629,21 @@ namespace UE::USDStageImporter::Private
 
 		FSoftObjectPath OldPath = Asset;
 
-		// If the asset is currently being referenced on its asset cache we can't just take it, and must duplicate it
-		// below with the DuplicateObject call
-		bool bForceDuplicate = false;
-		if (UUsdAssetCache2* Cache = ImportContext.AssetCache.Get())
-		{
-			FString Hash = Cache->GetHashForAsset(Asset);
-			const bool bIsAssetOwnedByCache = !Hash.IsEmpty();
-			if (bIsAssetOwnedByCache)
-			{
-				bool bRemovedFromCache = Cache->CanRemoveAsset(Hash);
-				if (bRemovedFromCache)
-				{
-					UObject* RemovedAsset = Cache->RemoveAsset(Hash);
-					ensure(!RemovedAsset || RemovedAsset == Asset);
-
-					bRemovedFromCache = RemovedAsset != nullptr;
-				}
-
-				if (!bRemovedFromCache)
-				{
-					bForceDuplicate = true;
-					UE_LOG(
-						LogUsd,
-						Log,
-						TEXT("About to duplicate asset '%s' on import as it can't be taken directly from the asset cache"),
-						*Asset->GetPathName()
-					);
-				}
-			}
-		}
+		// If the asset cache is targetting the transient package (by being within the transient package and having its
+		// AssetDirectory point at it) and the asset is transient, it means the asset was generated for this import only
+		// (even if the asset cache is otherwise persistent), so we can just rename it.
+		// Note that the asset cache may also store transient MIDs when opening the stage at runtime, but we don't import at
+		// runtime anyway, so we're all good!
+		//
+		// We also duplicate LevelSequences when importing because they are technically owned by the LevelSequenceHelper, and
+		// we can't just take them
+		const bool bMustDuplicate = (ImportContext.UsdAssetCache->AssetDirectory.Path != GetTransientPackage()->GetPathName())
+									|| (Asset->GetOutermost() != GetTransientPackage()) || Asset->IsA<ULevelSequence>();
 
 		// Strategy copied from FDatasmithImporterImpl::PublicizeAsset
 		// Replace existing asset (reimport or conflict) with new asset
 		UObject* MovedAsset = ExistingAsset;
-		if (bForceDuplicate || (ExistingAsset != nullptr && ExistingAsset != Asset && ReplacePolicy == EReplaceAssetPolicy::Replace))
+		if (bMustDuplicate || (ExistingAsset != nullptr && ExistingAsset != Asset && ReplacePolicy == EReplaceAssetPolicy::Replace))
 		{
 			// We have to rename the existing asset away, because some objects manage subobjects (like
 			// UMaterialInterface and its EditorOnlyData), and try to rename them whenever they're duplicated/renamed.
@@ -797,7 +655,7 @@ namespace UE::USDStageImporter::Private
 					->Rename(*UniqueName.ToString(), GetTransientPackage(), REN_DontCreateRedirectors | REN_NonTransactional | REN_DoNotDirty);
 			}
 
-			if (bForceDuplicate)
+			if (bMustDuplicate)
 			{
 				MovedAsset = DuplicateObject<UObject>(Asset, Package, *TargetAssetName);
 			}
@@ -839,9 +697,10 @@ namespace UE::USDStageImporter::Private
 			ObjectsToRemap.Add(Asset, MovedAsset);
 		}
 
-		// Important as some assets (e.g. material instances) are created with no flags
+		// Important as some assets (e.g. material instances) are created with no flags, and LevelSequences are created
+		// without RF_Standalone by the USDLevelSequenceHelper
 		MovedAsset->SetFlags(ImportContext.ImportObjectFlags | EObjectFlags::RF_Public | EObjectFlags::RF_Standalone);
-		MovedAsset->ClearFlags(EObjectFlags::RF_Transient | EObjectFlags::RF_DuplicateTransient | EObjectFlags::RF_NonPIEDuplicateTransient);
+		MovedAsset->ClearFlags(EObjectFlags::RF_Transient);
 
 		// Some subobjects like UStaticMesh::HiResSourceModel->StaticMeshDescriptionBulkData can't be left transient, or else they won't serialize
 		// their data. We probably never want to make them public or standalone if they aren't already though
@@ -850,8 +709,8 @@ namespace UE::USDStageImporter::Private
 		if (UMaterialInterface* Material = Cast<UMaterialInterface>(MovedAsset))
 		{
 			// Materials in particular have EditorOnlyData which behaves like a default subobject but kind of isn't flagged as one...
-			// Not sure at this point whether we should force non-transient on all subobjects of the asset, as that could also have edge cases that
-			// don't do what we expect.
+			// Since we just cleared the RF_Transient flag, and RF_Transient belongs to RF_PropagateToSubObjects, let's also clear it
+			// from its subobjects
 			if (UMaterialInterfaceEditorOnlyData* EditorOnlyData = Material->GetEditorOnlyData())
 			{
 				Subobjects.Add(EditorOnlyData);
@@ -859,7 +718,7 @@ namespace UE::USDStageImporter::Private
 		}
 		for (UObject* Subobject : Subobjects)
 		{
-			Subobject->ClearFlags(EObjectFlags::RF_Transient | EObjectFlags::RF_DuplicateTransient | EObjectFlags::RF_NonPIEDuplicateTransient);
+			Subobject->ClearFlags(EObjectFlags::RF_Transient & RF_PropagateToSubObjects);
 		}
 
 		// We need to make sure that "dirtying the final package" is not added to the transaction, because if we undo this transaction
@@ -878,6 +737,14 @@ namespace UE::USDStageImporter::Private
 			AssetEditorSubsystem->OpenEditorForAsset(MovedAsset);
 		}
 
+		// We want the asset cache to stop tracking anything that has been moved to the import package, but it's
+		// important to let it track the original assets, so that it can delete those once the import is complete (if it's set up
+		// to do that)
+		if (MovedAsset->GetPathName().StartsWith(ImportContext.PackagePath))
+		{
+			ImportContext.UsdAssetCache->StopTrackingAsset(ImportContext.UsdAssetCache->GetHashForAsset(MovedAsset));
+		}
+
 		ImportContext.ImportedAsset = MovedAsset;
 
 		return MovedAsset;
@@ -916,6 +783,7 @@ namespace UE::USDStageImporter::Private
 		const TSet<UObject*>& AssetsToPublish,
 		TMap<UObject*, UObject*>& ObjectsToRemap,
 		TMap<FSoftObjectPath, FSoftObjectPath>& SoftObjectsToRemap,
+		TSet<UObject*>& OutPublishedAssets,
 		TSet<UObject*>& OutAssetsToFinalize
 	)
 	{
@@ -1015,14 +883,7 @@ namespace UE::USDStageImporter::Private
 			}
 			else if (UGroomBindingAsset* GroomBinding = Cast<UGroomBindingAsset>(Asset))
 			{
-				if (GroomBinding->GetTargetSkeletalMesh())
-				{
-					AssetTypeFolderPtr = &SkeletalMeshesFolder;
-				}
-				else if (GroomBinding->GetTargetGeometryCache())
-				{
-					AssetTypeFolderPtr = &GeometryCachesFolder;
-				}
+				AssetTypeFolderPtr = &GroomsFolder;
 				GroomBindings.Add(Asset);
 			}
 			else if (Asset->IsA(USparseVolumeTexture::StaticClass()))
@@ -1062,7 +923,14 @@ namespace UE::USDStageImporter::Private
 			}
 		}
 
-		TFunction<void(const TArray<UObject*>&)> PublishAssetType = [&](const TArray<UObject*>& Assets)
+		TFunction<void(const TArray<UObject*>&)> PublishAssetType = [&AssetToContentFolder,
+																	 &UniqueAssetNames,
+																	 &ImportContext,
+																	 &ObjectsToRemap,
+																	 &SoftObjectsToRemap,
+																	 &OutAssetsToFinalize,
+																	 &OutPublishedAssets	//
+		](const TArray<UObject*>& Assets)
 		{
 			for (UObject* Asset : Assets)
 			{
@@ -1076,6 +944,7 @@ namespace UE::USDStageImporter::Private
 					if (PublishedAsset)
 					{
 						PrunePrimMetadata(PublishedAsset, ImportContext.Stage);
+						OutPublishedAssets.Add(PublishedAsset);
 					}
 				}
 			}
@@ -1529,7 +1398,7 @@ namespace UE::USDStageImporter::Private
 	}
 
 	/**
-	 * UUsdAssetCache2 can track which assets are requested/added to itself during translation, but it may miss some dependencies
+	 * UUsdAssetCache3 can track which assets are requested/added to itself during translation, but it may miss some dependencies
 	 * that are only retrieved/added themselves when the original asset is first parsed. This function recursively collects all of those.
 	 * Example: An UMaterialInstance is already in the cache, so when translating we just retrieve the existing asset --> The textures that it's using
 	 * won't be retrieved or marked as "Used" Example: An USkeletalMesh is already in the cache, so in the same way we would miss its USkeleton,
@@ -1537,7 +1406,8 @@ namespace UE::USDStageImporter::Private
 	 */
 	void CollectUsedAssetDependencies(FUsdStageImportContext& ImportContext, TSet<UObject*>& OutAssetsAndDependencies)
 	{
-		const int32 ReserveSize = OutAssetsAndDependencies.Num() + (ImportContext.AssetCache ? ImportContext.AssetCache->GetActiveAssets().Num() : 0);
+		const int32 ReserveSize = OutAssetsAndDependencies.Num()
+								  + (ImportContext.UsdAssetCache ? ImportContext.UsdAssetCache->GetActiveAssets().Num() : 0);
 
 		// We will only emit the level sequences if we have data in the main one.
 		// Keep subsequences even if they have no data as the main sequence/other sequences may reference them
@@ -1558,18 +1428,18 @@ namespace UE::USDStageImporter::Private
 			OutAssetsAndDependencies.Reserve(ReserveSize);
 		}
 
-		if (ImportContext.AssetCache)
+		if (ImportContext.UsdAssetCache)
 		{
-			const TSet<UObject*>& InPrimaryAssets = ImportContext.AssetCache->GetActiveAssets();
-			TArray<UObject*> AssetQueue = InPrimaryAssets.Array();
+			const TSet<FSoftObjectPath>& InPrimaryAssets = ImportContext.UsdAssetCache->GetActiveAssets();
+			TArray<FSoftObjectPath> AssetQueue = InPrimaryAssets.Array();
 
 			for (int32 AssetIndex = 0; AssetIndex < AssetQueue.Num(); ++AssetIndex)
 			{
-				UObject* Asset = AssetQueue[AssetIndex];
+				UObject* Asset = AssetQueue[AssetIndex].TryLoad();
 
-				// Only add it as a dependency if it's owned by the asset cache, but still traverse it because
+				// Only add it as a dependency if it's tracked by the asset cache, but still traverse it because
 				// we may be in some strange situation where the material shouldn't be in this list, but one of its used textures should
-				if (Asset && ImportContext.AssetCache->IsAssetOwnedByCache(Asset->GetPathName()))
+				if (Asset && ImportContext.UsdAssetCache->IsAssetTrackedByCache(Asset->GetPathName()))
 				{
 					OutAssetsAndDependencies.Add(Asset);
 				}
@@ -1622,47 +1492,30 @@ namespace UE::USDStageImporter::Private
 
 		IAssetTools& AssetTools = FModuleManager::LoadModuleChecked<FAssetToolsModule>("AssetTools").Get();
 		AssetTools.RenameReferencingSoftObjectPaths(Packages.Array(), ProcessedPaths);
-	}
-
-	void GetPublishedAssetsAndDependencies(
-		const TSet<UObject*>& UsedAssetsAndDependencies,
-		const TMap<UObject*, UObject*>& ObjectsToRemap,
-		const TMap<FSoftObjectPath, FSoftObjectPath>& SoftObjectsToRemap,
-		TSet<UObject*>& OutPublishedAssetsAndDependencies
-	)
-	{
-		OutPublishedAssetsAndDependencies.Empty(UsedAssetsAndDependencies.Num());
 
-		for (UObject* Asset : UsedAssetsAndDependencies)
+		// Unexpectedly, the RenameReferencingSoftObjectPaths call above will also *globally* register a redirector from
+		// source to target paths, that will affect *all other properties in all other objects*! This means that if at any point
+		// any other UObject had a FSoftObjectPath property pointing at any of our OriginalPaths, the property would be
+		// secretly rewritten when the UObject was serialized!
+		//
+		// This will cause major havoc on the AssetCache (e.g. two sequential Actions->Import), so we need to remove this global
+		// remapping ASAP
+		TArray<FSoftObjectPath> OriginalPaths;
+		SoftObjectsToRemap.GetKeys(OriginalPaths);
+		for (const FSoftObjectPath& OriginalPath : OriginalPaths)
 		{
-			if (!Asset)
-			{
-				continue;
-			}
-
-			UObject* RemappedAsset = ObjectsToRemap.FindRef(Asset);
-
-			// The asset itself was renamed onto its final package, no remapping needed
-			if (!RemappedAsset && Asset->GetOutermost() != GetTransientPackage())
+			// We check first here because if we query RemoveAssetPathRedirection with a path that doesn't
+			// have a remapping it will ensure...
+			FSoftObjectPath RedirectedPath = GRedirectCollector.GetAssetPathRedirection(OriginalPath);
+			if (RedirectedPath.IsValid())
 			{
-				OutPublishedAssetsAndDependencies.Add(Asset);
-				continue;
+				GRedirectCollector.RemoveAssetPathRedirection(OriginalPath);
 			}
-
-			// If we have remapped it, ObjectsToRemap and SoftObjectsToRemap should have agreed on where to
-			const FSoftObjectPath* RemappedSoftAsset = SoftObjectsToRemap.Find(FSoftObjectPath{Asset});
-			if (!RemappedAsset || !RemappedSoftAsset || RemappedAsset != RemappedSoftAsset->TryLoad())
-			{
-				UE_LOG(LogUsd, Warning, TEXT("Failed to publish or remap asset '%s'!"), *Asset->GetPathName());
-				continue;
-			}
-
-			OutPublishedAssetsAndDependencies.Add(RemappedAsset);
 		}
 	}
 
 	/** After we remapped everything, notify the AssetRegistry that we created some new assets */
-	void NotifyAssetRegistry(const TSet<UObject*>& UsedAssetsAndDependencies)
+	void BroadcastImportEvents(const TSet<UObject*>& UsedAssetsAndDependencies, bool bIsReimport = false)
 	{
 		for (UObject* Object : UsedAssetsAndDependencies)
 		{
@@ -1670,6 +1523,16 @@ namespace UE::USDStageImporter::Private
 			if (Object && Object->GetOutermost() != GetTransientPackage())
 			{
 				FAssetRegistryModule::AssetCreated(Object);
+
+				if (bIsReimport)
+				{
+					GEditor->GetEditorSubsystem<UImportSubsystem>()->BroadcastAssetReimport(Object);
+				}
+				else
+				{
+					UFactory* Factory = nullptr;
+					GEditor->GetEditorSubsystem<UImportSubsystem>()->BroadcastAssetPostImport(Factory, Object);
+				}
 			}
 		}
 	}
@@ -1915,21 +1778,6 @@ void UUsdStageImporter::ImportFromFile(FUsdStageImportContext& ImportContext)
 
 	double StartTime = FPlatformTime::Cycles64();
 
-	// If we don't have any transaction yet, let's temporarly disable the transaction buffer, to prevent code
-	// downstream from accidentally creating their own transactions. This happens for example on USkeleton::AccumulateCurveMetaData,
-	// and can lead to thousands of transactions showing up on the editor and a huge performance cost due to serialization spam.
-	// See also UEditorEngine::CanTransact.
-	TStrongObjectPtr<UTransactor> TransactorPin;
-	TOptional<TGuardValue<TObjectPtr<class UTransactor>>> TransactionSuppressor;
-	if (!GIsTransacting)
-	{
-		// Some stuff like the DuplicateObject or FBlueprintCompilationManager::CompileSynchronously inside SkelSkeletonTranslator
-		// will trigger GC, which we cannot prevent from here. If that happened when our transactor wasn't being referenced via
-		// GEditor->Trans it would actually get collected, so here we prevent that from happening
-		TransactorPin.Reset(GEditor->Trans.Get());
-		TransactionSuppressor.Emplace(GEditor->Trans, nullptr);
-	}
-
 	// Load some default options in case we don't have an options object, so that we don't have to check for this every time.
 	// We could also GetMutableDefault<UUsdStageImportOptions>() instead, but I think in this case it's expected that these
 	// should be the default value and not the config
@@ -1962,6 +1810,23 @@ void UUsdStageImporter::ImportFromFile(FUsdStageImportContext& ImportContext)
 		return;
 	}
 
+	// Reset the transaction buffer so we can't undo the creation of the assets.
+	// We'll also temporarily disable creating new transactions, to prevent code downstream from accidentally creating new ones.
+	// This happens for example on USkeleton::AccumulateCurveMetaData, and can lead to thousands of transactions showing up on
+	// the editor and a huge performance cost due to serialization spam.
+	// See also UEditorEngine::CanTransact.
+	GEditor->ResetTransaction(LOCTEXT("ResetBeforeUSDImport", "USD import from file"));
+	TStrongObjectPtr<UTransactor> TransactorPin;
+	TOptional<TGuardValue<TObjectPtr<class UTransactor>>> TransactionSuppressor;
+	if (ensure(!GIsTransacting))
+	{
+		// Some stuff like the DuplicateObject or FBlueprintCompilationManager::CompileSynchronously inside SkelSkeletonTranslator
+		// will trigger GC, which we cannot prevent from here. If that happened when our transactor wasn't being referenced via
+		// GEditor->Trans it would actually get collected, so here we prevent that from happening
+		TransactorPin.Reset(GEditor->Trans.Get());
+		TransactionSuppressor.Emplace(GEditor->Trans, nullptr);
+	}
+
 	FUsdDelegates::OnPreUsdImport.Broadcast(ImportContext.FilePath);
 
 	AActor* ExistingSceneActor = UE::USDStageImporter::Private::GetExistingSceneActor(ImportContext);
@@ -1971,18 +1836,36 @@ void UUsdStageImporter::ImportFromFile(FUsdStageImportContext& ImportContext)
 	TMap<FSoftObjectPath, FSoftObjectPath> SoftObjectsToRemap;
 	TMap<UObject*, UObject*> ObjectsToRemap;
 	TSet<UObject*> UsedAssetsAndDependencies;
-	TSet<UObject*> PublishedAssetsAndDependencies;
+	TSet<UObject*> PublishedAssets;
 	TSet<UObject*> AssetsToFinalize;
 	UsdUtils::FBlendShapeMap BlendShapesByPath;
 
 	// Ensure a valid asset cache
-	if (!ImportContext.AssetCache)
+	TStrongObjectPtr<UUsdAssetCache3> AssetCachePin;
+	if (!ImportContext.UsdAssetCache && ImportContext.ImportOptions->bUseExistingAssetCache)
+	{
+		// Try using a provided ExistingAssetCache if that wasn't put on our import context already
+		ImportContext.UsdAssetCache = Cast<UUsdAssetCache3>(ImportContext.ImportOptions->ExistingAssetCache.TryLoad());
+	}
+	if (!ImportContext.UsdAssetCache)
+	{
+		// Get a transient AssetCache.
+		// By being in the transient package this cache will also dump its assets in the transient package.
+		// We will then take these assets and publish them to the final import folder.
+		ImportContext.UsdAssetCache = NewObject<UUsdAssetCache3>();
+		AssetCachePin.Reset(ImportContext.UsdAssetCache);
+	}
+	if (!ensure(ImportContext.UsdAssetCache))
 	{
-		UE_LOG(LogUsd, Log, TEXT("Generating a temporary USD Asset Cache when importing '%s'."), *ImportContext.FilePath);
-		ImportContext.AssetCache = NewObject<UUsdAssetCache2>();
+		return;
 	}
+	// Even though we might be using an external asset cache, temporarily make sure that the asset directory
+	// is the transient package, so that any new assets that we create when importing don't end up in the
+	// content browser and can instead just be renamed by the stageimporter directly into the import location
+	TGuardValue<FString> AssetCacheTransientGuard{ImportContext.UsdAssetCache->AssetDirectory.Path, GetTransientPackage()->GetPathName()};
+
 	TSharedPtr<FUsdInfoCache> InfoCache = MakeShared<FUsdInfoCache>();
-	ImportContext.AssetCache->MarkAssetsAsStale();
+	ImportContext.UsdAssetCache->MarkAssetsAsStale();
 	ImportContext.LevelSequenceHelper.SetInfoCache(InfoCache);
 	ImportContext.LevelSequenceHelper.Init(ImportContext.Stage);	// Must happen after the context gets an InfoCache!
 	ImportContext.LevelSequenceHelper.SetRootMotionHandling(ImportContext.ImportOptions->RootMotionHandling);
@@ -2019,7 +1902,7 @@ void UUsdStageImporter::ImportFromFile(FUsdStageImportContext& ImportContext)
 
 	TSharedRef<FUsdSchemaTranslationContext> TranslationContext = MakeShared<FUsdSchemaTranslationContext>(
 		ImportContext.Stage,
-		*ImportContext.AssetCache
+		*ImportContext.UsdAssetCache
 	);
 	TranslationContext->bIsImporting = true;
 	TranslationContext->Level = ImportContext.World->GetCurrentLevel();
@@ -2058,23 +1941,24 @@ void UUsdStageImporter::ImportFromFile(FUsdStageImportContext& ImportContext)
 	UE::USDStageImporter::Private::PruneUnwantedAssets(ImportContext, UsedAssetsAndDependencies, ObjectsToRemap, SoftObjectsToRemap);
 	UE::USDStageImporter::Private::UpdateAssetUserData(UsedAssetsAndDependencies, ImportContext.FilePath, ImportContext.ImportOptions);
 	UE::USDStageImporter::Private::ResolveActorConflicts(ImportContext, ExistingSceneActor, ObjectsToRemap, SoftObjectsToRemap);
-	UE::USDStageImporter::Private::PublishAssets(ImportContext, UsedAssetsAndDependencies, ObjectsToRemap, SoftObjectsToRemap, AssetsToFinalize);
-	UE::USDStageImporter::Private::RemapReferences(ImportContext, UsedAssetsAndDependencies, ObjectsToRemap);
-	UE::USDStageImporter::Private::RemapSoftReferences(ImportContext, UsedAssetsAndDependencies, SoftObjectsToRemap);
+	UE::USDStageImporter::Private::PublishAssets(
+		ImportContext,
+		UsedAssetsAndDependencies,
+		ObjectsToRemap,
+		SoftObjectsToRemap,
+		PublishedAssets,
+		AssetsToFinalize
+	);
+	UE::USDStageImporter::Private::RemapReferences(ImportContext, PublishedAssets, ObjectsToRemap);
+	UE::USDStageImporter::Private::RemapSoftReferences(ImportContext, PublishedAssets, SoftObjectsToRemap);
 	UE::USDStageImporter::Private::Cleanup(
 		ImportContext.SceneActor,
 		ExistingSceneActor,
 		ImportContext.ImportOptions->ExistingActorPolicy,
 		AssetsToFinalize
 	);
-	UE::USDStageImporter::Private::GetPublishedAssetsAndDependencies(
-		UsedAssetsAndDependencies,
-		ObjectsToRemap,
-		SoftObjectsToRemap,
-		PublishedAssetsAndDependencies
-	);
-	UE::USDStageImporter::Private::CallAssetsPostEditChange(PublishedAssetsAndDependencies);
-	UE::USDStageImporter::Private::NotifyAssetRegistry(PublishedAssetsAndDependencies);
+	UE::USDStageImporter::Private::CallAssetsPostEditChange(PublishedAssets);
+	UE::USDStageImporter::Private::BroadcastImportEvents(PublishedAssets);
 	UE::USDStageImporter::Private::RefreshComponents(ImportContext.SceneActor, ImportContext.ImportOptions->bImportAtSpecificTimeCode);
 
 	if (IncludedPurposesToRevertBBoxCacheTo.IsSet())
@@ -2091,7 +1975,7 @@ void UUsdStageImporter::ImportFromFile(FUsdStageImportContext& ImportContext)
 	// Analytics
 	{
 		double ElapsedSeconds = FPlatformTime::ToSeconds64(FPlatformTime::Cycles64() - StartTime);
-		UE::USDStageImporter::Private::SendAnalytics(ImportContext, nullptr, TEXT("Import"), PublishedAssetsAndDependencies, ElapsedSeconds);
+		UE::USDStageImporter::Private::SendAnalytics(ImportContext, nullptr, TEXT("Import"), PublishedAssets, ElapsedSeconds);
 		UE_LOG(LogUsd, Log, TEXT("Imported '%s' in %.3f seconds."), *ImportContext.FilePath, ElapsedSeconds);
 
 		UsdUnreal::Analytics::CollectSchemaAnalytics(ImportContext.Stage, TEXT("Import"));
@@ -2114,21 +1998,6 @@ bool UUsdStageImporter::ReimportSingleAsset(
 #if USE_USD_SDK
 	double StartTime = FPlatformTime::Cycles64();
 
-	// If we don't have any transaction yet, let's temporarly disable the transaction buffer, to prevent code
-	// downstream from accidentally creating their own transactions. This happens for example on USkeleton::AccumulateCurveMetaData,
-	// and can lead to thousands of transactions showing up on the editor and a huge performance cost due to serialization spam.
-	// See also UEditorEngine::CanTransact.
-	TStrongObjectPtr<UTransactor> TransactorPin;
-	TOptional<TGuardValue<TObjectPtr<class UTransactor>>> TransactionSuppressor;
-	if (!GIsTransacting)
-	{
-		// Some stuff like the DuplicateObject or FBlueprintCompilationManager::CompileSynchronously inside SkelSkeletonTranslator
-		// will trigger GC, which we cannot prevent from here. If that happened when our transactor wasn't being referenced via
-		// GEditor->Trans it would actually get collected, so here we prevent that from happening
-		TransactorPin.Reset(GEditor->Trans.Get());
-		TransactionSuppressor.Emplace(GEditor->Trans, nullptr);
-	}
-
 	if (!ImportContext.ImportOptions)
 	{
 		ImportContext.ImportOptions = NewObject<UUsdStageImportOptions>();
@@ -2148,6 +2017,23 @@ bool UUsdStageImporter::ReimportSingleAsset(
 		return bSuccess;
 	}
 
+	// Reset the transaction buffer so we can't undo the creation of the assets.
+	// We'll also temporarily disable creating new transactions, to prevent code downstream from accidentally creating new ones.
+	// This happens for example on USkeleton::AccumulateCurveMetaData, and can lead to thousands of transactions showing up on
+	// the editor and a huge performance cost due to serialization spam.
+	// See also UEditorEngine::CanTransact.
+	GEditor->ResetTransaction(LOCTEXT("ResetBeforeUSDReimport", "USD reimport from file"));
+	TStrongObjectPtr<UTransactor> TransactorPin;
+	TOptional<TGuardValue<TObjectPtr<class UTransactor>>> TransactionSuppressor;
+	if (ensure(!GIsTransacting))
+	{
+		// Some stuff like the DuplicateObject or FBlueprintCompilationManager::CompileSynchronously inside SkelSkeletonTranslator
+		// will trigger GC, which we cannot prevent from here. If that happened when our transactor wasn't being referenced via
+		// GEditor->Trans it would actually get collected, so here we prevent that from happening
+		TransactorPin.Reset(GEditor->Trans.Get());
+		TransactionSuppressor.Emplace(GEditor->Trans, nullptr);
+	}
+
 	FUsdDelegates::OnPreUsdImport.Broadcast(ImportContext.FilePath);
 
 	// We still need the scene actor to remap all other users of the mesh to the new reimported one. It's not critical if we fail though,
@@ -2162,13 +2048,31 @@ bool UUsdStageImporter::ReimportSingleAsset(
 	UsdUtils::FBlendShapeMap BlendShapesByPath;
 
 	// Ensure a valid asset cache
-	if (!ImportContext.AssetCache)
+	TStrongObjectPtr<UUsdAssetCache3> AssetCachePin;
+	if (!ImportContext.UsdAssetCache && ImportContext.ImportOptions->bUseExistingAssetCache)
+	{
+		// Try using a provided ExistingAssetCache if that wasn't put on our import context already
+		ImportContext.UsdAssetCache = Cast<UUsdAssetCache3>(ImportContext.ImportOptions->ExistingAssetCache.TryLoad());
+	}
+	if (!ImportContext.UsdAssetCache)
 	{
-		UE_LOG(LogUsd, Log, TEXT("Generating a temporary USD Asset Cache when importing '%s'."), *ImportContext.FilePath);
-		ImportContext.AssetCache = NewObject<UUsdAssetCache2>();
+		// Get a transient AssetCache.
+		// By being in the transient package this cache will also dump its assets in the transient package.
+		// We will then take these assets and publish them to the final import folder.
+		ImportContext.UsdAssetCache = NewObject<UUsdAssetCache3>();
+		AssetCachePin.Reset(ImportContext.UsdAssetCache);
 	}
+	if (!ensure(ImportContext.UsdAssetCache))
+	{
+		return bSuccess;
+	}
+	// Even though we might be using an external asset cache, temporarily make sure that the asset directory
+	// is the transient package, so that any new assets that we create when importing don't end up in the
+	// content browser and can instead just be renamed by the stageimporter directly into the import location
+	TGuardValue<FString> AssetCacheTransientGuard{ImportContext.UsdAssetCache->AssetDirectory.Path, GetTransientPackage()->GetPathName()};
+
 	TSharedPtr<FUsdInfoCache> InfoCache = MakeShared<FUsdInfoCache>();
-	ImportContext.AssetCache->MarkAssetsAsStale();
+	ImportContext.UsdAssetCache->MarkAssetsAsStale();
 	ImportContext.LevelSequenceHelper.SetInfoCache(InfoCache);
 	ImportContext.LevelSequenceHelper.Init(ImportContext.Stage);	// Must happen after the context gets an InfoCache!
 	ImportContext.LevelSequenceHelper.SetRootMotionHandling(ImportContext.ImportOptions->RootMotionHandling);
@@ -2204,7 +2108,7 @@ bool UUsdStageImporter::ReimportSingleAsset(
 
 	TSharedRef<FUsdSchemaTranslationContext> TranslationContext = MakeShared<FUsdSchemaTranslationContext>(
 		ImportContext.Stage,
-		*ImportContext.AssetCache
+		*ImportContext.UsdAssetCache
 	);
 	TranslationContext->bIsImporting = true;
 	TranslationContext->Level = ImportContext.World->GetCurrentLevel();
@@ -2241,9 +2145,18 @@ bool UUsdStageImporter::ReimportSingleAsset(
 	}
 	TranslationContext->CompleteTasks();
 
+	TSet<FSoftObjectPath> ActiveAssetPaths = ImportContext.UsdAssetCache->GetActiveAssets();
+	TSet<UObject*> ActiveAssets;
+	ActiveAssets.Reserve(ActiveAssetPaths.Num());
+	for (const FSoftObjectPath& ActivePath : ActiveAssetPaths)
+	{
+		ActiveAssets.Add(ActivePath.TryLoad());
+	}
+	ActiveAssets.Remove(nullptr);
+
 	// Look for our reimported asset in the assets cache as we may have multiple assets with the same prim path
 	UObject* ReimportedObject = nullptr;
-	for (UObject* Asset : ImportContext.AssetCache->GetActiveAssets())
+	for (UObject* Asset : ActiveAssets)
 	{
 		UUsdAssetUserData* UserData = nullptr;
 		if (IInterface_AssetUserData* UserDataInterface = Cast<IInterface_AssetUserData>(Asset))
@@ -2277,23 +2190,17 @@ bool UUsdStageImporter::ReimportSingleAsset(
 			AssetsToFinalize
 		);
 
-		TSet<UObject*> UsedAssetsAndDependencies = ImportContext.AssetCache->GetActiveAssets();
 		TSet<UObject*> PublishedAssetsAndDependencies;
-		UE::USDStageImporter::Private::RemapReferences(ImportContext, UsedAssetsAndDependencies, ObjectsToRemap);
-		UE::USDStageImporter::Private::RemapSoftReferences(ImportContext, UsedAssetsAndDependencies, SoftObjectsToRemap);
-		UE::USDStageImporter::Private::GetPublishedAssetsAndDependencies(
-			UsedAssetsAndDependencies,
-			ObjectsToRemap,
-			SoftObjectsToRemap,
-			PublishedAssetsAndDependencies
-		);
+		UE::USDStageImporter::Private::RemapReferences(ImportContext, {OutReimportedAsset}, ObjectsToRemap);
+		UE::USDStageImporter::Private::RemapSoftReferences(ImportContext, {OutReimportedAsset}, SoftObjectsToRemap);
 		UE::USDStageImporter::Private::CallAssetsPostEditChange(PublishedAssetsAndDependencies);
 
-		bSuccess = OutReimportedAsset != nullptr && ImportContext.AssetCache->GetActiveAssets().Contains(ReimportedObject);
+		bSuccess = OutReimportedAsset != nullptr && ActiveAssets.Contains(ReimportedObject);
 	}
 
+	const bool bIsReimport = true;
 	UE::USDStageImporter::Private::Cleanup(ImportContext.SceneActor, nullptr, ImportContext.ImportOptions->ExistingActorPolicy, AssetsToFinalize);
-	UE::USDStageImporter::Private::NotifyAssetRegistry({ReimportedObject});
+	UE::USDStageImporter::Private::BroadcastImportEvents({ReimportedObject}, bIsReimport);
 	UE::USDStageImporter::Private::RefreshComponents(ImportContext.SceneActor, ImportContext.ImportOptions->bImportAtSpecificTimeCode);
 
 	if (IncludedPurposesToRevertTo.IsSet())
diff --git a/Engine/Plugins/Importers/USDImporter/Source/USDStageImporter/Public/USDStageImportContext.h b/Engine/Plugins/Importers/USDImporter/Source/USDStageImporter/Public/USDStageImportContext.h
--- a/Engine/Plugins/Importers/USDImporter/Source/USDStageImporter/Public/USDStageImportContext.h
+++ b/Engine/Plugins/Importers/USDImporter/Source/USDStageImporter/Public/USDStageImportContext.h
@@ -13,6 +13,7 @@
 #include "USDStageImportContext.generated.h"
 
 class UUsdAssetCache2;
+class UUsdAssetCache3;
 class UUsdStageImportOptions;
 class FTokenizedMessage;
 
@@ -55,6 +56,10 @@ struct USDSTAGEIMPORTER_API FUsdStageImportContext
 	/** Level sequence that will contain the animation data during the import process */
 	FUsdLevelSequenceHelper LevelSequenceHelper;
 
+	UPROPERTY()
+	TObjectPtr<UUsdAssetCache3> UsdAssetCache;
+
+	UE_DEPRECATED(5.5, "Use the 'UsdAssetCache' member instead, which is of the new UUsdAssetCache3 type")
 	UPROPERTY()
 	TObjectPtr<UUsdAssetCache2> AssetCache;
 
diff --git a/Engine/Plugins/Importers/USDImporter/Source/USDStageImporter/Public/USDStageImportOptions.h b/Engine/Plugins/Importers/USDImporter/Source/USDStageImporter/Public/USDStageImportOptions.h
--- a/Engine/Plugins/Importers/USDImporter/Source/USDStageImporter/Public/USDStageImportOptions.h
+++ b/Engine/Plugins/Importers/USDImporter/Source/USDStageImporter/Public/USDStageImportOptions.h
@@ -93,6 +93,13 @@ public:
 	UPROPERTY(BlueprintReadWrite, EditAnywhere, Category = "Prims to Import")
 	TArray<FString> PrimsToImport = TArray<FString>{TEXT("/")};
 
+	UPROPERTY(BlueprintReadWrite, config, EditAnywhere, Category = "USD options", meta = (InlineEditConditionToggle))
+	bool bUseExistingAssetCache;
+
+	/** Copy assets from an existing UsdAsset cache instead of generating them from scratch, if possible */
+	UPROPERTY(BlueprintReadWrite, config, EditAnywhere, Category = "USD options", meta = (AllowedClasses = "/Script/USDClasses.UsdAssetCache3", EditCondition = "bUseExistingAssetCache"))
+	FSoftObjectPath ExistingAssetCache = nullptr;
+
 	/** Only import prims with these specific purposes from the USD file */
 	UPROPERTY(
 		BlueprintReadWrite,
diff --git a/Engine/Plugins/Importers/USDImporter/Source/USDUtilities/Private/USDConversionUtils.cpp b/Engine/Plugins/Importers/USDImporter/Source/USDUtilities/Private/USDConversionUtils.cpp
--- a/Engine/Plugins/Importers/USDImporter/Source/USDUtilities/Private/USDConversionUtils.cpp
+++ b/Engine/Plugins/Importers/USDImporter/Source/USDUtilities/Private/USDConversionUtils.cpp
@@ -11,6 +11,7 @@
 #include "USDGeomMeshConversion.h"
 #include "USDLayerUtils.h"
 #include "USDLog.h"
+#include "USDObjectUtils.h"
 #include "USDProjectSettings.h"
 #include "USDSkeletalDataConversion.h"
 #include "USDTypesConversion.h"
@@ -1536,7 +1537,7 @@ FString UsdUtils::GetAssetPathFromPrimPath(const FString& RootContentPath, const
 	ModelApi.GetAssetName(&RawAssetName);
 
 	FString AssetName = UsdToUnreal::ConvertString(RawAssetName);
-	FString MeshName = IUsdClassesModule::SanitizeObjectName(RawPrimName);
+	FString MeshName = UsdUnreal::ObjectUtils::SanitizeObjectName(RawPrimName);
 
 	FString USDPath = UsdToUnreal::ConvertString(Prim.GetPrimPath().GetString().c_str());
 
@@ -1646,6 +1647,8 @@ double UsdUtils::GetEarliestTimeCode()
 #endif
 }
 
+// We can't just redirect the functions to USDObjectUtils.h because of the module dependencies
+PRAGMA_DISABLE_DEPRECATION_WARNINGS
 UUsdAssetImportData* UsdUtils::GetAssetImportData(UObject* Asset)
 {
 	UUsdAssetImportData* ImportData = nullptr;
@@ -1858,6 +1861,7 @@ bool UsdUtils::SetAssetUserData(UObject* Object, UUsdAssetUserData* AssetUserDat
 	AssetUserDataInterface->AddAssetUserData(AssetUserData);
 	return true;
 }
+PRAGMA_ENABLE_DEPRECATION_WARNINGS
 
 #if USE_USD_SDK
 FString UsdUtils::GetAssetHashPrefix(const pxr::UsdPrim& PrimForAsset, bool bReuseIdenticalAssets)
@@ -2254,49 +2258,12 @@ bool UsdUtils::RenamePrim(UE::FUsdPrim& Prim, const TCHAR* NewPrimName)
 
 bool UsdUtils::RemoveNumberedSuffix(FString& Prefix)
 {
-	if (Prefix.IsNumeric())
-	{
-		return false;
-	}
-
-	bool bRemoved = false;
-
-	FString LastChar = Prefix.Right(1);
-	while ((LastChar.IsNumeric() || LastChar == TEXT("_")) && Prefix.Len() > 1)
-	{
-		Prefix.LeftChopInline(1, EAllowShrinking::No);
-		LastChar = Prefix.Right(1);
-
-		bRemoved = true;
-	}
-	Prefix.Shrink();
-
-	return bRemoved;
+	return UsdUnreal::ObjectUtils::RemoveNumberedSuffix(Prefix);
 }
 
 FString UsdUtils::GetUniqueName(FString Name, const TSet<FString>& UsedNames)
 {
-	if (!UsedNames.Contains(Name))
-	{
-		return Name;
-	}
-
-	const bool bRemoved = RemoveNumberedSuffix(Name);
-
-	// Its possible that removing the suffix made it into a unique name already
-	if (bRemoved && !UsedNames.Contains(Name))
-	{
-		return Name;
-	}
-
-	int32 Suffix = 0;
-	FString Result;
-	do
-	{
-		Result = FString::Printf(TEXT("%s_%d"), *Name, Suffix++);
-	} while (UsedNames.Contains(Result));
-
-	return Result;
+	return UsdUnreal::ObjectUtils::GetUniqueName(Name, UsedNames);
 }
 
 #if USE_USD_SDK
@@ -2315,7 +2282,7 @@ FString UsdUtils::GetValidChildName(FString InName, const pxr::UsdPrim& ParentPr
 		UsedNames.Add(UsdToUnreal::ConvertToken(Child.GetName()));
 	}
 
-	return GetUniqueName(SanitizeUsdIdentifier(*InName), UsedNames);
+	return UsdUnreal::ObjectUtils::GetUniqueName(SanitizeUsdIdentifier(*InName), UsedNames);
 }
 #endif	  // USE_USD_SDK
 
@@ -2728,7 +2695,7 @@ bool UsdUtils::CopyPrims(const TArray<UE::FUsdPrim>& Prims)
 		// which could potentially yield stale results until USD actually emits the notices about these prims being
 		// added.
 		FString PrimName = Prim.GetName().ToString();
-		FString UniqueName = GetUniqueName(SanitizeUsdIdentifier(*PrimName), UsedNames);
+		FString UniqueName = UsdUnreal::ObjectUtils::GetUniqueName(SanitizeUsdIdentifier(*PrimName), UsedNames);
 		UsedNames.Add(UniqueName);
 
 		const bool bSuccess = pxr::SdfCopySpec(
@@ -2813,7 +2780,7 @@ TArray<UE::FSdfPath> UsdUtils::PastePrims(const UE::FUsdPrim& ParentPrim)
 		}
 
 		const FString OriginalName = UsdToUnreal::ConvertToken(ClipboardPrim.GetName());
-		FString ValidName = GetUniqueName(SanitizeUsdIdentifier(*OriginalName), UsedNames);
+		FString ValidName = UsdUnreal::ObjectUtils::GetUniqueName(SanitizeUsdIdentifier(*OriginalName), UsedNames);
 		UsedNames.Add(ValidName);
 
 		pxr::SdfPath TargetSpecPath = UsdParentPrim.GetPath().AppendChild(UnrealToUsd::ConvertToken(*ValidName).Get());
diff --git a/Engine/Plugins/Importers/USDImporter/Source/USDUtilities/Private/USDGeomMeshConversion.cpp b/Engine/Plugins/Importers/USDImporter/Source/USDUtilities/Private/USDGeomMeshConversion.cpp
--- a/Engine/Plugins/Importers/USDImporter/Source/USDUtilities/Private/USDGeomMeshConversion.cpp
+++ b/Engine/Plugins/Importers/USDImporter/Source/USDUtilities/Private/USDGeomMeshConversion.cpp
@@ -6,12 +6,12 @@
 
 #include "UnrealUSDWrapper.h"
 #include "USDAttributeUtils.h"
-#include "USDClassesModule.h"
 #include "USDConversionUtils.h"
 #include "USDDrawModeComponent.h"
 #include "USDLayerUtils.h"
 #include "USDLog.h"
 #include "USDMemory.h"
+#include "USDObjectUtils.h"
 #include "USDPrimConversion.h"
 #include "USDProjectSettings.h"
 #include "USDShadeConversion.h"
@@ -3107,9 +3107,9 @@ UMaterialInstanceDynamic* UsdUtils::CreateDisplayColorMaterialInstanceDynamic(co
 			GetTransientPackage(),
 			UMaterialInstanceConstant::StaticClass(),
 			*FString::Printf(
-				TEXT("DisplayColor_%s_%s"),
-				DisplayColorDescription.bHasOpacity ? TEXT("Opacity") : TEXT("NoOpacity"),
-				DisplayColorDescription.bIsDoubleSided ? TEXT("DoubleSided") : TEXT("SingleSided")
+				TEXT("DisplayColor%s%s"),
+				DisplayColorDescription.bHasOpacity ? TEXT("_Translucent") : TEXT(""),
+				DisplayColorDescription.bIsDoubleSided ? TEXT("_TwoSided") : TEXT("")
 			)
 		);
 
@@ -3161,9 +3161,9 @@ UMaterialInstanceConstant* UsdUtils::CreateDisplayColorMaterialInstanceConstant(
 			GetTransientPackage(),
 			UMaterialInstanceConstant::StaticClass(),
 			*FString::Printf(
-				TEXT("DisplayColor_%s_%s"),
-				DisplayColorDescription.bHasOpacity ? TEXT("Opacity") : TEXT("NoOpacity"),
-				DisplayColorDescription.bIsDoubleSided ? TEXT("DoubleSided") : TEXT("SingleSided")
+				TEXT("DisplayColor%s%s"),
+				DisplayColorDescription.bHasOpacity ? TEXT("_Translucent") : TEXT(""),
+				DisplayColorDescription.bIsDoubleSided ? TEXT("_TwoSided") : TEXT("")
 			)
 		);
 
@@ -4175,7 +4175,7 @@ namespace UE::UsdGeometryCacheConversion::Private
 					if (ExportContext.SlotNames.IsValidIndex(SectionIndex))
 					{
 						SectionName = ExportContext.SlotNames[SectionIndex].ToString();
-						SectionName = UsdUtils::GetUniqueName(SectionName, UsedSectionNames);
+						SectionName = UsdUnreal::ObjectUtils::GetUniqueName(SectionName, UsedSectionNames);
 						UsedSectionNames.Add(SectionName);
 					}
 					else
@@ -4903,7 +4903,7 @@ void UsdUtils::ReplaceUnrealMaterialsWithBaked(
 					{
 						FString MatName = FPaths::GetBaseFilename(UnrealMaterialAssetPath);
 						MatName = UsdToUnreal::ConvertString(pxr::TfMakeValidIdentifier(UnrealToUsd::ConvertString(*MatName).Get()));
-						FString MatPrimName = UsdUtils::GetUniqueName(MatName, MatPrimScopePtr->UsedPrimNames);
+						FString MatPrimName = UsdUnreal::ObjectUtils::GetUniqueName(MatName, MatPrimScopePtr->UsedPrimNames);
 						MatPrimScopePtr->UsedPrimNames.Add(MatPrimName);
 
 						MatPrim = StageToTraverse->DefinePrim(
diff --git a/Engine/Plugins/Importers/USDImporter/Source/USDUtilities/Private/USDLightConversion.cpp b/Engine/Plugins/Importers/USDImporter/Source/USDUtilities/Private/USDLightConversion.cpp
--- a/Engine/Plugins/Importers/USDImporter/Source/USDUtilities/Private/USDLightConversion.cpp
+++ b/Engine/Plugins/Importers/USDImporter/Source/USDUtilities/Private/USDLightConversion.cpp
@@ -2,11 +2,14 @@
 
 #include "USDLightConversion.h"
 
-#include "USDAssetCache2.h"
+#include "USDAssetCache3.h"
+#include "USDAssetUserData.h"
 #include "USDAttributeUtils.h"
+#include "USDClassesModule.h"
 #include "USDConversionUtils.h"
 #include "USDLayerUtils.h"
 #include "USDLog.h"
+#include "USDObjectUtils.h"
 #include "USDShadeConversion.h"
 #include "USDTypesConversion.h"
 
@@ -281,23 +284,22 @@ bool UsdToUnreal::ConvertDomeLight(
 		return true;
 	}
 
-	const FString PrefixedTextureHash = UsdUtils::GetAssetHashPrefix(Prim, bReuseIdenticalAssets)
-										+ LexToString(FMD5Hash::HashFile(*ResolvedDomeTexturePath));
-	UTextureCube* Cubemap = Cast<UTextureCube>(TexturesCache ? TexturesCache->GetCachedAsset(PrefixedTextureHash) : nullptr);
+	const FString& DesiredTextureName = FPaths::GetBaseFilename(ResolvedDomeTexturePath);
 
-	if (!Cubemap)
-	{
-		Cubemap = Cast<UTextureCube>(UsdUtils::CreateTexture(
-			DomeLight.GetTextureFileAttr(),
-			UsdToUnreal::ConvertPath(DomeLight.GetPrim().GetPath()),
-			TEXTUREGROUP_Skybox,
-			TexturesCache
-		));
+	// TODO: Ideally we'd expose this
+	const EObjectFlags DesiredFlags = RF_Public | RF_Standalone | RF_Transactional;
 
-		if (TexturesCache)
-		{
-			TexturesCache->CacheAsset(PrefixedTextureHash, Cubemap);
-		}
+	TextureGroup Group = TextureGroup::TEXTUREGROUP_Skybox;
+
+	UTextureCube* Cubemap = nullptr;
+
+	UObject* Outer = GetTransientPackage();
+	FName TextureName = MakeUniqueObjectName(Outer, UTextureCube::StaticClass(), *UsdUnreal::ObjectUtils::SanitizeObjectName(DesiredTextureName));
+	Cubemap = Cast<UTextureCube>(UsdUtils::CreateTexture(ResolvedDomeTexturePath, TextureName, Group, DesiredFlags));
+
+	if (UUsdAssetUserData* TextureUserData = UsdUnreal::ObjectUtils::GetOrCreateAssetUserData(Cubemap))
+	{
+		TextureUserData->PrimPaths.AddUnique(UsdToUnreal::ConvertPath(DomeLight.GetPrim().GetPath()));
 	}
 
 	if (Cubemap)
@@ -336,12 +338,6 @@ bool UsdToUnreal::ConvertLuxShapingAPI(const pxr::UsdPrim& Prim, USpotLightCompo
 	return true;
 }
 
-bool UsdToUnreal::ConvertDomeLight(const pxr::UsdPrim& Prim, USkyLightComponent& LightComponent, UUsdAssetCache* TexturesCache)
-{
-	UUsdAssetCache2* NewCache = nullptr;
-	return UsdToUnreal::ConvertDomeLight(Prim, LightComponent, NewCache);
-}
-
 float UsdToUnreal::ConvertLightIntensityAttr(float UsdIntensity, float UsdExposure)
 {
 	return UsdIntensity * FMath::Exp2(UsdExposure);
diff --git a/Engine/Plugins/Importers/USDImporter/Source/USDUtilities/Private/USDPrimConversion.cpp b/Engine/Plugins/Importers/USDImporter/Source/USDUtilities/Private/USDPrimConversion.cpp
--- a/Engine/Plugins/Importers/USDImporter/Source/USDUtilities/Private/USDPrimConversion.cpp
+++ b/Engine/Plugins/Importers/USDImporter/Source/USDUtilities/Private/USDPrimConversion.cpp
@@ -3,7 +3,6 @@
 #include "USDPrimConversion.h"
 
 #include "UnrealUSDWrapper.h"
-#include "USDAssetCache2.h"
 #include "USDAssetUserData.h"
 #include "USDAttributeUtils.h"
 #include "USDConversionUtils.h"
@@ -11,6 +10,7 @@
 #include "USDLayerUtils.h"
 #include "USDLightConversion.h"
 #include "USDLog.h"
+#include "USDObjectUtils.h"
 #include "USDShadeConversion.h"
 #include "USDSkeletalDataConversion.h"
 #include "USDTypesConversion.h"
@@ -4570,7 +4570,7 @@ UnrealToUsd::FPropertyTrackWriter UnrealToUsd::CreatePropertyTrackWriter(
 							if (SparseVolumeTexture->GetNumFrames() > 1)
 							{
 								if (UUsdSparseVolumeTextureAssetUserData* UserData = Cast<UUsdSparseVolumeTextureAssetUserData>(
-										UsdUtils::GetAssetUserData(SparseVolumeTexture)
+										UsdUnreal::ObjectUtils::GetAssetUserData(SparseVolumeTexture)
 									))
 								{
 									SourceOpenVDBAssetPrimPaths = &UserData->SourceOpenVDBAssetPrimPaths;
diff --git a/Engine/Plugins/Importers/USDImporter/Source/USDUtilities/Private/USDShadeConversion.cpp b/Engine/Plugins/Importers/USDImporter/Source/USDUtilities/Private/USDShadeConversion.cpp
--- a/Engine/Plugins/Importers/USDImporter/Source/USDUtilities/Private/USDShadeConversion.cpp
+++ b/Engine/Plugins/Importers/USDImporter/Source/USDUtilities/Private/USDShadeConversion.cpp
@@ -5,17 +5,17 @@
 #if USE_USD_SDK
 
 #include "UnrealUSDWrapper.h"
-#include "USDAssetCache.h"
 #include "USDAssetCache2.h"
+#include "USDAssetCache3.h"
 #include "USDAssetImportData.h"
 #include "USDAssetUserData.h"
 #include "USDAttributeUtils.h"
-#include "USDClassesModule.h"
 #include "USDConversionUtils.h"
 #include "USDErrorUtils.h"
 #include "USDLayerUtils.h"
 #include "USDLog.h"
 #include "USDMemory.h"
+#include "USDObjectUtils.h"
 #include "USDTypesConversion.h"
 
 #include "UsdWrappers/UsdAttribute.h"
@@ -264,19 +264,21 @@ namespace UE
 				FName TextureName = MakeUniqueObjectName(
 					Outer,
 					UTexture::StaticClass(),
-					*IUsdClassesModule::SanitizeObjectName(FPaths::GetBaseFilename(ResolvedTexturePath))
+					*UsdUnreal::ObjectUtils::SanitizeObjectName(FPaths::GetBaseFilename(ResolvedTexturePath))
 				);
 
+				UObject* Context = nullptr;
+				FFeedbackContext* FeedbackContext = GWarn;
 				return Cast<UTexture>(TextureFactory->FactoryCreateBinary(
 					UTexture::StaticClass(),
 					Outer,
 					TextureName,
 					ObjectFlags,
-					nullptr,
+					Context,
 					*TextureExtension,
 					BufferStart,
 					BufferStart + BufferSize,
-					nullptr
+					FeedbackContext
 				));
 #endif	  // WITH_EDITOR
 				return nullptr;
@@ -553,10 +555,10 @@ namespace UE
 
 			bool GetTextureParameterValue(
 				pxr::UsdShadeInput& ShadeInput,
-				TextureGroup LODGroup,
+				TextureGroup Group,
 				FParameterValue& OutValue,
 				const UMaterialInterface* Material,
-				UUsdAssetCache2* TexturesCache,
+				UUsdAssetCache3* TexturesCache,
 				bool bReuseIdenticalAssets
 			)
 			{
@@ -666,64 +668,67 @@ namespace UE
 							return false;
 						}
 
-						// We'll add these to the hash because the materials are built to try and reuse the same textures for multiple channels,
-						// and those may expect linear or sRGB values. Without this we may parse a texture as linear because we hit the opacity
-						// channel first, and then reuse it as linear for the base color channel even though it should have been sRGB. Plus we may
-						// have something weird like a normal map being plugged into the base color and the normal channel.
-						bool bSRGB = true;
-						TextureCompressionSettings CompressionSettings = TextureCompressionSettings::TC_Default;
-						if (bIsNormalInput)
-						{
-							// Disable SRGB when parsing float textures, as they're likely specular/roughness maps
-							bSRGB = false;
-							CompressionSettings = TextureCompressionSettings::TC_Normalmap;
-						}
-						else
+						// We only actually want to retrieve the textures if we have a cache to put them in
+						if (TexturesCache)
 						{
-							if (LODGroup == TEXTUREGROUP_WorldSpecular)
+							// We'll add these to the hash because the materials are built to try and reuse the same textures for multiple channels,
+							// and those may expect linear or sRGB values. Without this we may parse a texture as linear because we hit the opacity
+							// channel first, and then reuse it as linear for the base color channel even though it should have been sRGB. Plus we may
+							// have something weird like a normal map being plugged into the base color and the normal channel.
+							bool bSRGB = true;
+							TextureCompressionSettings CompressionSettings = TextureCompressionSettings::TC_Default;
+							if (bIsNormalInput)
 							{
+								// Disable SRGB when parsing float textures, as they're likely specular/roughness maps
 								bSRGB = false;
+								CompressionSettings = TextureCompressionSettings::TC_Normalmap;
 							}
-
-							// Override the sRGB flag with whatever is specified on the texture shader, if it has anything specific.
-							if (pxr::UsdAttribute SourceColorSpaceAttr = UsdUVTextureSource.GetInput(UnrealIdentifiers::SourceColorSpaceToken))
+							else
 							{
-								pxr::TfToken SourceColorSpaceValue;
-								if (SourceColorSpaceAttr.HasAuthoredValue() && SourceColorSpaceAttr.Get(&SourceColorSpaceValue))
+								if (Group == TEXTUREGROUP_WorldSpecular)
 								{
-									bSRGB = SourceColorSpaceValue == UnrealIdentifiers::RawColorSpaceToken	  ? false
-											: SourceColorSpaceValue == UnrealIdentifiers::SRGBColorSpaceToken ? true
-																											  : bSRGB;
+									bSRGB = false;
 								}
-							}
-						}
-
-						const FString PrefixedTextureHash = UsdUtils::GetAssetHashPrefix(ShadeInput.GetPrim(), bReuseIdenticalAssets)
-															+ UsdUtils::GetTextureHash(TexturePath, bSRGB, CompressionSettings, AddressX, AddressY);
 
-						// We only actually want to retrieve the textures if we have a cache to put them in
-						if (TexturesCache)
-						{
-							UTexture* Texture = Cast<UTexture>(TexturesCache->GetCachedAsset(PrefixedTextureHash));
-							if (!Texture)
-							{
-								// Give the same prim path to the texture, so that it ends up imported right next to the material
-								FString MaterialPrimPath;
-								if (Material)
+								// Override the sRGB flag with whatever is specified on the texture shader, if it has anything specific.
+								if (pxr::UsdAttribute SourceColorSpaceAttr = UsdUVTextureSource.GetInput(UnrealIdentifiers::SourceColorSpaceToken))
 								{
-									if (UUsdAssetUserData* UserData = const_cast<UMaterialInterface*>(Material)->GetAssetUserData<UUsdAssetUserData>(
-										))
+									pxr::TfToken SourceColorSpaceValue;
+									if (SourceColorSpaceAttr.HasAuthoredValue() && SourceColorSpaceAttr.Get(&SourceColorSpaceValue))
 									{
-										if (!UserData->PrimPaths.IsEmpty())
-										{
-											MaterialPrimPath = UserData->PrimPaths[0];
-										}
+										bSRGB = SourceColorSpaceValue == UnrealIdentifiers::RawColorSpaceToken	  ? false
+												: SourceColorSpaceValue == UnrealIdentifiers::SRGBColorSpaceToken ? true
+																												  : bSRGB;
 									}
 								}
+							}
+
+							const FString
+								PrefixedTextureHash = UsdUtils::GetAssetHashPrefix(ShadeInput.GetPrim(), bReuseIdenticalAssets)
+													  + UsdUtils::GetTextureHash(TexturePath, bSRGB, CompressionSettings, AddressX, AddressY);
+
+							pxr::UsdAttribute FileInputAttr = FileInput.GetAttr();
+							const FString ResolvedTexturePath = UsdUtils::GetResolvedAssetPath(FileInputAttr);
 
-								Texture = UsdUtils::CreateTexture(FileInput.GetAttr(), MaterialPrimPath, LODGroup, TexturesCache);
+							const FString& DesiredTextureName = FPaths::GetBaseFilename(ResolvedTexturePath);
+
+							const EObjectFlags DesiredFlags = RF_Public | RF_Standalone | RF_Transactional;
+
+							bool bCreatedTexture = false;
+							UTexture* Texture = TexturesCache->GetOrCreateCustomCachedAsset<UTexture2D>(
+								PrefixedTextureHash,
+								DesiredTextureName,
+								DesiredFlags,
+								[&ResolvedTexturePath, Group, TexturesCache](UPackage* Outer, FName SanitizedName, EObjectFlags DesiredFlags)
+								{
+									return UsdUtils::CreateTexture(ResolvedTexturePath, SanitizedName, Group, DesiredFlags, Outer);
+								},
+								&bCreatedTexture
+							);
 
-								if (Texture)
+							if (Texture)
+							{
+								if (bCreatedTexture)
 								{
 									Texture->SRGB = bSRGB;
 									Texture->CompressionSettings = CompressionSettings;
@@ -733,13 +738,23 @@ namespace UE
 										Texture2D->AddressY = AddressY;
 									}
 									Texture->UpdateResource();
+#if WITH_EDITOR
+									Texture->PostEditChange();
+#endif	  // WITH_EDITOR
+								}
 
-									TexturesCache->CacheAsset(PrefixedTextureHash, Texture);
+								if (UUsdAssetUserData* AssetUserData = UsdUnreal::ObjectUtils::GetOrCreateAssetUserData(Texture))
+								{
+									// Give the same prim path to the texture, so that it ends up imported right next to the material
+									if (UUsdAssetUserData* MaterialAssetUserData = UsdUnreal::ObjectUtils::GetAssetUserData(Material))
+									{
+										for (const FString& MaterialPrimPath : MaterialAssetUserData->PrimPaths)
+										{
+											AssetUserData->PrimPaths.AddUnique(MaterialPrimPath);
+										}
+									}
 								}
-							}
 
-							if (Texture)
-							{
 								FTextureParameterValue OutTextureValue;
 								OutTextureValue.Texture = Texture;
 								GetSTPrimvarAndTransform(UsdUVTextureSource, OutTextureValue);
@@ -801,7 +816,7 @@ namespace UE
 				float DefaultValue,
 				FParameterValue& OutValue,
 				const UMaterialInterface* Material,
-				UUsdAssetCache2* TexturesCache,
+				UUsdAssetCache3* TexturesCache,
 				bool bReuseIdenticalAssets
 			)
 			{
@@ -933,7 +948,7 @@ namespace UE
 				FParameterValue& OutValue,
 				bool bIsNormalMap,
 				const UMaterialInterface* Material,
-				UUsdAssetCache2* TexturesCache,
+				UUsdAssetCache3* TexturesCache,
 				bool bReuseIdenticalAssets
 			)
 			{
@@ -1034,7 +1049,7 @@ namespace UE
 				bool DefaultValue,
 				FParameterValue& OutValue,
 				UMaterialInterface* Material = nullptr,
-				UUsdAssetCache2* TexturesCache = nullptr
+				UUsdAssetCache3* TexturesCache = nullptr
 			)
 			{
 				FScopedUsdAllocs Allocs;
@@ -1312,129 +1327,97 @@ namespace UE
 			}
 
 			UTexture* CreateTextureWithEditor(
-				const pxr::UsdAttribute& TextureAssetPathAttr,
-				const FString& PrimPath,
-				TextureGroup LODGroup,
+				const FString& ResolvedTexturePath,
+				FName SanitizedName,
+				TextureGroup Group,
+				EObjectFlags ObjectFlags,
 				UObject* Outer
 			)
 			{
 				UTexture* Texture = nullptr;
 #if WITH_EDITOR
-				FScopedUsdAllocs UsdAllocs;
-
-				pxr::SdfAssetPath TextureAssetPath;
-				TextureAssetPathAttr.Get<pxr::SdfAssetPath>(&TextureAssetPath);
-
-				FString TexturePath = UsdToUnreal::ConvertString(TextureAssetPath.GetAssetPath());
-				const bool bIsSupportedUdimTexture = TexturePath.Contains(TEXT("<UDIM>"));
-				FPaths::NormalizeFilename(TexturePath);
+				if (ResolvedTexturePath.IsEmpty() || !Outer)
+				{
+					return nullptr;
+				}
 
 				FScopedUnrealAllocs UnrealAllocs;
 
-				if (!TexturePath.IsEmpty())
+				UTextureFactory* TextureFactory = NewObject<UTextureFactory>();
+				TextureFactory->SuppressImportOverwriteDialog();
+				TextureFactory->bUseHashAsGuid = true;
+				TextureFactory->LODGroup = Group;
+				TextureFactory->HDRImportShouldBeLongLatCubeMap = EAppReturnType::YesAll;
+
+				const bool bIsSupportedUdimTexture = ResolvedTexturePath.Contains(TEXT("<UDIM>"));
+				if (bIsSupportedUdimTexture)
 				{
-					bool bOutCancelled = false;
-					UTextureFactory* TextureFactory = NewObject<UTextureFactory>();
-					TextureFactory->SuppressImportOverwriteDialog();
-					TextureFactory->bUseHashAsGuid = true;
-					TextureFactory->LODGroup = LODGroup;
-					TextureFactory->HDRImportShouldBeLongLatCubeMap = EAppReturnType::YesAll;
+					FString BaseFileName = FPaths::GetBaseFilename(ResolvedTexturePath);
 
-					if (bIsSupportedUdimTexture)
-					{
-						FString BaseFileName = FPaths::GetBaseFilename(TexturePath);
+					FString BaseFileNameBeforeUdim;
+					FString BaseFileNameAfterUdim;
+					BaseFileName.Split(TEXT("<UDIM>"), &BaseFileNameBeforeUdim, &BaseFileNameAfterUdim);
 
-						FString BaseFileNameBeforeUdim;
-						FString BaseFileNameAfterUdim;
-						BaseFileName.Split(TEXT("<UDIM>"), &BaseFileNameBeforeUdim, &BaseFileNameAfterUdim);
+					FString UdimRegexPattern = FString::Printf(TEXT(R"((%s)(\d{4})(%s))"), *BaseFileNameBeforeUdim, *BaseFileNameAfterUdim);
+					TextureFactory->UdimRegexPattern = MoveTemp(UdimRegexPattern);
+				}
 
-						FString UdimRegexPattern = FString::Printf(TEXT(R"((%s)(\d{4})(%s))"), *BaseFileNameBeforeUdim, *BaseFileNameAfterUdim);
-						TextureFactory->UdimRegexPattern = MoveTemp(UdimRegexPattern);
+				if (!ResolvedTexturePath.IsEmpty())
+				{
+					FString TextureExtension;
+					if (IsInsideUsdzArchive(ResolvedTexturePath, TextureExtension))
+					{
+						// Always prefer using the TextureFactory if we can, as it may provide compression, which the runtime version never will
+						Texture = ReadTextureFromUsdzArchiveEditor(ResolvedTexturePath, TextureExtension, TextureFactory, Outer, ObjectFlags);
 					}
-
-					const FString ResolvedTexturePath = UsdUtils::GetResolvedAssetPath(TextureAssetPathAttr);
-					if (!ResolvedTexturePath.IsEmpty())
+					// Not inside an USDZ archive, just a regular texture
+					else
 					{
-						EObjectFlags ObjectFlags = RF_Transactional | RF_Transient;
-						if (!Outer)
-						{
-							Outer = GetTransientPackage();
-						}
-
-						FString TextureExtension;
-						if (IsInsideUsdzArchive(ResolvedTexturePath, TextureExtension))
-						{
-							// Always prefer using the TextureFactory if we can, as it may provide compression, which the runtime version never will
-							Texture = ReadTextureFromUsdzArchiveEditor(ResolvedTexturePath, TextureExtension, TextureFactory, Outer, ObjectFlags);
-						}
-						// Not inside an USDZ archive, just a regular texture
-						else
-						{
-							FName TextureName = MakeUniqueObjectName(
-								Outer,
-								UTexture::StaticClass(),
-								*IUsdClassesModule::SanitizeObjectName(FPaths::GetBaseFilename(ResolvedTexturePath))
-							);
-							Texture = Cast<UTexture>(TextureFactory->ImportObject(
-								UTexture::StaticClass(),
-								Outer,
-								TextureName,
-								ObjectFlags,
-								ResolvedTexturePath,
-								TEXT(""),
-								bOutCancelled
-							));
-						}
+						bool bOutCancelled = false;
+						const TCHAR* Params = TEXT("");
+						Texture = Cast<UTexture>(TextureFactory->ImportObject(	  //
+							UTexture::StaticClass(),
+							Outer,
+							SanitizedName,
+							ObjectFlags,
+							ResolvedTexturePath,
+							Params,
+							bOutCancelled
+						));
+					}
 
-						if (Texture)
-						{
-							UUsdAssetUserData* UserData = NewObject<UUsdAssetUserData>(Texture, TEXT("USDAssetUserData"));
-							UserData->PrimPaths = {PrimPath};
-							Texture->AddAssetUserData(UserData);
-
-							// We set this even if we're not going to import so that we can track our original texture filepath
-							// in case we later do an Actions->Import
-							UUsdAssetImportData* ImportData = NewObject<UUsdAssetImportData>(Texture);
-							ImportData->UpdateFilenameOnly(ResolvedTexturePath);
-							Texture->AssetImportData = ImportData;
-						}
+					if (Texture)
+					{
+						// We set this even if we're not going to import so that we can track our original texture filepath
+						// in case we later do an Actions->Import
+						UUsdAssetImportData* ImportData = NewObject<UUsdAssetImportData>(Texture);
+						ImportData->UpdateFilenameOnly(ResolvedTexturePath);
+						Texture->AssetImportData = ImportData;
 					}
 				}
 #endif	  // WITH_EDITOR
 				return Texture;
 			}
 
-			UTexture* CreateTextureAtRuntime(const pxr::UsdAttribute& TextureAssetPathAttr, const FString& PrimPath, TextureGroup LODGroup)
+			UTexture* CreateTextureAtRuntime(const FString& ResolvedTexturePath)
 			{
-				FScopedUsdAllocs UsdAllocs;
-
-				pxr::SdfAssetPath TextureAssetPath;
-				TextureAssetPathAttr.Get<pxr::SdfAssetPath>(&TextureAssetPath);
-
-				FString TexturePath = UsdToUnreal::ConvertString(TextureAssetPath.GetAssetPath());
-				FPaths::NormalizeFilename(TexturePath);
-
 				FScopedUnrealAllocs UnrealAllocs;
 
 				UTexture* Texture = nullptr;
 
-				if (!TexturePath.IsEmpty())
+				if (!ResolvedTexturePath.IsEmpty())
 				{
-					const FString ResolvedTexturePath = UsdUtils::GetResolvedAssetPath(TextureAssetPathAttr);
-					if (!ResolvedTexturePath.IsEmpty())
+					// Try checking if the texture is inside an USDZ archive first, or else TextureFactory throws an error
+					FString TextureExtension;
+					if (IsInsideUsdzArchive(ResolvedTexturePath, TextureExtension))
 					{
-						// Try checking if the texture is inside an USDZ archive first, or else TextureFactory throws an error
-						FString TextureExtension;
-						if (IsInsideUsdzArchive(ResolvedTexturePath, TextureExtension))
-						{
-							Texture = ReadTextureFromUsdzArchiveRuntime(ResolvedTexturePath);
-						}
+						Texture = ReadTextureFromUsdzArchiveRuntime(ResolvedTexturePath);
+					}
 
-						// Not inside an USDZ archive, just a regular texture
-						if (!Texture)
-						{
-							Texture = FImageUtils::ImportFileAsTexture2D(ResolvedTexturePath);
-						}
+					// Not inside an USDZ archive, just a regular texture
+					if (!Texture)
+					{
+						Texture = FImageUtils::ImportFileAsTexture2D(ResolvedTexturePath);
 					}
 				}
 
@@ -1588,11 +1571,11 @@ namespace UE
 
 					// Final FilePath will be something like "C:/TexturesFolder/Game_ContentFolder_Materials_Red_BaseColor.png", which automatically
 					// guarantees it won't overwrite another texture from the same export, but will overwrite old textures from previous exports
-					FString TextureFileName = IUsdClassesModule::SanitizeObjectName(FPaths::ChangeExtension(TextureNamePrefix, TEXT("")));
+					FString TextureFileName = UsdUnreal::ObjectUtils::SanitizeObjectName(FPaths::ChangeExtension(TextureNamePrefix, TEXT("")));
 					TextureFileName.RemoveFromStart(TEXT("_"));
 					FString TextureFilePath = FPaths::Combine(
 						TexturesFolder.Path,
-						FString::Printf(TEXT("%s_%s.exr"), *IUsdClassesModule::SanitizeObjectName(TextureFileName), *TrimmedPropertyName)
+						FString::Printf(TEXT("%s_%s.exr"), *UsdUnreal::ObjectUtils::SanitizeObjectName(TextureFileName), *TrimmedPropertyName)
 					);
 
 					// For some reason the baked samples always have zero alpha and there is nothing we can do about it... It seems like the material
@@ -1650,11 +1633,11 @@ namespace UE
 
 					// Final FilePath will be something like "C:/TexturesFolder/Game_ContentFolder_Materials_Red_BaseColor.png", which automatically
 					// guarantees it won't overwrite another texture from the same export, but will overwrite old textures from previous exports
-					FString TextureFileName = IUsdClassesModule::SanitizeObjectName(FPaths::ChangeExtension(TextureNamePrefix, TEXT("")));
+					FString TextureFileName = UsdUnreal::ObjectUtils::SanitizeObjectName(FPaths::ChangeExtension(TextureNamePrefix, TEXT("")));
 					TextureFileName.RemoveFromStart(TEXT("_"));
 					FString TextureFilePath = FPaths::Combine(
 						TexturesFolder.Path,
-						FString::Printf(TEXT("%s_%s.png"), *IUsdClassesModule::SanitizeObjectName(TextureFileName), *TrimmedPropertyName)
+						FString::Printf(TEXT("%s_%s.png"), *UsdUnreal::ObjectUtils::SanitizeObjectName(TextureFileName), *TrimmedPropertyName)
 					);
 
 					// For some reason the baked samples always have zero alpha and there is nothing we can do about it... It seems like the material
@@ -2154,7 +2137,7 @@ namespace UsdShadeConversionImpl = UE::UsdShadeConversion::Private;
 bool UsdToUnreal::ConvertMaterial(
 	const pxr::UsdShadeMaterial& UsdShadeMaterial,
 	UMaterialInstance& Material,
-	UUsdAssetCache2* TexturesCache,
+	UUsdAssetCache3* TexturesCache,
 	const TCHAR* RenderContext,
 	bool bReuseIdenticalAssets
 )
@@ -2357,7 +2340,7 @@ bool UsdToUnreal::ConvertMaterial(
 bool UsdToUnreal::ConvertMaterial(
 	const pxr::UsdShadeMaterial& UsdShadeMaterial,
 	UMaterial& Material,
-	UUsdAssetCache2* TexturesCache,
+	UUsdAssetCache3* TexturesCache,
 	const TCHAR* RenderContext,
 	bool bReuseIdenticalAssets
 )
@@ -2576,7 +2559,7 @@ bool UsdToUnreal::ConvertMaterial(
 				else
 				{
 					UE_LOG(
-						LogTemp,
+						LogUsd,
 						Warning,
 						TEXT("Failed to find primvar '%s' when setting material parameter. Available primvars and UV indices: %s.%s"),
 						*TextureParameterValue->Primvar,
@@ -2673,34 +2656,10 @@ bool UsdToUnreal::ConvertMaterial(
 #endif	  // WITH_EDITOR
 }
 
-// Deprecated
-bool UsdToUnreal::ConvertMaterial(
-	const pxr::UsdShadeMaterial& UsdShadeMaterial,
-	UMaterialInstance& Material,
-	UUsdAssetCache2* TexturesCache,
-	TMap<FString, int32>& InPrimvarToUVIndex,
-	const TCHAR* RenderContext
-)
-{
-	return ConvertMaterial(UsdShadeMaterial, Material, TexturesCache, RenderContext);
-}
-
-// Deprecated
-bool UsdToUnreal::ConvertMaterial(
-	const pxr::UsdShadeMaterial& UsdShadeMaterial,
-	UMaterial& Material,
-	UUsdAssetCache2* TexturesCache,
-	TMap<FString, int32>& PrimvarToUVIndex,
-	const TCHAR* RenderContext
-)
-{
-	return ConvertMaterial(UsdShadeMaterial, Material, TexturesCache, RenderContext);
-}
-
 bool UsdToUnreal::ConvertShadeInputsToParameters(
 	const pxr::UsdShadeMaterial& UsdShadeMaterial,
 	UMaterialInstance& MaterialInstance,
-	UUsdAssetCache2* TexturesCache,
+	UUsdAssetCache3* TexturesCache,
 	const TCHAR* RenderContext,
 	bool bReuseIdenticalAssets
 )
@@ -2738,6 +2697,11 @@ bool UsdToUnreal::ConvertShadeInputsToParameters(
 			ConnectInput = Source.GetInput(SourceName);
 		}
 
+		if (!ConnectInput)
+		{
+			continue;
+		}
+
 		FString DisplayName = UsdToUnreal::ConvertString(ConnectInput.GetAttr().GetDisplayName());
 
 		if (DisplayName.IsEmpty())
@@ -2809,37 +2773,38 @@ PRAGMA_DISABLE_DEPRECATION_WARNINGS
 bool UsdToUnreal::ConvertMaterial(
 	const pxr::UsdShadeMaterial& UsdShadeMaterial,
 	UMaterialInstance& Material,
-	UUsdAssetCache* TexturesCache,
-	TMap<FString, int32>& PrimvarToUVIndex,
-	const TCHAR* RenderContext
+	UUsdAssetCache2* TexturesCache,
+	const TCHAR* RenderContext,
+	bool bReuseIdenticalAssets
 )
 {
-	return UsdToUnreal::ConvertMaterial(UsdShadeMaterial, Material);
+	UUsdAssetCache3* NewCache = nullptr;
+	return UsdToUnreal::ConvertMaterial(UsdShadeMaterial, Material, NewCache, RenderContext, bReuseIdenticalAssets);
 }
 
 bool UsdToUnreal::ConvertMaterial(
 	const pxr::UsdShadeMaterial& UsdShadeMaterial,
 	UMaterial& Material,
-	UUsdAssetCache* TexturesCache,
-	TMap<FString, int32>& PrimvarToUVIndex,
-	const TCHAR* RenderContext
+	UUsdAssetCache2* TexturesCache,
+	const TCHAR* RenderContext,
+	bool bReuseIdenticalAssets
 )
 {
-	UUsdAssetCache2* NewCache = nullptr;
-	return UsdToUnreal::ConvertMaterial(UsdShadeMaterial, Material, NewCache, RenderContext);
+	UUsdAssetCache3* NewCache = nullptr;
+	return UsdToUnreal::ConvertMaterial(UsdShadeMaterial, Material, NewCache, RenderContext, bReuseIdenticalAssets);
 }
 
 bool UsdToUnreal::ConvertShadeInputsToParameters(
 	const pxr::UsdShadeMaterial& UsdShadeMaterial,
 	UMaterialInstance& MaterialInstance,
-	UUsdAssetCache* TexturesCache,
-	const TCHAR* RenderContext
+	UUsdAssetCache2* TexturesCache,
+	const TCHAR* RenderContext,
+	bool bReuseIdenticalAssets
 )
 {
-	UUsdAssetCache2* NewCache = nullptr;
-	return UsdToUnreal::ConvertShadeInputsToParameters(UsdShadeMaterial, MaterialInstance, NewCache, RenderContext);
+	UUsdAssetCache3* NewCache = nullptr;
+	return UsdToUnreal::ConvertShadeInputsToParameters(UsdShadeMaterial, MaterialInstance, NewCache, RenderContext, bReuseIdenticalAssets);
 }
-
 PRAGMA_ENABLE_DEPRECATION_WARNINGS
 
 #if WITH_EDITOR
@@ -2932,10 +2897,9 @@ FString UsdUtils::GetResolvedAssetPath(const pxr::UsdAttribute& AssetPathAttr, p
 		AssetIdentifier = Resolver.CreateIdentifier(AssetIdentifier);
 	}
 
-	FString ResolvedTexturePath = UsdToUnreal::ConvertString(AssetIdentifier);
-	FPaths::NormalizeFilename(ResolvedTexturePath);
+	FString ResolvedPath = UsdToUnreal::ConvertString(AssetIdentifier);
 
-	if (ResolvedTexturePath.IsEmpty())
+	if (ResolvedPath.IsEmpty())
 	{
 		FString TexturePath = UsdToUnreal::ConvertString(AssetPath.GetAssetPath());
 		FPaths::NormalizeFilename(TexturePath);
@@ -2943,11 +2907,12 @@ FString UsdUtils::GetResolvedAssetPath(const pxr::UsdAttribute& AssetPathAttr, p
 		if (!TexturePath.IsEmpty())
 		{
 			pxr::SdfLayerRefPtr TextureLayer = UsdUtils::FindLayerForAttribute(AssetPathAttr, TimeCode.GetValue());
-			ResolvedTexturePath = UsdShadeConversionImpl::ResolveAssetPath(TextureLayer, TexturePath);
+			ResolvedPath = UsdShadeConversionImpl::ResolveAssetPath(TextureLayer, TexturePath);
 		}
 	}
 
-	return ResolvedTexturePath;
+	FPaths::NormalizeFilename(ResolvedPath);
+	return ResolvedPath;
 }
 
 // Deprecated
@@ -3016,18 +2981,33 @@ FString UsdUtils::GetTextureHash(
 	return LexToString(Hash);
 }
 
-UTexture* UsdUtils::CreateTexture(const pxr::UsdAttribute& TextureAssetPathAttr, const FString& PrimPath, TextureGroup LODGroup, UObject* Outer)
+UTexture* UsdUtils::CreateTexture(const pxr::UsdAttribute& TextureAssetPathAttr, const FString& PrimPath, TextureGroup Group, UObject* Outer)
+{
+	const FString ResolvedTexturePath = UsdUtils::GetResolvedAssetPath(TextureAssetPathAttr);
+
+	FName TextureName = MakeUniqueObjectName(
+		Outer,
+		UTexture::StaticClass(),
+		*UsdUnreal::ObjectUtils::SanitizeObjectName(FPaths::GetBaseFilename(ResolvedTexturePath))
+	);
+
+	EObjectFlags Flags = RF_Public | RF_Standalone;
+
+	return UsdUtils::CreateTexture(ResolvedTexturePath, TextureName, Group, Flags, Outer);
+}
+
+UTexture* UsdUtils::CreateTexture(const FString& ResolvedTexturePath, FName SanitizedName, TextureGroup Group, EObjectFlags Flags, UObject* Outer)
 {
 	TRACE_CPUPROFILER_EVENT_SCOPE(UsdUtils::CreateTexture);
 
 	// Standalone game does have WITH_EDITOR defined, but it can't use the texture factories, so we need to check this instead
 	if (GIsEditor)
 	{
-		return UsdShadeConversionImpl::CreateTextureWithEditor(TextureAssetPathAttr, PrimPath, LODGroup, Outer);
+		return UsdShadeConversionImpl::CreateTextureWithEditor(ResolvedTexturePath, SanitizedName, Group, Flags, Outer);
 	}
 	else
 	{
-		return UsdShadeConversionImpl::CreateTextureAtRuntime(TextureAssetPathAttr, PrimPath, LODGroup);
+		return UsdShadeConversionImpl::CreateTextureAtRuntime(ResolvedTexturePath);
 	}
 }
 
@@ -3342,7 +3322,7 @@ void UsdUtils::AuthorUnrealMaterialBinding(pxr::UsdPrim& MeshOrGeomSubsetPrim, c
 			UsedNames.Add(UsdToUnreal::ConvertToken(Child.GetName()));
 		}
 
-		ChildMaterialName = UsdUtils::GetUniqueName(ChildMaterialName, UsedNames);
+		ChildMaterialName = UsdUnreal::ObjectUtils::GetUniqueName(ChildMaterialName, UsedNames);
 	}
 
 	pxr::UsdStageRefPtr Stage = MeshOrGeomSubsetPrim.GetStage();
@@ -3480,7 +3460,7 @@ bool UsdUtils::IsMaterialTranslucent(const pxr::UsdShadeMaterial& UsdShadeMateri
 	pxr::UsdShadeConnectableAPI Connectable{SurfaceShader};
 
 	const UMaterialInterface* Material = nullptr;
-	UUsdAssetCache2* TexturesCache = nullptr;
+	UUsdAssetCache3* TexturesCache = nullptr;
 	const bool bReuseIdenticalAssets = true;
 	UsdShadeConversionImpl::FParameterValue ParameterValue;
 	bool bHasOpacityConnection = UsdShadeConversionImpl::GetFloatParameterValue(
diff --git a/Engine/Plugins/Importers/USDImporter/Source/USDUtilities/Private/USDSkeletalDataConversion.cpp b/Engine/Plugins/Importers/USDImporter/Source/USDUtilities/Private/USDSkeletalDataConversion.cpp
--- a/Engine/Plugins/Importers/USDImporter/Source/USDUtilities/Private/USDSkeletalDataConversion.cpp
+++ b/Engine/Plugins/Importers/USDImporter/Source/USDUtilities/Private/USDSkeletalDataConversion.cpp
@@ -13,6 +13,7 @@
 #include "USDLayerUtils.h"
 #include "USDLog.h"
 #include "USDMemory.h"
+#include "USDObjectUtils.h"
 #include "USDPrimConversion.h"
 #include "USDShadeConversion.h"
 #include "USDTypesConversion.h"
@@ -1191,13 +1192,13 @@ namespace UnrealToUsdImpl
 
 		// The first bone is the root, and has ParentIndex == -1, so do it separately here to void checking the indices for all bones
 		// Sanitize because ExportName can have spaces, which USD doesn't like
-		OutFullPaths[0] = IUsdClassesModule::SanitizeObjectName(BoneNamesInOrder[0].ExportName);
+		OutFullPaths[0] = UsdUnreal::ObjectUtils::SanitizeObjectName(BoneNamesInOrder[0].ExportName);
 
 		// Bones are always stored in an increasing order, so we can do all paths in a single pass
 		for (int32 BoneIndex = 1; BoneIndex < NumBones; ++BoneIndex)
 		{
 			const FMeshBoneInfo& BoneInfo = BoneNamesInOrder[BoneIndex];
-			FString SanitizedBoneName = IUsdClassesModule::SanitizeObjectName(BoneInfo.ExportName);
+			FString SanitizedBoneName = UsdUnreal::ObjectUtils::SanitizeObjectName(BoneInfo.ExportName);
 
 			OutFullPaths[BoneIndex] = FString::Printf(TEXT("%s/%s"), *OutFullPaths[BoneInfo.ParentIndex], *SanitizedBoneName);
 		}
@@ -1338,7 +1339,7 @@ bool UsdToUnreal::ConvertSkeleton(const pxr::UsdSkelSkeletonQuery& SkeletonQuery
 		const uint64 FirstIndex = 0;
 		const int32 RootParentIndex = INDEX_NONE;
 
-		FString UniqueNewRootName = UsdUtils::GetUniqueName(TEXT("Root"), TSet<FString>{JointNames});
+		FString UniqueNewRootName = UsdUnreal::ObjectUtils::GetUniqueName(TEXT("Root"), TSet<FString>{JointNames});
 
 		JointNames.Insert(UniqueNewRootName, FirstIndex);
 		BoneTransforms.Insert(FTransform::Identity, FirstIndex);
@@ -2642,8 +2643,8 @@ bool UsdToUnreal::ConvertBlendShape(
 	// Note that we can't just use the prim path here and need an index to guarantee uniqueness,
 	// because although the path is usually unique, USD has case sensitive paths and the FNames of the
 	// UMorphTargets are case insensitive
-	FString PrimaryName = UsdUtils::GetUniqueName(
-		IUsdClassesModule::SanitizeObjectName(UsdToUnreal::ConvertString(UsdBlendShape.GetPrim().GetName())),
+	FString PrimaryName = UsdUnreal::ObjectUtils::GetUniqueName(
+		UsdUnreal::ObjectUtils::SanitizeObjectName(UsdToUnreal::ConvertString(UsdBlendShape.GetPrim().GetName())),
 		UsedMorphTargetNames
 	);
 	FString PrimaryPath = UsdToUnreal::ConvertPath(UsdBlendShape.GetPrim().GetPath());
@@ -2677,8 +2678,8 @@ bool UsdToUnreal::ConvertBlendShape(
 
 		FString OrigInbetweenName = UsdToUnreal::ConvertString(Inbetween.GetAttr().GetName());
 		FString InbetweenPath = FString::Printf(TEXT("%s_%s"), *PrimaryPath, *OrigInbetweenName);
-		FString InbetweenName = UsdUtils::GetUniqueName(
-			IUsdClassesModule::SanitizeObjectName(FPaths::GetCleanFilename(InbetweenPath)),
+		FString InbetweenName = UsdUnreal::ObjectUtils::GetUniqueName(
+			UsdUnreal::ObjectUtils::SanitizeObjectName(FPaths::GetCleanFilename(InbetweenPath)),
 			UsedMorphTargetNames
 		);
 
@@ -2745,47 +2746,79 @@ USkeletalMesh* UsdToUnreal::GetSkeletalMeshFromImportData(
 {
 	TRACE_CPUPROFILER_EVENT_SCOPE(UsdToUnreal::GetSkeletalMeshFromImportData);
 
-	if (LODIndexToSkeletalMeshImportData.Num() == 0 || InSkeletonBones.Num() == 0)
-	{
-		return nullptr;
-	}
-
-	// A SkeletalMesh could be retrieved for re-use and updated for animations
-	// For now, create a new USkeletalMesh
-	// Note: Remember to initialize UsedMorphTargetNames with existing morph targets, whenever the SkeletalMesh is reused
 	FName UniqueMeshName = MakeUniqueObjectName(
 		GetTransientPackage(),
 		USkeletalMesh::StaticClass(),
-		*IUsdClassesModule::SanitizeObjectName(MeshName.ToString())
+		*UsdUnreal::ObjectUtils::SanitizeObjectName(MeshName.ToString())
 	);
-	USkeletalMesh* SkeletalMesh = NewObject<USkeletalMesh>(
+	USkeletalMesh* SkeletalMesh = NewObject<USkeletalMesh>(GetTransientPackage(), UniqueMeshName, ObjectFlags);
+
+	// Generate a Skeleton and associate it to the SkeletalMesh
+	FName UniqueSkeletonName = MakeUniqueObjectName(
 		GetTransientPackage(),
-		UniqueMeshName,
-		ObjectFlags | EObjectFlags::RF_Public | EObjectFlags::RF_Transient
+		USkeleton::StaticClass(),
+		*UsdUnreal::ObjectUtils::SanitizeObjectName(SkeletonName.ToString())
 	);
+	USkeleton* Skeleton = NewObject<USkeleton>(GetTransientPackage(), UniqueSkeletonName, ObjectFlags);
+
+	Skeleton->SetPreviewMesh(SkeletalMesh);
+	SkeletalMesh->SetSkeleton(Skeleton);
+
+	bool bSuccess = ConvertSkeletalImportData(LODIndexToSkeletalMeshImportData, InSkeletonBones, InBlendShapesByPath, SkeletalMesh);
+	if (!bSuccess)
+	{
+		SkeletalMesh->MarkAsGarbage();
+		SkeletalMesh = nullptr;
+
+		Skeleton->MarkAsGarbage();
+		Skeleton = nullptr;
+	}
+
+	return SkeletalMesh;
+}
+
+bool UsdToUnreal::ConvertSkeletalImportData(
+	TArray<FSkeletalMeshImportData>& InLODIndexToSkeletalMeshImportData,
+	const TArray<SkeletalMeshImportData::FBone>& InSkeletonBones,
+	UsdUtils::FBlendShapeMap& InBlendShapesByPath,
+	USkeletalMesh* InOutSkeletalMesh
+)
+{
+	TRACE_CPUPROFILER_EVENT_SCOPE(UsdToUnreal::ConvertSkeletalImportData);
+
+	if (!InOutSkeletalMesh || InLODIndexToSkeletalMeshImportData.Num() == 0)
+	{
+		return false;
+	}
+
+	USkeleton* Skeleton = InOutSkeletalMesh->GetSkeleton();
+	if (!Skeleton)
+	{
+		return false;
+	}
 
 	// Process reference skeleton from import data
 	int32 SkeletalDepth = 0;
 	FSkeletalMeshImportData DummyData;
 	DummyData.RefBonesBinary = InSkeletonBones;
-	if (!SkeletalMeshImportUtils::ProcessImportMeshSkeleton(SkeletalMesh->GetSkeleton(), SkeletalMesh->GetRefSkeleton(), SkeletalDepth, DummyData))
+	if (!SkeletalMeshImportUtils::ProcessImportMeshSkeleton(
+			InOutSkeletalMesh->GetSkeleton(),
+			InOutSkeletalMesh->GetRefSkeleton(),
+			SkeletalDepth,
+			DummyData
+		))
 	{
-		return nullptr;
-	}
-	if (SkeletalMesh->GetRefSkeleton().GetRawBoneNum() == 0)
-	{
-		SkeletalMesh->MarkAsGarbage();
-		return nullptr;
+		return false;
 	}
 
 	// This prevents PostEditChange calls when it is alive, also ensuring it is called once when we return from this function.
 	// This is required because we must ensure the morphtargets are in the SkeletalMesh before the first call to PostEditChange(),
 	// or else they will be effectively discarded
-	FScopedSkeletalMeshPostEditChange ScopedPostEditChange(SkeletalMesh);
-	SkeletalMesh->PreEditChange(nullptr);
+	FScopedSkeletalMeshPostEditChange ScopedPostEditChange(InOutSkeletalMesh);
+	InOutSkeletalMesh->PreEditChange(nullptr);
 
 	// Create initial bounding box based on expanded version of reference pose for meshes without physics assets
-	const FSkeletalMeshImportData& LowestLOD = LODIndexToSkeletalMeshImportData[0];
+	const FSkeletalMeshImportData& LowestLOD = InLODIndexToSkeletalMeshImportData[0];
 	FBox3f BoundingBox(LowestLOD.Points.GetData(), LowestLOD.Points.Num());
 	FBox3f Temp = BoundingBox;
 	FVector3f MidMesh = 0.5f * (Temp.Min + Temp.Max);
@@ -2796,20 +2829,20 @@ USkeletalMesh* UsdToUnreal::GetSkeletalMeshFromImportData(
 	if (LowestLOD.Points.Num() > 2 && BoundingBoxSize.X < THRESH_POINTS_ARE_SAME && BoundingBoxSize.Y < THRESH_POINTS_ARE_SAME
 		&& BoundingBoxSize.Z < THRESH_POINTS_ARE_SAME)
 	{
-		return nullptr;
+		return false;
 	}
 
 #if WITH_EDITOR
 	IMeshUtilities& MeshUtilities = FModuleManager::Get().LoadModuleChecked<IMeshUtilities>("MeshUtilities");
 #endif	  // WITH_EDITOR
 
-	FSkeletalMeshModel* ImportedResource = SkeletalMesh->GetImportedModel();
+	FSkeletalMeshModel* ImportedResource = InOutSkeletalMesh->GetImportedModel();
 	ImportedResource->LODModels.Empty();
-	SkeletalMesh->ResetLODInfo();
+	InOutSkeletalMesh->ResetLODInfo();
 	bool bHasVertexColors = false;
-	for (int32 LODIndex = 0; LODIndex < LODIndexToSkeletalMeshImportData.Num(); ++LODIndex)
+	for (int32 LODIndex = 0; LODIndex < InLODIndexToSkeletalMeshImportData.Num(); ++LODIndex)
 	{
-		FSkeletalMeshImportData& LODImportData = LODIndexToSkeletalMeshImportData[LODIndex];
+		FSkeletalMeshImportData& LODImportData = InLODIndexToSkeletalMeshImportData[LODIndex];
 
 		// In the future it will be expected for bone data to be inside FSkeletalMeshImportData as well so we should
 		// probably do this
@@ -2819,9 +2852,9 @@ USkeletalMesh* UsdToUnreal::GetSkeletalMeshFromImportData(
 		FSkeletalMeshLODModel& LODModel = ImportedResource->LODModels.Last();
 
 		// Process bones influence (normalization and optimization) (optional)
-		SkeletalMeshImportUtils::ProcessImportMeshInfluences(LODImportData, SkeletalMesh->GetPathName());
+		SkeletalMeshImportUtils::ProcessImportMeshInfluences(LODImportData, InOutSkeletalMesh->GetPathName());
 
-		FSkeletalMeshLODInfo& NewLODInfo = SkeletalMesh->AddLODInfo();
+		FSkeletalMeshLODInfo& NewLODInfo = InOutSkeletalMesh->AddLODInfo();
 		NewLODInfo.ReductionSettings.NumOfTrianglesPercentage = 1.0f;
 		NewLODInfo.ReductionSettings.NumOfVertPercentage = 1.0f;
 		NewLODInfo.ReductionSettings.MaxDeviationPercentage = 0.0f;
@@ -2858,8 +2891,8 @@ USkeletalMesh* UsdToUnreal::GetSkeletalMeshFromImportData(
 
 		bool bBuildSuccess = MeshUtilities.BuildSkeletalMesh(
 			LODModel,
-			SkeletalMesh->GetPathName(),
-			SkeletalMesh->GetRefSkeleton(),
+			InOutSkeletalMesh->GetPathName(),
+			InOutSkeletalMesh->GetRefSkeleton(),
 			LODInfluences,
 			LODWedges,
 			LODFaces,
@@ -2887,8 +2920,7 @@ USkeletalMesh* UsdToUnreal::GetSkeletalMeshFromImportData(
 
 		if (!bBuildSuccess)
 		{
-			SkeletalMesh->MarkAsGarbage();
-			return nullptr;
+			return false;
 		}
 
 		// We must also provide the ImportData with morph target information now.
@@ -2925,32 +2957,23 @@ USkeletalMesh* UsdToUnreal::GetSkeletalMeshFromImportData(
 		// This is important because it will fill in the LODModel's RawSkeletalMeshBulkDataID,
 		// which is the part of the skeletal mesh's DDC key that is affected by the actual mesh data
 		PRAGMA_DISABLE_DEPRECATION_WARNINGS
-		SkeletalMesh->SaveLODImportedData(LODIndex, LODImportData);
+		InOutSkeletalMesh->SaveLODImportedData(LODIndex, LODImportData);
 		PRAGMA_ENABLE_DEPRECATION_WARNINGS
 #endif	  // WITH_EDITOR
 	}
 
-	SkeletalMesh->SetImportedBounds(FBoxSphereBounds((FBox)BoundingBox));
-	SkeletalMesh->SetHasVertexColors(bHasVertexColors);
-	SkeletalMesh->SetVertexColorGuid(SkeletalMesh->GetHasVertexColors() ? FGuid::NewGuid() : FGuid());
-	SkeletalMesh->CalculateInvRefMatrices();
+	InOutSkeletalMesh->SetImportedBounds(FBoxSphereBounds((FBox)BoundingBox));
+	InOutSkeletalMesh->SetHasVertexColors(bHasVertexColors);
+	InOutSkeletalMesh->SetVertexColorGuid(InOutSkeletalMesh->GetHasVertexColors() ? FGuid::NewGuid() : FGuid());
+	InOutSkeletalMesh->CalculateInvRefMatrices();
 
-	// Generate a Skeleton and associate it to the SkeletalMesh
-	FName UniqueSkeletonName = MakeUniqueObjectName(
-		GetTransientPackage(),
-		USkeleton::StaticClass(),
-		*IUsdClassesModule::SanitizeObjectName(SkeletonName.ToString())
-	);
-	USkeleton* Skeleton = NewObject<USkeleton>(
-		GetTransientPackage(),
-		UniqueSkeletonName,
-		ObjectFlags | EObjectFlags::RF_Public | EObjectFlags::RF_Transient
-	);
-	Skeleton->MergeAllBonesToBoneTree(SkeletalMesh);
-	Skeleton->SetPreviewMesh(SkeletalMesh);
-	SkeletalMesh->SetSkeleton(Skeleton);
+	Skeleton->MergeAllBonesToBoneTree(InOutSkeletalMesh);
+	if (InOutSkeletalMesh->GetRefSkeleton().GetRawBoneNum() == 0)
+	{
+		return false;
+	}
 
-	UsdToUnrealImpl::CreateMorphTargets(InBlendShapesByPath, LODIndexToSkeletalMeshImportData, SkeletalMesh);
+	UsdToUnrealImpl::CreateMorphTargets(InBlendShapesByPath, InLODIndexToSkeletalMeshImportData, InOutSkeletalMesh);
 
 	// "Declare" the morph target curves on the skeleton or skeletal mesh according to bAddCurveMetadataToSkeleton.
 	// This is important otherwise the ControlRig will not hoist these curves as controls when using e.g. FKControlRig.
@@ -2967,11 +2990,11 @@ USkeletalMesh* UsdToUnreal::GetSkeletalMeshFromImportData(
 		}
 		else
 		{
-			UAnimCurveMetaData* AnimCurveMetaData = SkeletalMesh->GetAssetUserData<UAnimCurveMetaData>();
+			UAnimCurveMetaData* AnimCurveMetaData = InOutSkeletalMesh->GetAssetUserData<UAnimCurveMetaData>();
 			if (AnimCurveMetaData == nullptr)
 			{
-				AnimCurveMetaData = NewObject<UAnimCurveMetaData>(SkeletalMesh, NAME_None, RF_Transactional);
-				SkeletalMesh->AddAssetUserData(AnimCurveMetaData);
+				AnimCurveMetaData = NewObject<UAnimCurveMetaData>(InOutSkeletalMesh, NAME_None, RF_Transactional);
+				InOutSkeletalMesh->AddAssetUserData(AnimCurveMetaData);
 			}
 
 			AnimCurveMetaData->AddCurveMetaData(CurveName);
@@ -2984,7 +3007,7 @@ USkeletalMesh* UsdToUnreal::GetSkeletalMeshFromImportData(
 		}
 	}
 
-	return SkeletalMesh;
+	return true;
 }
 
 #endif	  // #if USE_USD_SDK && WITH_EDITOR
diff --git a/Engine/Plugins/Importers/USDImporter/Source/USDUtilities/Public/USDConversionUtils.h b/Engine/Plugins/Importers/USDImporter/Source/USDUtilities/Public/USDConversionUtils.h
--- a/Engine/Plugins/Importers/USDImporter/Source/USDUtilities/Public/USDConversionUtils.h
+++ b/Engine/Plugins/Importers/USDImporter/Source/USDUtilities/Public/USDConversionUtils.h
@@ -331,40 +331,31 @@ namespace UsdUtils
 	/** Returns the earliest possible timecode. Use it to always fetch the first frame of an animated attribute */
 	USDUTILITIES_API double GetEarliestTimeCode();
 
-	/**
-	 * Utilities to allow getting and setting our AssetImportData to an asset from a base UObject*.
-	 * Note that not all asset types support AssetImportData, and in some cases when retrieving it for e.g. a Skeleton,
-	 * we'll actually check it's preview mesh instead (since Skeletons don't have AssetImportData). The setter won't do
-	 * anything if you try setting asset import data on e.g. a Skeleton, on the other hand.
-	 */
+	UE_DEPRECATED(5.5, "This function has been moved to the UsdUnreal::ObjectUtils namespace, from USDObjectUtils.h")
 	USDUTILITIES_API UUsdAssetImportData* GetAssetImportData(UObject* Asset);
+	UE_DEPRECATED(5.5, "This function has been moved to the UsdUnreal::ObjectUtils namespace, from USDObjectUtils.h")
 	USDUTILITIES_API void SetAssetImportData(UObject* Asset, UAssetImportData* ImportData);
-
-	/**
-	 * Returns the object's UsdAssetUserData of a particular subclass if it has one
-	 */
+	UE_DEPRECATED(5.5, "This function has been moved to the UsdUnreal::ObjectUtils namespace, from USDObjectUtils.h")
 	USDUTILITIES_API UUsdAssetUserData* GetAssetUserData(const UObject* Object, TSubclassOf<UUsdAssetUserData> Class = {});
-
 	template<typename T>
+	UE_DEPRECATED(5.5, "This function has been moved to the UsdUnreal::ObjectUtils namespace, from USDObjectUtils.h")
 	inline T* GetAssetUserData(UObject* Object)
 	{
+		PRAGMA_DISABLE_DEPRECATION_WARNINGS
 		return Cast<T>(GetAssetUserData(Object, T::StaticClass()));
+		PRAGMA_ENABLE_DEPRECATION_WARNINGS
 	}
-
-	/**
-	 * Makes sure Object has an instance of UUsdAssetUserData of the provided subclass (defaulting to just UUsdAssetUserData itself) and returns it
-	 */
+	UE_DEPRECATED(5.5, "This function has been moved to the UsdUnreal::ObjectUtils namespace, from USDObjectUtils.h")
 	USDUTILITIES_API UUsdAssetUserData* GetOrCreateAssetUserData(UObject* Object, TSubclassOf<UUsdAssetUserData> Class = {});
-
 	template<typename T>
+	UE_DEPRECATED(5.5, "This function has been moved to the UsdUnreal::ObjectUtils namespace, from USDObjectUtils.h")
 	inline T* GetOrCreateAssetUserData(UObject* Object)
 	{
+		PRAGMA_DISABLE_DEPRECATION_WARNINGS
 		return Cast<T>(GetOrCreateAssetUserData(Object, T::StaticClass()));
+		PRAGMA_ENABLE_DEPRECATION_WARNINGS
 	}
-
-	/**
-	 * Removes all other UUsdAssetUserData instances from Object if they exist, then sets AssetUserData as Object's single UUsdAssetUserData
-	 */
+	UE_DEPRECATED(5.5, "This function has been moved to the UsdUnreal::ObjectUtils namespace, from USDObjectUtils.h")
 	USDUTILITIES_API bool SetAssetUserData(UObject* Object, UUsdAssetUserData* AssetUserData);
 
 #if USE_USD_SDK
@@ -414,6 +405,7 @@ namespace UsdUtils
 	 * Removes any numbered suffix, followed by any number of underscores (e.g. Asset_2, Asset__232_31 or Asset94 all become 'Asset'), making
 	 * sure the string is kept at least one character long. Returns true if it removed anything.
 	 */
+	 UE_DEPRECATED(5.5, "This function has been moved to USDObjectUtils.h, within the USDClasses module")
 	USDUTILITIES_API bool RemoveNumberedSuffix(FString& Prefix);
 
 	/**
@@ -423,6 +415,7 @@ namespace UsdUtils
 	 * @param UsedNames - Strings that cannot be used for the result
 	 * @return Modified Name so that it doesn't match anything in UsedNames (e.g. "MyName" again, or "MyName_0" or "MyName_423")
 	 */
+	 UE_DEPRECATED(5.5, "This function has been moved to USDObjectUtils.h, within the USDClasses module")
 	USDUTILITIES_API FString GetUniqueName(FString Name, const TSet<FString>& UsedNames);
 
 #if USE_USD_SDK
diff --git a/Engine/Plugins/Importers/USDImporter/Source/USDUtilities/Public/USDLightConversion.h b/Engine/Plugins/Importers/USDImporter/Source/USDUtilities/Public/USDLightConversion.h
--- a/Engine/Plugins/Importers/USDImporter/Source/USDUtilities/Public/USDLightConversion.h
+++ b/Engine/Plugins/Importers/USDImporter/Source/USDUtilities/Public/USDLightConversion.h
@@ -31,6 +31,7 @@ PXR_NAMESPACE_OPEN_SCOPE
 PXR_NAMESPACE_CLOSE_SCOPE
 
 class UUsdAssetCache2;
+class UUsdAssetCache3;
 class UDirectionalLightComponent;
 class ULightComponentBase;
 class UPointLightComponent;
@@ -81,12 +82,6 @@ namespace UsdToUnreal
 		UPointLightComponent& LightComponent,
 		double UsdTimeCode = UsdUtils::GetDefaultTimeCode()
 	);
-	USDUTILITIES_API bool ConvertDomeLight(
-		const pxr::UsdPrim& Prim,
-		USkyLightComponent& LightComponent,
-		UUsdAssetCache2* TexturesCache,
-		bool bReuseIdenticalAssets = true
-	);
 	USDUTILITIES_API bool ConvertLuxShapingAPI(
 		const pxr::UsdPrim& Prim,
 		USpotLightComponent& LightComponent,
@@ -94,8 +89,13 @@ namespace UsdToUnreal
 	);
 
 	PRAGMA_DISABLE_DEPRECATION_WARNINGS
-	UE_DEPRECATED(5.2, "Use the other overload that receives an UUsdAssetCache2 object instead")
-	USDUTILITIES_API bool ConvertDomeLight(const pxr::UsdPrim& Prim, USkyLightComponent& LightComponent, UUsdAssetCache* TexturesCache);
+	UE_DEPRECATED(5.5, "No longer used: The assets are created and assigned by the FUsdLuxLightTranslator directly")
+	USDUTILITIES_API bool ConvertDomeLight(
+		const pxr::UsdPrim& Prim,
+		USkyLightComponent& LightComponent,
+		UUsdAssetCache2* TexturesCache,
+		bool bReuseIdenticalAssets = true
+	);
 	PRAGMA_ENABLE_DEPRECATION_WARNINGS
 
 	// These are separately exposed so that they can be reused when reading data into MovieScene tracks.
diff --git a/Engine/Plugins/Importers/USDImporter/Source/USDUtilities/Public/USDPrimConversion.h b/Engine/Plugins/Importers/USDImporter/Source/USDUtilities/Public/USDPrimConversion.h
--- a/Engine/Plugins/Importers/USDImporter/Source/USDUtilities/Public/USDPrimConversion.h
+++ b/Engine/Plugins/Importers/USDImporter/Source/USDUtilities/Public/USDPrimConversion.h
@@ -45,7 +45,6 @@ class UMovieSceneTrack;
 class UMovieSceneVisibilityTrack;
 class USceneComponent;
 class USkeletalMeshComponent;
-class UUsdAssetCache2;
 class UUsdDrawModeComponent;
 enum ERichCurveInterpMode : int;
 struct FFrameRate;
