USD: Implements UsdAssetCache3 which should be a simpler and more flexible way of managing USD assets;

diff --git a/Engine/Plugins/Importers/USDImporter/Source/USDUtilities/Public/USDShadeConversion.h b/Engine/Plugins/Importers/USDImporter/Source/USDUtilities/Public/USDShadeConversion.h
--- a/Engine/Plugins/Importers/USDImporter/Source/USDUtilities/Public/USDShadeConversion.h
+++ b/Engine/Plugins/Importers/USDImporter/Source/USDUtilities/Public/USDShadeConversion.h
@@ -29,8 +29,8 @@ class FSHAHash;
 class UMaterial;
 class UMaterialOptions;
 class UTexture;
-class UUsdAssetCache;
 class UUsdAssetCache2;
+class UUsdAssetCache3;
 enum class EFlattenMaterialProperties : uint8;
 enum EMaterialProperty : int;
 enum TextureAddress : int;
@@ -61,14 +61,14 @@ namespace UsdToUnreal
 	USDUTILITIES_API bool ConvertMaterial(
 		const pxr::UsdShadeMaterial& UsdShadeMaterial,
 		UMaterialInstance& Material,
-		UUsdAssetCache2* TexturesCache = nullptr,
+		UUsdAssetCache3* TexturesCache = nullptr,
 		const TCHAR* RenderContext = nullptr,
 		bool bReuseIdenticalAssets = true
 	);
 	USDUTILITIES_API bool ConvertMaterial(
 		const pxr::UsdShadeMaterial& UsdShadeMaterial,
 		UMaterial& Material,
-		UUsdAssetCache2* TexturesCache = nullptr,
+		UUsdAssetCache3* TexturesCache = nullptr,
 		const TCHAR* RenderContext = nullptr,
 		bool bReuseIdenticalAssets = true
 	);
@@ -87,50 +87,35 @@ namespace UsdToUnreal
 	USDUTILITIES_API bool ConvertShadeInputsToParameters(
 		const pxr::UsdShadeMaterial& UsdShadeMaterial,
 		UMaterialInstance& MaterialInstance,
-		UUsdAssetCache2* TexturesCache,
+		UUsdAssetCache3* TexturesCache,
 		const TCHAR* RenderContext = nullptr,
 		bool bReuseIdenticalAssets = true
 	);
 
 	PRAGMA_DISABLE_DEPRECATION_WARNINGS
-	UE_DEPRECATED(5.3, "Please use the overload that doesn't use PrimvarToUVIndex anymore")
-	USDUTILITIES_API bool ConvertMaterial(
-		const pxr::UsdShadeMaterial& UsdShadeMaterial,
-		UMaterial& Material,
-		UUsdAssetCache2* TexturesCache,
-		TMap<FString, int32>& PrimvarToUVIndex,
-		const TCHAR* RenderContext = nullptr
-	);
-	UE_DEPRECATED(5.3, "Use the other overload that receives an UUsdAssetCache2 object instead")
+	UE_DEPRECATED(5.5, "Use the other overload that receives an UUsdAssetCache3 object instead")
 	USDUTILITIES_API bool ConvertMaterial(
 		const pxr::UsdShadeMaterial& UsdShadeMaterial,
 		UMaterialInstance& Material,
-		UUsdAssetCache2* TexturesCache,
-		TMap<FString, int32>& PrimvarToUVIndex,
-		const TCHAR* RenderContext = nullptr
-	);
-	UE_DEPRECATED(5.2, "Use the other overload that receives an UUsdAssetCache2 object instead")
-	USDUTILITIES_API bool ConvertMaterial(
-		const pxr::UsdShadeMaterial& UsdShadeMaterial,
-		UMaterialInstance& Material,
-		UUsdAssetCache* TexturesCache,
-		TMap<FString, int32>& PrimvarToUVIndex,
-		const TCHAR* RenderContext = nullptr
+		UUsdAssetCache2* TexturesCache = nullptr,
+		const TCHAR* RenderContext = nullptr,
+		bool bReuseIdenticalAssets = true
 	);
-	UE_DEPRECATED(5.2, "Use the other overload that receives an UUsdAssetCache2 object instead")
+	UE_DEPRECATED(5.5, "Use the other overload that receives an UUsdAssetCache3 object instead")
 	USDUTILITIES_API bool ConvertMaterial(
 		const pxr::UsdShadeMaterial& UsdShadeMaterial,
 		UMaterial& Material,
-		UUsdAssetCache* TexturesCache,
-		TMap<FString, int32>& PrimvarToUVIndex,
-		const TCHAR* RenderContext = nullptr
+		UUsdAssetCache2* TexturesCache = nullptr,
+		const TCHAR* RenderContext = nullptr,
+		bool bReuseIdenticalAssets = true
 	);
-	UE_DEPRECATED(5.2, "Use the other overload that receives an UUsdAssetCache2 object instead")
+	UE_DEPRECATED(5.5, "Use the other overload that receives an UUsdAssetCache3 object instead")
 	USDUTILITIES_API bool ConvertShadeInputsToParameters(
 		const pxr::UsdShadeMaterial& UsdShadeMaterial,
 		UMaterialInstance& MaterialInstance,
-		UUsdAssetCache* TexturesCache,
-		const TCHAR* RenderContext = nullptr
+		UUsdAssetCache2* TexturesCache,
+		const TCHAR* RenderContext = nullptr,
+		bool bReuseIdenticalAssets = true
 	);
 	PRAGMA_ENABLE_DEPRECATION_WARNINGS
 }
@@ -248,12 +233,19 @@ namespace UsdUtils
 		TextureAddress AddressY
 	);
 
-	/** Creates a texture from a pxr::SdfAssetPath attribute. PrimPath is optional, and should point to the source shadematerial prim path. It will be
-	 * placed in its UUsdAssetUserData */
+	UE_DEPRECATED(5.5, "Use the overload that receives the resolved texture path directly.")
 	USDUTILITIES_API UTexture* CreateTexture(
 		const pxr::UsdAttribute& TextureAssetPathAttr,
 		const FString& PrimPath = FString(),
-		TextureGroup LODGroup = TEXTUREGROUP_World,
+		TextureGroup Group = TEXTUREGROUP_World,
+		UObject* Outer = GetTransientPackage()
+	);
+
+	USDUTILITIES_API UTexture* CreateTexture(
+		const FString& ResolvedTexturePath,
+		FName SanitizedName,
+		TextureGroup Group = TEXTUREGROUP_World,
+		EObjectFlags Flags = RF_Public | RF_Standalone | RF_Transactional,
 		UObject* Outer = GetTransientPackage()
 	);
 
diff --git a/Engine/Plugins/Importers/USDImporter/Source/USDUtilities/Public/USDSkeletalDataConversion.h b/Engine/Plugins/Importers/USDImporter/Source/USDUtilities/Public/USDSkeletalDataConversion.h
--- a/Engine/Plugins/Importers/USDImporter/Source/USDUtilities/Public/USDSkeletalDataConversion.h
+++ b/Engine/Plugins/Importers/USDImporter/Source/USDUtilities/Public/USDSkeletalDataConversion.h
@@ -281,6 +281,7 @@ namespace UsdToUnreal
 	 * @param SkeletonName - Name to use for the new USkeleton asset
 	 * @return Newly created USkeletalMesh, or nullptr in case of failure
 	 */
+	UE_DEPRECATED(5.5, "Prefer manually creating/reusing the SkeletalMesh and Skeleton beforehand and then calling ConvertSkeletalImportData.")
 	USDUTILITIES_API USkeletalMesh* GetSkeletalMeshFromImportData(
 		TArray<FSkeletalMeshImportData>& LODIndexToSkeletalMeshImportData,
 		const TArray<SkeletalMeshImportData::FBone>& InSkeletonBones,
@@ -289,6 +290,21 @@ namespace UsdToUnreal
 		const FName& MeshName = NAME_None,
 		const FName& SkeletonName = NAME_None
 	);
+
+	/**
+	 * Fills in InOutSkeletalMesh (and its Skeleton) with all the provided data converted from USD.
+	 * @param InLODIndexToSkeletalMeshImportData - Container with the imported skeletal mesh data per LOD level
+	 * @param InSkeletonBones - Bones to use for the reference skeleton (skeleton data on each LODIndexToSkeletalMeshImportData will be ignored).
+	 * @param InBlendShapesByPath - Blend shapes to convert to morph targets
+	 * @param InOutSkeletalMesh - SkeletalMesh to receive the converted data. It should have a valid USkeleton set.
+	 * @return Whether the conversion was successful or not.
+	 */
+	USDUTILITIES_API bool ConvertSkeletalImportData(
+		TArray<FSkeletalMeshImportData>& InLODIndexToSkeletalMeshImportData,
+		const TArray<SkeletalMeshImportData::FBone>& InSkeletonBones,
+		UsdUtils::FBlendShapeMap& InBlendShapesByPath,
+		USkeletalMesh* InOutSkeletalMesh
+	);
 }
 
 namespace UnrealToUsd
