USD: Fix duplicate LOCTEXT keys.

diff --git a/Engine/Plugins/Importers/USDImporter/Source/USDClasses/Private/USDAssetCache3.cpp b/Engine/Plugins/Importers/USDImporter/Source/USDClasses/Private/USDAssetCache3.cpp
--- a/Engine/Plugins/Importers/USDImporter/Source/USDClasses/Private/USDAssetCache3.cpp
+++ b/Engine/Plugins/Importers/USDImporter/Source/USDClasses/Private/USDAssetCache3.cpp
@@ -127,7 +127,7 @@ namespace UE::USDAssetCache3::Private
 		TMap<UObject*, TSet<UObject*>> ReferencerToReferenced;
 		{
 			OuterProgress.EnterProgressFrame(95.0f);
-			FScopedSlowTask Progress(AllObjectsToDelete.Num(), LOCTEXT("FindingReferences", "Finding references"));
+			FScopedSlowTask Progress(AllObjectsToDelete.Num(), LOCTEXT("FindingReferencesOuter", "Finding references"));
 
 #if WITH_EDITOR
 			if (GEditor && GEditor->Trans)
