USD: Fixed issue with texture creation. The sRGB setting is set at creation instead of being updated after creation.

diff --git a/Engine/Plugins/Importers/USDImporter/Source/USDUtilities/Private/USDShadeConversion.cpp b/Engine/Plugins/Importers/USDImporter/Source/USDUtilities/Private/USDShadeConversion.cpp
--- a/Engine/Plugins/Importers/USDImporter/Source/USDUtilities/Private/USDShadeConversion.cpp
+++ b/Engine/Plugins/Importers/USDImporter/Source/USDUtilities/Private/USDShadeConversion.cpp
@@ -719,17 +719,30 @@ namespace UE
 								PrefixedTextureHash,
 								DesiredTextureName,
 								DesiredFlags,
-								[&ResolvedTexturePath, Group, TexturesCache](UPackage* Outer, FName SanitizedName, EObjectFlags DesiredFlags)
+								[&ResolvedTexturePath, Group, TexturesCache, bSRGB](UPackage* Outer, FName SanitizedName, EObjectFlags DesiredFlags)
 								{
-									return UsdUtils::CreateTexture(ResolvedTexturePath, SanitizedName, Group, DesiredFlags, Outer);
+									return UsdUtils::CreateTexture(ResolvedTexturePath, SanitizedName, Group, DesiredFlags, Outer, !bSRGB);
 								},
 								&bCreatedTexture
 							);
 
 							if (Texture)
 							{
-								if (bCreatedTexture)
+								// The texture resource needs to be updated only if the following settings have changed from their default values
+								bool bNeedUpdateResource = false;
+								bNeedUpdateResource |= AddressX != TextureAddress::TA_Wrap;
+								bNeedUpdateResource |= AddressY != TextureAddress::TA_Wrap;
+#if !WITH_EDITOR
+								// In editor, these settings are already set on the factory when the texture is created so no further update is needed
+								bNeedUpdateResource |= bSRGB != true;
+								// TC_Normalmap is set via the TEXTUREGROUP_WorldNormalMap Group
+								bNeedUpdateResource |= CompressionSettings != TextureCompressionSettings::TC_Default;
+#endif	// !WITH_EDITOR
+								if (bCreatedTexture && bNeedUpdateResource)
 								{
+#if WITH_EDITOR
+									Texture->PreEditChange(nullptr);
+#endif	// WITH_EDITOR
 									Texture->SRGB = bSRGB;
 									Texture->CompressionSettings = CompressionSettings;
 									if (UTexture2D* Texture2D = Cast<UTexture2D>(Texture))
@@ -737,10 +750,11 @@ namespace UE
 										Texture2D->AddressX = AddressX;
 										Texture2D->AddressY = AddressY;
 									}
-									Texture->UpdateResource();
 #if WITH_EDITOR
 									Texture->PostEditChange();
-#endif	  // WITH_EDITOR
+#else
+									Texture->UpdateResource();
+#endif	// WITH_EDITOR
 								}
 
 								if (UUsdAssetUserData* AssetUserData = UsdUnreal::ObjectUtils::GetOrCreateAssetUserData(Texture))
@@ -1331,7 +1345,8 @@ namespace UE
 				FName SanitizedName,
 				TextureGroup Group,
 				EObjectFlags ObjectFlags,
-				UObject* Outer
+				UObject* Outer,
+				bool bForceLinear
 			)
 			{
 				UTexture* Texture = nullptr;
@@ -1348,6 +1363,8 @@ namespace UE
 				TextureFactory->bUseHashAsGuid = true;
 				TextureFactory->LODGroup = Group;
 				TextureFactory->HDRImportShouldBeLongLatCubeMap = EAppReturnType::YesAll;
+				// To maintain existing behavior, ColorSpaceMode is left at auto if linear is not requested
+				TextureFactory->ColorSpaceMode = bForceLinear ? ETextureSourceColorSpace::Linear : ETextureSourceColorSpace::Auto;
 
 				const bool bIsSupportedUdimTexture = ResolvedTexturePath.Contains(TEXT("<UDIM>"));
 				if (bIsSupportedUdimTexture)
@@ -2996,14 +3013,14 @@ UTexture* UsdUtils::CreateTexture(const pxr::UsdAttribute& TextureAssetPathAttr,
 	return UsdUtils::CreateTexture(ResolvedTexturePath, TextureName, Group, Flags, Outer);
 }
 
-UTexture* UsdUtils::CreateTexture(const FString& ResolvedTexturePath, FName SanitizedName, TextureGroup Group, EObjectFlags Flags, UObject* Outer)
+UTexture* UsdUtils::CreateTexture(const FString& ResolvedTexturePath, FName SanitizedName, TextureGroup Group, EObjectFlags Flags, UObject* Outer, bool bForceLinear)
 {
 	TRACE_CPUPROFILER_EVENT_SCOPE(UsdUtils::CreateTexture);
 
 	// Standalone game does have WITH_EDITOR defined, but it can't use the texture factories, so we need to check this instead
 	if (GIsEditor)
 	{
-		return UsdShadeConversionImpl::CreateTextureWithEditor(ResolvedTexturePath, SanitizedName, Group, Flags, Outer);
+		return UsdShadeConversionImpl::CreateTextureWithEditor(ResolvedTexturePath, SanitizedName, Group, Flags, Outer, bForceLinear);
 	}
 	else
 	{
diff --git a/Engine/Plugins/Importers/USDImporter/Source/USDUtilities/Public/USDShadeConversion.h b/Engine/Plugins/Importers/USDImporter/Source/USDUtilities/Public/USDShadeConversion.h
--- a/Engine/Plugins/Importers/USDImporter/Source/USDUtilities/Public/USDShadeConversion.h
+++ b/Engine/Plugins/Importers/USDImporter/Source/USDUtilities/Public/USDShadeConversion.h
@@ -246,7 +246,8 @@ namespace UsdUtils
 		FName SanitizedName,
 		TextureGroup Group = TEXTUREGROUP_World,
 		EObjectFlags Flags = RF_Public | RF_Standalone | RF_Transactional,
-		UObject* Outer = GetTransientPackage()
+		UObject* Outer = GetTransientPackage(),
+		bool bForceLinear = false
 	);
 
 	/** Checks if this texture needs virtual textures and emits a warning if it is disabled for the project */
