USD: Rename 'bReuseIdenticalAssets' to 'bShareAssetsForIdenticalPrims'.

diff --git a/Engine/Plugins/Importers/USDImporter/Source/USDSchemas/Private/Custom/MDLUSDShadeMaterialTranslator.cpp b/Engine/Plugins/Importers/USDImporter/Source/USDSchemas/Private/Custom/MDLUSDShadeMaterialTranslator.cpp
--- a/Engine/Plugins/Importers/USDImporter/Source/USDSchemas/Private/Custom/MDLUSDShadeMaterialTranslator.cpp
+++ b/Engine/Plugins/Importers/USDImporter/Source/USDSchemas/Private/Custom/MDLUSDShadeMaterialTranslator.cpp
@@ -148,7 +148,7 @@ void FMdlUsdShadeMaterialTranslator::CreateAssets()
 
 		const FString MdlFullName = MdlModuleName + TEXT("::") + MdlDefinitionName;
 		const FString MdlFullInstanceName = MdlFullName + TEXT("_Instance");
-		const FString HashPrefix = UsdUtils::GetAssetHashPrefix(GetPrim(), Context->bReuseIdenticalAssets);
+		const FString HashPrefix = UsdUtils::GetAssetHashPrefix(GetPrim(), Context->bShareAssetsForIdenticalPrims);
 		const FString MaterialHash = HashPrefix + MdlFullName;
 		const FString MdlSearchPath = FPaths::GetPath(Context->Stage.GetRootLayer().GetRealPath());
 
@@ -241,7 +241,7 @@ void FMdlUsdShadeMaterialTranslator::CreateAssets()
 				*MdlMaterialInstance,
 				Context->UsdAssetCache.Get(),
 				*Context->RenderContext.ToString(),
-				Context->bReuseIdenticalAssets
+				Context->bShareAssetsForIdenticalPrims
 			);
 
 			// We can't blindly recreate all component render states when a level is being added, because we may end up first creating
diff --git a/Engine/Plugins/Importers/USDImporter/Source/USDSchemas/Private/Custom/MaterialXUSDShadeMaterialTranslator.cpp b/Engine/Plugins/Importers/USDImporter/Source/USDSchemas/Private/Custom/MaterialXUSDShadeMaterialTranslator.cpp
--- a/Engine/Plugins/Importers/USDImporter/Source/USDSchemas/Private/Custom/MaterialXUSDShadeMaterialTranslator.cpp
+++ b/Engine/Plugins/Importers/USDImporter/Source/USDSchemas/Private/Custom/MaterialXUSDShadeMaterialTranslator.cpp
@@ -168,7 +168,7 @@ namespace UE::USDMaterialXTranslator::Private
 		const pxr::UsdPrim& MaterialXReferencerPrim,
 		UUsdAssetCache3& AssetCache,
 		EObjectFlags ObjectFlags,
-		bool bReuseIdenticalAssets
+		bool bShareAssetsForIdenticalPrims
 	)
 	{
 		if (!FPaths::FileExists(MaterialXFilePath))
@@ -184,7 +184,7 @@ namespace UE::USDMaterialXTranslator::Private
 		FImportAssetParameters InterchangeParameters;
 		InterchangeParameters.bIsAutomated = true;
 
-		const FString HashPrefix = UsdUtils::GetAssetHashPrefix(MaterialXReferencerPrim, bReuseIdenticalAssets);
+		const FString HashPrefix = UsdUtils::GetAssetHashPrefix(MaterialXReferencerPrim, bShareAssetsForIdenticalPrims);
 
 		// Annoyingly we have to make a new target folder for the interchange import, and then rename all the assets over.
 		// This because for direct imports we'll be dealing with an asset cache that is pointing at the transient package.
@@ -463,7 +463,7 @@ void FMaterialXUsdShadeMaterialTranslator::CreateAssets()
 	}
 
 	FString TargetHashSuffix = TEXT("/") + UsdToUnreal::ConvertString(Prim.GetName());
-	FString TargetHashPrefix = UsdUtils::GetAssetHashPrefix(Prim, Context->bReuseIdenticalAssets);
+	FString TargetHashPrefix = UsdUtils::GetAssetHashPrefix(Prim, Context->bShareAssetsForIdenticalPrims);
 
 	FString FoundMaterialAssetHash;
 	UMaterialInterface* ParsedMaterial = nullptr;
@@ -503,7 +503,7 @@ void FMaterialXUsdShadeMaterialTranslator::CreateAssets()
 				MaterialXReferencerPrim,
 				*Context->UsdAssetCache,
 				Context->ObjectFlags,
-				Context->bReuseIdenticalAssets
+				Context->bShareAssetsForIdenticalPrims
 			);
 
 			if (!bSuccess)
diff --git a/Engine/Plugins/Importers/USDImporter/Source/USDSchemas/Private/MeshTranslationImpl.cpp b/Engine/Plugins/Importers/USDImporter/Source/USDSchemas/Private/MeshTranslationImpl.cpp
--- a/Engine/Plugins/Importers/USDImporter/Source/USDSchemas/Private/MeshTranslationImpl.cpp
+++ b/Engine/Plugins/Importers/USDImporter/Source/USDSchemas/Private/MeshTranslationImpl.cpp
@@ -171,7 +171,7 @@ namespace UE::MeshTranslationImplInternal::Private
 		UUsdAssetCache3* AssetCache,
 		FUsdInfoCache* InfoCache,
 		const FString& MaterialHashPrefix,
-		bool bReuseIdenticalAssets
+		bool bShareAssetsForIdenticalPrims
 	)
 	{
 		UUsdMaterialAssetUserData* MaterialAssetUserData = Material.GetAssetUserData<UUsdMaterialAssetUserData>();
@@ -441,7 +441,7 @@ TMap<const UsdUtils::FUsdPrimMaterialSlot*, UMaterialInterface*> MeshTranslation
 	UUsdAssetCache3& AssetCache,
 	FUsdInfoCache& InfoCache,
 	EObjectFlags Flags,
-	bool bReuseIdenticalAssets
+	bool bShareAssetsForIdenticalPrims
 )
 {
 	FScopedUnrealAllocs Allocs;
@@ -506,7 +506,7 @@ TMap<const UsdUtils::FUsdPrimMaterialSlot*, UMaterialInterface*> MeshTranslation
 						SHA1.GetHash(&Hash.Hash[0]);
 						DisplayColorHash = Hash.ToString();
 					}
-					const FString PrefixedHash = UsdUtils::GetAssetHashPrefix(UsdPrim, bReuseIdenticalAssets) + DisplayColorHash;
+					const FString PrefixedHash = UsdUtils::GetAssetHashPrefix(UsdPrim, bShareAssetsForIdenticalPrims) + DisplayColorHash;
 
 					FString DisplayColorName = FString::Printf(
 						TEXT("DisplayColor%s%s"),
@@ -599,7 +599,7 @@ TMap<const UsdUtils::FUsdPrimMaterialSlot*, UMaterialInterface*> MeshTranslation
 					}
 
 					FString PrefixedMaterialHash = Material ? AssetCache.GetHashForAsset(Material) : FString{};
-					FString HashPrefix = UsdUtils::GetAssetHashPrefix(UsdPrim.GetStage()->GetPrimAtPath(MaterialPrimPath), bReuseIdenticalAssets);
+					FString HashPrefix = UsdUtils::GetAssetHashPrefix(UsdPrim.GetStage()->GetPrimAtPath(MaterialPrimPath), bShareAssetsForIdenticalPrims);
 
 					// Need to create a two-sided material on-demand, *before* we make it compatible:
 					// This because at runtime we can't just set the base property overrides, and just instead create a new
@@ -697,7 +697,7 @@ TMap<const UsdUtils::FUsdPrimMaterialSlot*, UMaterialInterface*> MeshTranslation
 									&AssetCache,
 									&InfoCache,
 									HashPrefix,
-									bReuseIdenticalAssets
+									bShareAssetsForIdenticalPrims
 								);
 
 							if (CompatibleMaterial)
@@ -775,7 +775,7 @@ void MeshTranslationImpl::SetMaterialOverrides(
 	bool bInterpretLODs,
 	const FName& RenderContext,
 	const FName& MaterialPurpose,
-	bool bReuseIdenticalAssets
+	bool bShareAssetsForIdenticalPrims
 )
 {
 	FScopedUsdAllocs Allocs;
@@ -915,7 +915,7 @@ void MeshTranslationImpl::SetMaterialOverrides(
 			AssetCache,
 			InfoCache,
 			Flags,
-			bReuseIdenticalAssets
+			bShareAssetsForIdenticalPrims
 		);
 	}
 
diff --git a/Engine/Plugins/Importers/USDImporter/Source/USDSchemas/Private/MeshTranslationImpl.h b/Engine/Plugins/Importers/USDImporter/Source/USDSchemas/Private/MeshTranslationImpl.h
--- a/Engine/Plugins/Importers/USDImporter/Source/USDSchemas/Private/MeshTranslationImpl.h
+++ b/Engine/Plugins/Importers/USDImporter/Source/USDSchemas/Private/MeshTranslationImpl.h
@@ -36,7 +36,7 @@ namespace MeshTranslationImpl
 		UUsdAssetCache3& AssetCache,
 		FUsdInfoCache& InfoCache,
 		EObjectFlags Flags,
-		bool bReuseIdenticalAssets
+		bool bShareAssetsForIdenticalPrims
 	);
 
 	/**
@@ -54,7 +54,7 @@ namespace MeshTranslationImpl
 		bool bInterpretLODs,
 		const FName& RenderContext,
 		const FName& MaterialPurpose,
-		bool bReuseIdenticalAssets
+		bool bShareAssetsForIdenticalPrims
 	);
 
 	void RecordSourcePrimsForMaterialSlots(
diff --git a/Engine/Plugins/Importers/USDImporter/Source/USDSchemas/Private/USDGeomMeshTranslator.cpp b/Engine/Plugins/Importers/USDImporter/Source/USDSchemas/Private/USDGeomMeshTranslator.cpp
--- a/Engine/Plugins/Importers/USDImporter/Source/USDSchemas/Private/USDGeomMeshTranslator.cpp
+++ b/Engine/Plugins/Importers/USDImporter/Source/USDSchemas/Private/USDGeomMeshTranslator.cpp
@@ -258,7 +258,7 @@ namespace UsdGeomMeshTranslatorImpl
 		FUsdInfoCache* InfoCache,
 		float Time,
 		EObjectFlags Flags,
-		bool bReuseIdenticalAssets
+		bool bShareAssetsForIdenticalPrims
 	)
 	{
 		if (!InfoCache)
@@ -280,7 +280,7 @@ namespace UsdGeomMeshTranslatorImpl
 			AssetCache,
 			*InfoCache,
 			Flags,
-			bReuseIdenticalAssets
+			bShareAssetsForIdenticalPrims
 		);
 
 		uint32 StaticMeshSlotIndex = 0;
@@ -709,7 +709,7 @@ namespace UsdGeomMeshTranslatorImpl
 			SHA1.GetHash(&AllLODHash.Hash[0]);
 		}
 
-		FString PrefixedAssetHash = UsdUtils::GetAssetHashPrefix(Prim, Context.bReuseIdenticalAssets) + AllLODHash.ToString();
+		FString PrefixedAssetHash = UsdUtils::GetAssetHashPrefix(Prim, Context.bShareAssetsForIdenticalPrims) + AllLODHash.ToString();
 
 		FString DesiredName = FPaths::GetBaseFilename(MeshName);
 
@@ -1755,7 +1755,7 @@ void FBuildStaticMeshTaskChain::SetupTasks()
 					   Context->InfoCache.Get(),
 					   Context->Time,
 					   Context->ObjectFlags,
-					   Context->bReuseIdenticalAssets
+					   Context->bShareAssetsForIdenticalPrims
 				   );
 
 #if WITH_EDITOR
diff --git a/Engine/Plugins/Importers/USDImporter/Source/USDSchemas/Private/USDGeomPointInstancerTranslator.cpp b/Engine/Plugins/Importers/USDImporter/Source/USDSchemas/Private/USDGeomPointInstancerTranslator.cpp
--- a/Engine/Plugins/Importers/USDImporter/Source/USDSchemas/Private/USDGeomPointInstancerTranslator.cpp
+++ b/Engine/Plugins/Importers/USDImporter/Source/USDSchemas/Private/USDGeomPointInstancerTranslator.cpp
@@ -546,7 +546,7 @@ void FUsdGeomPointInstancerTranslator::UpdateComponents(USceneComponent* PointIn
 					Context->bAllowInterpretingLODs,
 					Context->RenderContext,
 					Context->MaterialPurpose,
-					Context->bReuseIdenticalAssets
+					Context->bShareAssetsForIdenticalPrims
 				);
 			}
 		}
diff --git a/Engine/Plugins/Importers/USDImporter/Source/USDSchemas/Private/USDGeomPrimitiveTranslator.cpp b/Engine/Plugins/Importers/USDImporter/Source/USDSchemas/Private/USDGeomPrimitiveTranslator.cpp
--- a/Engine/Plugins/Importers/USDImporter/Source/USDSchemas/Private/USDGeomPrimitiveTranslator.cpp
+++ b/Engine/Plugins/Importers/USDImporter/Source/USDSchemas/Private/USDGeomPrimitiveTranslator.cpp
@@ -171,7 +171,7 @@ USceneComponent* FUsdGeomPrimitiveTranslator::CreateComponents()
 					Context->bAllowInterpretingLODs,
 					Context->RenderContext,
 					Context->MaterialPurpose,
-					Context->bReuseIdenticalAssets
+					Context->bShareAssetsForIdenticalPrims
 				);
 			}
 		}
diff --git a/Engine/Plugins/Importers/USDImporter/Source/USDSchemas/Private/USDGeomXformableTranslator.cpp b/Engine/Plugins/Importers/USDImporter/Source/USDSchemas/Private/USDGeomXformableTranslator.cpp
--- a/Engine/Plugins/Importers/USDImporter/Source/USDSchemas/Private/USDGeomXformableTranslator.cpp
+++ b/Engine/Plugins/Importers/USDImporter/Source/USDSchemas/Private/USDGeomXformableTranslator.cpp
@@ -341,7 +341,7 @@ USceneComponent* FUsdGeomXformableTranslator::CreateComponents()
 					Context->bAllowInterpretingLODs,
 					Context->RenderContext,
 					Context->MaterialPurpose,
-					Context->bReuseIdenticalAssets
+					Context->bShareAssetsForIdenticalPrims
 				);
 			}
 		}
@@ -990,7 +990,7 @@ void FUsdGeomXformableTranslator::CreateAlternativeDrawModeAssets(EUsdDrawMode D
 			}
 			else
 			{
-				const FString HashPrefix = UsdUtils::GetAssetHashPrefix(GetPrim(), Context->bReuseIdenticalAssets);
+				const FString HashPrefix = UsdUtils::GetAssetHashPrefix(GetPrim(), Context->bShareAssetsForIdenticalPrims);
 				const FString PrefixedTextureHash = HashPrefix + LexToString(FMD5Hash::HashFile(*ResolvedPath));
 
 				const FString& DesiredTextureName = FPaths::GetBaseFilename(ResolvedPath);
diff --git a/Engine/Plugins/Importers/USDImporter/Source/USDSchemas/Private/USDGeometryCacheTranslator.cpp b/Engine/Plugins/Importers/USDImporter/Source/USDSchemas/Private/USDGeometryCacheTranslator.cpp
--- a/Engine/Plugins/Importers/USDImporter/Source/USDSchemas/Private/USDGeometryCacheTranslator.cpp
+++ b/Engine/Plugins/Importers/USDImporter/Source/USDSchemas/Private/USDGeometryCacheTranslator.cpp
@@ -86,7 +86,7 @@ namespace UsdGeometryCacheTranslatorImpl
 		FUsdInfoCache* InfoCache,
 		float Time,
 		EObjectFlags Flags,
-		bool bReuseIdenticalAssets
+		bool bShareAssetsForIdenticalPrims
 	)
 	{
 		TMap<const UsdUtils::FUsdPrimMaterialSlot*, UMaterialInterface*> ResolvedMaterials = MeshTranslationImpl::ResolveMaterialAssignmentInfo(
@@ -95,7 +95,7 @@ namespace UsdGeometryCacheTranslatorImpl
 			AssetCache,
 			*InfoCache,
 			Flags,
-			bReuseIdenticalAssets
+			bShareAssetsForIdenticalPrims
 		);
 
 		uint32 SlotIndex = 0;
@@ -441,7 +441,7 @@ namespace UsdGeometryCacheTranslatorImpl
 
 		FSHAHash GeoCacheHash;
 		SHA1.GetHash(&GeoCacheHash.Hash[0]);
-		const FString PrefixedGeoCacheHash = UsdUtils::GetAssetHashPrefix(RootPrim, Context->bReuseIdenticalAssets) + GeoCacheHash.ToString();
+		const FString PrefixedGeoCacheHash = UsdUtils::GetAssetHashPrefix(RootPrim, Context->bShareAssetsForIdenticalPrims) + GeoCacheHash.ToString();
 
 		const FString DesiredName = FPaths::GetBaseFilename(RootPrimPath);
 
@@ -926,7 +926,7 @@ void FGeometryCacheCreateAssetsTaskChain::SetupTasks()
 						Context->InfoCache.Get(),
 						Context->Time,
 						Context->ObjectFlags,
-						Context->bReuseIdenticalAssets
+						Context->bShareAssetsForIdenticalPrims
 					);
 				}
 
@@ -1050,7 +1050,7 @@ USceneComponent* FUsdGeometryCacheTranslator::CreateComponents()
 					bAllowInterpretingLODs,
 					Context->RenderContext,
 					Context->MaterialPurpose,
-					Context->bReuseIdenticalAssets
+					Context->bShareAssetsForIdenticalPrims
 				);
 
 				// Check if the prim has the GroomBinding schema and setup the component and assets necessary to bind the groom to the GeometryCache
@@ -1061,7 +1061,7 @@ USceneComponent* FUsdGeometryCacheTranslator::CreateComponents()
 						*Context->UsdAssetCache,
 						*Context->InfoCache,
 						Context->ObjectFlags,
-						Context->bReuseIdenticalAssets
+						Context->bShareAssetsForIdenticalPrims
 					);
 
 					// For the groom binding to work, the GroomComponent must be a child of the SceneComponent
diff --git a/Engine/Plugins/Importers/USDImporter/Source/USDSchemas/Private/USDGroomTranslator.cpp b/Engine/Plugins/Importers/USDImporter/Source/USDSchemas/Private/USDGroomTranslator.cpp
--- a/Engine/Plugins/Importers/USDImporter/Source/USDSchemas/Private/USDGroomTranslator.cpp
+++ b/Engine/Plugins/Importers/USDImporter/Source/USDSchemas/Private/USDGroomTranslator.cpp
@@ -191,7 +191,8 @@ protected:
 
 				FSHAHash SHAHash = UE::UsdGroomTranslator::Private::ComputeHairDescriptionHash(HairDescription, ImportOptions->InterpolationSettings);
 
-				const FString PrefixedAssetHash = UsdUtils::GetAssetHashPrefix(GetPrim(), Context->bReuseIdenticalAssets) + SHAHash.ToString();
+				const FString PrefixedAssetHash = UsdUtils::GetAssetHashPrefix(GetPrim(), Context->bShareAssetsForIdenticalPrims)
+												  + SHAHash.ToString();
 
 				const FString PrimPathString = PrimPath.GetString();
 				const FString DesiredName = FPaths::GetBaseFilename(PrimPathString);
@@ -281,7 +282,7 @@ protected:
 				FSHAHash Hash;
 				SHA1.GetHash(Hash.Hash);
 
-				PrefixedGroomCacheHash = UsdUtils::GetAssetHashPrefix(GetPrim(), Context->bReuseIdenticalAssets) + Hash.ToString();
+				PrefixedGroomCacheHash = UsdUtils::GetAssetHashPrefix(GetPrim(), Context->bShareAssetsForIdenticalPrims) + Hash.ToString();
 				return true;
 			}
 		);
diff --git a/Engine/Plugins/Importers/USDImporter/Source/USDSchemas/Private/USDGroomTranslatorUtils.cpp b/Engine/Plugins/Importers/USDImporter/Source/USDSchemas/Private/USDGroomTranslatorUtils.cpp
--- a/Engine/Plugins/Importers/USDImporter/Source/USDSchemas/Private/USDGroomTranslatorUtils.cpp
+++ b/Engine/Plugins/Importers/USDImporter/Source/USDSchemas/Private/USDGroomTranslatorUtils.cpp
@@ -154,7 +154,7 @@ namespace UsdGroomTranslatorUtils
 		UUsdAssetCache3& AssetCache,
 		FUsdInfoCache& InfoCache,
 		EObjectFlags ObjectFlags,
-		bool bReuseIdenticalAssets
+		bool bShareAssetsForIdenticalPrims
 	)
 	{
 		// At this point, the prim (SkelRoot/Skeleton or GeomMesh) has already been checked to have the GroomBindingAPI,
@@ -194,7 +194,7 @@ namespace UsdGroomTranslatorUtils
 		Settings.TargetMesh = TargetMesh;
 		Settings.SourceMesh = GetGroomBindingSourceMesh(Prim, InfoCache, GroomBindingType);
 
-		const FString PrefixedBindingHash = UsdUtils::GetAssetHashPrefix(Prim, bReuseIdenticalAssets) + ComputeGroomBindingHash(Settings).ToString();
+		const FString PrefixedBindingHash = UsdUtils::GetAssetHashPrefix(Prim, bShareAssetsForIdenticalPrims) + ComputeGroomBindingHash(Settings).ToString();
 
 		const FString GroomBindingPath = FString::Printf(TEXT("%s_groombinding"), *PrimPath);
 
diff --git a/Engine/Plugins/Importers/USDImporter/Source/USDSchemas/Private/USDGroomTranslatorUtils.h b/Engine/Plugins/Importers/USDImporter/Source/USDSchemas/Private/USDGroomTranslatorUtils.h
--- a/Engine/Plugins/Importers/USDImporter/Source/USDSchemas/Private/USDGroomTranslatorUtils.h
+++ b/Engine/Plugins/Importers/USDImporter/Source/USDSchemas/Private/USDGroomTranslatorUtils.h
@@ -32,7 +32,7 @@ namespace UsdGroomTranslatorUtils
 		UUsdAssetCache3& AssetCache,
 		FUsdInfoCache& InfoCache,
 		EObjectFlags ObjectFlags,
-		bool bReuseIdenticalAssets
+		bool bShareAssetsForIdenticalPrims
 	);
 
 	// Set the groom asset targeted by the given prim with GroomBindingAPI on the SceneComponent (with a GroomComponent as a direct child)
diff --git a/Engine/Plugins/Importers/USDImporter/Source/USDSchemas/Private/USDLuxLightTranslator.cpp b/Engine/Plugins/Importers/USDImporter/Source/USDSchemas/Private/USDLuxLightTranslator.cpp
--- a/Engine/Plugins/Importers/USDImporter/Source/USDSchemas/Private/USDLuxLightTranslator.cpp
+++ b/Engine/Plugins/Importers/USDImporter/Source/USDSchemas/Private/USDLuxLightTranslator.cpp
@@ -76,7 +76,7 @@ void FUsdLuxLightTranslator::CreateAssets()
 		return;
 	}
 
-	const FString PrefixedTextureHash = UsdUtils::GetAssetHashPrefix(Prim, Context->bReuseIdenticalAssets)
+	const FString PrefixedTextureHash = UsdUtils::GetAssetHashPrefix(Prim, Context->bShareAssetsForIdenticalPrims)
 										+ LexToString(FMD5Hash::HashFile(*ResolvedDomeTexturePath));
 
 	const FString& DesiredTextureName = FPaths::GetBaseFilename(ResolvedDomeTexturePath);
diff --git a/Engine/Plugins/Importers/USDImporter/Source/USDSchemas/Private/USDShadeMaterialTranslator.cpp b/Engine/Plugins/Importers/USDImporter/Source/USDSchemas/Private/USDShadeMaterialTranslator.cpp
--- a/Engine/Plugins/Importers/USDImporter/Source/USDSchemas/Private/USDShadeMaterialTranslator.cpp
+++ b/Engine/Plugins/Importers/USDImporter/Source/USDSchemas/Private/USDShadeMaterialTranslator.cpp
@@ -345,7 +345,7 @@ void FUsdShadeMaterialTranslator::CreateAssets()
 		SHA1.GetHash(&OutHash.Hash[0]);
 		MaterialHash = OutHash.ToString();
 	}
-	const FString PrefixedMaterialHash = UsdUtils::GetAssetHashPrefix(GetPrim(), Context->bReuseIdenticalAssets) + MaterialHash;
+	const FString PrefixedMaterialHash = UsdUtils::GetAssetHashPrefix(GetPrim(), Context->bShareAssetsForIdenticalPrims) + MaterialHash;
 
 	const FString PrimPathString = PrimPath.GetString();
 	const FString DesiredName = FPaths::GetBaseFilename(PrimPathString);
@@ -373,7 +373,7 @@ void FUsdShadeMaterialTranslator::CreateAssets()
 				*MIC,
 				Context->UsdAssetCache.Get(),
 				*Context->RenderContext.ToString(),
-				Context->bReuseIdenticalAssets
+				Context->bShareAssetsForIdenticalPrims
 			);
 			if (!bSuccess)
 			{
@@ -477,7 +477,7 @@ void FUsdShadeMaterialTranslator::CreateAssets()
 				*MI,
 				Context->UsdAssetCache.Get(),
 				*Context->RenderContext.ToString(),
-				Context->bReuseIdenticalAssets
+				Context->bShareAssetsForIdenticalPrims
 			);
 			if (!bSuccess)
 			{
diff --git a/Engine/Plugins/Importers/USDImporter/Source/USDSchemas/Private/USDSkelRootTranslator.cpp b/Engine/Plugins/Importers/USDImporter/Source/USDSchemas/Private/USDSkelRootTranslator.cpp
--- a/Engine/Plugins/Importers/USDImporter/Source/USDSchemas/Private/USDSkelRootTranslator.cpp
+++ b/Engine/Plugins/Importers/USDImporter/Source/USDSchemas/Private/USDSkelRootTranslator.cpp
@@ -75,7 +75,7 @@ namespace UsdSkelRootTranslatorImpl
 		float Time,
 		EObjectFlags Flags,
 		bool bSkeletalMeshHasMorphTargets,
-		bool bReuseIdenticalAssets
+		bool bShareAssetsForIdenticalPrims
 	)
 	{
 		TRACE_CPUPROFILER_EVENT_SCOPE(UsdSkelRootTranslatorImpl::ProcessMaterials);
@@ -105,7 +105,7 @@ namespace UsdSkelRootTranslatorImpl
 			AssetCache,
 			InfoCache,
 			Flags,
-			bReuseIdenticalAssets
+			bShareAssetsForIdenticalPrims
 		);
 
 		bool bMaterialsHaveChanged = false;
@@ -665,7 +665,7 @@ namespace UsdSkelRootTranslatorImpl
 		bool bInterpretLODs,
 		const FName& RenderContext,
 		const FName& MaterialPurpose,
-		bool bReuseIdenticalAssets
+		bool bShareAssetsForIdenticalPrims
 	)
 	{
 		FScopedUsdAllocs Allocs;
@@ -812,7 +812,7 @@ namespace UsdSkelRootTranslatorImpl
 			AssetCache,
 			InfoCache,
 			Flags,
-			bReuseIdenticalAssets
+			bShareAssetsForIdenticalPrims
 		);
 
 		// Compare resolved materials with existing assignments, and create overrides if we need to
@@ -983,7 +983,7 @@ namespace UsdSkelRootTranslatorImpl
 			SHA1.UpdateWithString(*PrimPath, PrimPath.Len());
 			SHA1.Final();
 			SHA1.GetHash(&Hash.Hash[0]);
-			const FString PrefixedAnimBPHash = UsdUtils::GetAssetHashPrefix(Prim, Context.bReuseIdenticalAssets) + Hash.ToString();
+			const FString PrefixedAnimBPHash = UsdUtils::GetAssetHashPrefix(Prim, Context.bShareAssetsForIdenticalPrims) + Hash.ToString();
 
 			// Check if we can find an AnimBP for this prim in the asset cache (useful when doing Action->Import)
 			bool bReusedAnimBP = false;
@@ -1358,7 +1358,7 @@ namespace UsdSkelRootTranslatorImpl
 
 				FString SkelRootPath = PrimPath.GetString();
 				FSHAHash SkeletalMeshHash = UsdSkelRootTranslatorImpl::ComputeSHAHash(LODIndexToSkeletalMeshImportData, SkeletonBones, BlendShapes);
-				PrefixedSkelMeshHash = UsdUtils::GetAssetHashPrefix(GetPrim(), Context->bReuseIdenticalAssets) + SkeletalMeshHash.ToString();
+				PrefixedSkelMeshHash = UsdUtils::GetAssetHashPrefix(GetPrim(), Context->bShareAssetsForIdenticalPrims) + SkeletalMeshHash.ToString();
 
 				USkeletalMesh* SkeletalMesh = Cast<USkeletalMesh>(Context->UsdAssetCache->GetCachedAsset(PrefixedSkelMeshHash));
 
@@ -1402,7 +1402,7 @@ namespace UsdSkelRootTranslatorImpl
 							Context->Time,
 							Context->ObjectFlags,
 							NewBlendShapes.Num() > 0,
-							Context->bReuseIdenticalAssets
+							Context->bShareAssetsForIdenticalPrims
 						);
 
 						if (bMaterialsHaveChanged)
@@ -1585,7 +1585,7 @@ namespace UsdSkelRootTranslatorImpl
 														  : pxr::UsdPrim{};
 
 						FSHAHash Hash = UsdSkelRootTranslatorImpl::ComputeSHAHash(SkelQuery, RootMotionPrim, PrefixedSkelMeshHash);
-						FString PrefixedSkelAnimHash = UsdUtils::GetAssetHashPrefix(SkelAnimationPrim, Context->bReuseIdenticalAssets)
+						FString PrefixedSkelAnimHash = UsdUtils::GetAssetHashPrefix(SkelAnimationPrim, Context->bShareAssetsForIdenticalPrims)
 													   + Hash.ToString();
 						UAnimSequence* AnimSequence = Cast<UAnimSequence>(Context->UsdAssetCache->GetCachedAsset(PrefixedSkelAnimHash));
 
@@ -1705,7 +1705,7 @@ USceneComponent* FUsdSkelRootTranslator::CreateComponents()
 			*Context->UsdAssetCache,
 			*Context->InfoCache,
 			Context->ObjectFlags,
-			Context->bReuseIdenticalAssets
+			Context->bShareAssetsForIdenticalPrims
 		);
 
 		// For the groom binding to work, the GroomComponent must be a child of the SceneComponent
@@ -1806,7 +1806,7 @@ void FUsdSkelRootTranslator::UpdateComponents(USceneComponent* SceneComponent)
 				Context->bAllowInterpretingLODs,
 				Context->RenderContext,
 				Context->MaterialPurpose,
-				Context->bReuseIdenticalAssets
+				Context->bShareAssetsForIdenticalPrims
 			);
 		}
 	}
diff --git a/Engine/Plugins/Importers/USDImporter/Source/USDSchemas/Private/USDSkelSkeletonTranslator.cpp b/Engine/Plugins/Importers/USDImporter/Source/USDSchemas/Private/USDSkelSkeletonTranslator.cpp
--- a/Engine/Plugins/Importers/USDImporter/Source/USDSchemas/Private/USDSkelSkeletonTranslator.cpp
+++ b/Engine/Plugins/Importers/USDImporter/Source/USDSchemas/Private/USDSkelSkeletonTranslator.cpp
@@ -85,7 +85,7 @@ namespace UsdSkelSkeletonTranslatorImpl
 		float Time,
 		EObjectFlags Flags,
 		bool bSkeletalMeshHasMorphTargets,
-		bool bReuseIdenticalAssets
+		bool bShareAssetsForIdenticalPrims
 	)
 	{
 		TRACE_CPUPROFILER_EVENT_SCOPE(UsdSkelSkeletonTranslatorImpl::ProcessMaterials);
@@ -115,7 +115,7 @@ namespace UsdSkelSkeletonTranslatorImpl
 			AssetCache,
 			InfoCache,
 			Flags,
-			bReuseIdenticalAssets
+			bShareAssetsForIdenticalPrims
 		);
 
 		bool bMaterialsHaveChanged = false;
@@ -754,7 +754,7 @@ namespace UsdSkelSkeletonTranslatorImpl
 		const FName& RenderContext,
 		const FName& MaterialPurpose,
 		const EUsdPurpose PurposesToLoad,
-		bool bReuseIdenticalAssets
+		bool bShareAssetsForIdenticalPrims
 	)
 	{
 		FScopedUsdAllocs Allocs;
@@ -907,7 +907,7 @@ namespace UsdSkelSkeletonTranslatorImpl
 			AssetCache,
 			InfoCache,
 			Flags,
-			bReuseIdenticalAssets
+			bShareAssetsForIdenticalPrims
 		);
 
 		// Compare resolved materials with existing assignments, and create overrides if we need to
@@ -1048,7 +1048,7 @@ namespace UsdSkelSkeletonTranslatorImpl
 			SHA1.UpdateWithString(*PrimPath, PrimPath.Len());
 			SHA1.Final();
 			SHA1.GetHash(&Hash.Hash[0]);
-			const FString PrefixedAnimBPHash = UsdUtils::GetAssetHashPrefix(SkeletonPrim, Context.bReuseIdenticalAssets) + Hash.ToString();
+			const FString PrefixedAnimBPHash = UsdUtils::GetAssetHashPrefix(SkeletonPrim, Context.bShareAssetsForIdenticalPrims) + Hash.ToString();
 
 			bool bReusedAnimBP = false;
 			if (UAnimBlueprint* CachedAnimBP = Context.UsdAssetCache->GetCachedAsset<UAnimBlueprint>(PrefixedAnimBPHash))
@@ -1524,7 +1524,8 @@ namespace UsdSkelSkeletonTranslatorImpl
 					SkeletonBones,
 					BlendShapes
 				);
-				PrefixedSkelMeshHash = UsdUtils::GetAssetHashPrefix(GetSkeletonPrim(), Context->bReuseIdenticalAssets) + SkeletalMeshHash.ToString();
+				PrefixedSkelMeshHash = UsdUtils::GetAssetHashPrefix(GetSkeletonPrim(), Context->bShareAssetsForIdenticalPrims)
+									   + SkeletalMeshHash.ToString();
 
 				const FString DesiredSkeletalMeshName = UsdToUnreal::ConvertString(ClosestParentSkelRoot.Get().GetPrim().GetName());
 
@@ -1654,7 +1655,7 @@ namespace UsdSkelSkeletonTranslatorImpl
 							Context->Time,
 							Context->ObjectFlags,
 							NewBlendShapes.Num() > 0,
-							Context->bReuseIdenticalAssets
+							Context->bShareAssetsForIdenticalPrims
 						);
 
 						if (bMaterialsHaveChanged)
@@ -1865,7 +1866,8 @@ namespace UsdSkelSkeletonTranslatorImpl
 				}
 
 				FSHAHash Hash = UsdSkelSkeletonTranslatorImpl::ComputeSHAHash(SkeletonQuery.Get(), RootMotionPrim, PrefixedSkelMeshHash);
-				FString PrefixedSkelAnimHash = UsdUtils::GetAssetHashPrefix(SkelAnimationPrim, Context->bReuseIdenticalAssets) + Hash.ToString();
+				FString PrefixedSkelAnimHash = UsdUtils::GetAssetHashPrefix(SkelAnimationPrim, Context->bShareAssetsForIdenticalPrims)
+											   + Hash.ToString();
 
 				const FString DesiredName = UsdToUnreal::ConvertToken(SkelAnimationPrim.GetName());
 
@@ -2020,7 +2022,7 @@ USceneComponent* FUsdSkelSkeletonTranslator::CreateComponents()
 				*Context->UsdAssetCache,
 				*Context->InfoCache,
 				Context->ObjectFlags,
-				Context->bReuseIdenticalAssets
+				Context->bShareAssetsForIdenticalPrims
 			);
 
 			// For the groom binding to work, the GroomComponent must be a child of the SceneComponent
@@ -2169,7 +2171,7 @@ void FUsdSkelSkeletonTranslator::UpdateComponents(USceneComponent* SceneComponen
 				Context->RenderContext,
 				Context->MaterialPurpose,
 				Context->PurposesToLoad,
-				Context->bReuseIdenticalAssets
+				Context->bShareAssetsForIdenticalPrims
 			);
 		}
 	}
diff --git a/Engine/Plugins/Importers/USDImporter/Source/USDSchemas/Private/USDVolVolumeTranslator.cpp b/Engine/Plugins/Importers/USDImporter/Source/USDSchemas/Private/USDVolVolumeTranslator.cpp
--- a/Engine/Plugins/Importers/USDImporter/Source/USDSchemas/Private/USDVolVolumeTranslator.cpp
+++ b/Engine/Plugins/Importers/USDImporter/Source/USDSchemas/Private/USDVolVolumeTranslator.cpp
@@ -960,7 +960,7 @@ void FUsdVolVolumeTranslator::CreateAssets()
 	}
 	pxr::UsdStageRefPtr Stage = VolumePrim.GetStage();
 
-	FString VolumePrimHashPrefix = UsdUtils::GetAssetHashPrefix(GetPrim(), Context->bReuseIdenticalAssets);
+	FString VolumePrimHashPrefix = UsdUtils::GetAssetHashPrefix(GetPrim(), Context->bShareAssetsForIdenticalPrims);
 
 	// Create the info structs from the requested files
 	TMap<FString, FSparseVolumeTextureInfo> FilePathHashToTextureInfo = CollectSparseVolumeTextureInfoByFilePath(Volume);
diff --git a/Engine/Plugins/Importers/USDImporter/Source/USDSchemas/Public/USDSchemaTranslator.h b/Engine/Plugins/Importers/USDImporter/Source/USDSchemas/Public/USDSchemaTranslator.h
--- a/Engine/Plugins/Importers/USDImporter/Source/USDSchemas/Public/USDSchemaTranslator.h
+++ b/Engine/Plugins/Importers/USDImporter/Source/USDSchemas/Public/USDSchemaTranslator.h
@@ -277,6 +277,9 @@ struct USDSCHEMAS_API FUsdSchemaTranslationContext : public TSharedFromThis<FUsd
 	 * that asset is generated, and the asset is shared by the components generated for both prims.
 	 * If false, we will always generate a dedicated asset for each prim.
 	 */
+	bool bShareAssetsForIdenticalPrims = true;
+
+	UE_DEPRECATED(5.5, "This property has been renamed to 'Share Assets for Identical Prims'")
 	bool bReuseIdenticalAssets = true;
 
 	/**
diff --git a/Engine/Plugins/Importers/USDImporter/Source/USDStage/Private/USDStageActor.cpp b/Engine/Plugins/Importers/USDImporter/Source/USDStage/Private/USDStageActor.cpp
--- a/Engine/Plugins/Importers/USDImporter/Source/USDStage/Private/USDStageActor.cpp
+++ b/Engine/Plugins/Importers/USDImporter/Source/USDStage/Private/USDStageActor.cpp
@@ -155,7 +155,7 @@ struct FUsdStageActorImpl
 
 		TranslationContext->KindsToCollapse = (EUsdDefaultKind)StageActor->KindsToCollapse;
 		TranslationContext->bMergeIdenticalMaterialSlots = StageActor->bMergeIdenticalMaterialSlots;
-		TranslationContext->bReuseIdenticalAssets = StageActor->bReuseIdenticalAssets;
+		TranslationContext->bShareAssetsForIdenticalPrims = StageActor->bShareAssetsForIdenticalPrims;
 
 		UE::FSdfPath UsdPrimPath(*PrimPath);
 		UUsdPrimTwin* ParentUsdPrimTwin = StageActor->GetRootPrimTwin()->Find(UsdPrimPath.GetParentPath().GetString());
@@ -553,7 +553,7 @@ struct FUsdStageActorImpl
 			EventAttributes.Emplace(TEXT("InterpolationType"), LexToString((uint8)StageActor->InterpolationType));
 			EventAttributes.Emplace(TEXT("KindsToCollapse"), LexToString(StageActor->KindsToCollapse));
 			EventAttributes.Emplace(TEXT("MergeIdenticalMaterialSlots"), LexToString(StageActor->bMergeIdenticalMaterialSlots));
-			EventAttributes.Emplace(TEXT("ReuseIdenticalAssets"), LexToString(StageActor->bReuseIdenticalAssets));
+			EventAttributes.Emplace(TEXT("bShareAssetsForIdenticalPrims"), StageActor->bShareAssetsForIdenticalPrims);
 			EventAttributes.Emplace(TEXT("PurposesToLoad"), LexToString(StageActor->PurposesToLoad));
 			EventAttributes.Emplace(TEXT("NaniteTriangleThreshold"), LexToString(StageActor->NaniteTriangleThreshold));
 			EventAttributes.Emplace(TEXT("RenderContext"), StageActor->RenderContext.ToString());
@@ -1058,7 +1058,7 @@ AUsdStageActor::AUsdStageActor()
 	, InterpolationType(EUsdInterpolationType::Linear)
 	, KindsToCollapse((int32)(EUsdDefaultKind::Component | EUsdDefaultKind::Subcomponent))
 	, bMergeIdenticalMaterialSlots(true)
-	, bReuseIdenticalAssets(true)
+	, bShareAssetsForIdenticalPrims(true)
 	, PurposesToLoad((int32)EUsdPurpose::Proxy)
 	, NaniteTriangleThreshold((uint64)1000000)
 	, MaterialPurpose(*UnrealIdentifiers::MaterialPreviewPurpose)
@@ -2697,6 +2697,21 @@ void AUsdStageActor::SetMergeIdenticalMaterialSlots(bool bMerge)
 	LoadUsdStage();
 }
 
+void AUsdStageActor::SetShareAssetsForIdenticalPrims(bool bShare)
+{
+	if (bShare == bShareAssetsForIdenticalPrims)
+	{
+		return;
+	}
+
+	const bool bMarkDirty = false;
+	Modify(bMarkDirty);
+
+	bShareAssetsForIdenticalPrims = bShare;
+	LoadUsdStage();
+}
+
+PRAGMA_DISABLE_DEPRECATION_WARNINGS
 void AUsdStageActor::SetReuseIdenticalAssets(bool bReuse)
 {
 	if (bReuse == bReuseIdenticalAssets)
@@ -2710,6 +2725,7 @@ void AUsdStageActor::SetReuseIdenticalAssets(bool bReuse)
 	bReuseIdenticalAssets = bReuse;
 	LoadUsdStage();
 }
+PRAGMA_ENABLE_DEPRECATION_WARNINGS
 
 void AUsdStageActor::SetCollapseTopLevelPointInstancers(bool bCollapse)
 {
@@ -4768,11 +4784,11 @@ void AUsdStageActor::HandlePropertyChangedEvent(FPropertyChangedEvent& PropertyC
 		bMergeIdenticalMaterialSlots = !bMergeIdenticalMaterialSlots;
 		SetMergeIdenticalMaterialSlots(bCorrectMergeMaterialSlots);
 	}
-	else if (PropertyName == GET_MEMBER_NAME_CHECKED(AUsdStageActor, bReuseIdenticalAssets))
+	else if (PropertyName == GET_MEMBER_NAME_CHECKED(AUsdStageActor, bShareAssetsForIdenticalPrims))
 	{
-		const bool bCorrectReuse = bReuseIdenticalAssets;
-		bReuseIdenticalAssets = !bReuseIdenticalAssets;
-		SetReuseIdenticalAssets(bCorrectReuse);
+		const bool bCorrectShare = bShareAssetsForIdenticalPrims;
+		bShareAssetsForIdenticalPrims = !bShareAssetsForIdenticalPrims;
+		SetShareAssetsForIdenticalPrims(bCorrectShare);
 	}
 	else if (PropertyName == GET_MEMBER_NAME_CHECKED(AUsdStageActor, PurposesToLoad))
 	{
diff --git a/Engine/Plugins/Importers/USDImporter/Source/USDStage/Public/USDStageActor.h b/Engine/Plugins/Importers/USDImporter/Source/USDStage/Public/USDStageActor.h
--- a/Engine/Plugins/Importers/USDImporter/Source/USDStage/Public/USDStageActor.h
+++ b/Engine/Plugins/Importers/USDImporter/Source/USDStage/Public/USDStageActor.h
@@ -88,6 +88,10 @@ public:
 	 * If false, we will always generate a dedicated asset for each prim.
 	 */
 	UPROPERTY(EditAnywhere, BlueprintReadOnly, Category = "USD", config)
+	bool bShareAssetsForIdenticalPrims;
+
+	UE_DEPRECATED(5.5, "This property has been renamed to 'Share Assets for Identical Prims'")
+	UPROPERTY()
 	bool bReuseIdenticalAssets;
 
 	/**
@@ -176,6 +180,10 @@ public:
 	UFUNCTION(BlueprintCallable, Category = "USD", meta = (CallInEditor = "true"))
 	USDSTAGE_API void SetMergeIdenticalMaterialSlots(bool bMerge);
 
+	UFUNCTION(BlueprintCallable, Category = "USD", meta = (CallInEditor = "true"))
+	USDSTAGE_API void SetShareAssetsForIdenticalPrims(bool bShare);
+
+	UE_DEPRECATED(5.5, "Use 'SetShareAssetsForIdenticalPrims' instead.")
 	UFUNCTION(BlueprintCallable, Category = "USD", meta = (CallInEditor = "true"))
 	USDSTAGE_API void SetReuseIdenticalAssets(bool bReuse);
 
diff --git a/Engine/Plugins/Importers/USDImporter/Source/USDStageEditor/Private/SUSDStage.cpp b/Engine/Plugins/Importers/USDImporter/Source/USDStageEditor/Private/SUSDStage.cpp
--- a/Engine/Plugins/Importers/USDImporter/Source/USDStageEditor/Private/SUSDStage.cpp
+++ b/Engine/Plugins/Importers/USDImporter/Source/USDStageEditor/Private/SUSDStage.cpp
@@ -1078,9 +1078,9 @@ void SUsdStage::FillOptionsMenu(FMenuBuilder& MenuBuilder)
 		);
 
 		MenuBuilder.AddSubMenu(
-			LOCTEXT("AssetReuse", "Asset reuse"),
-			LOCTEXT("AssetReuse_ToolTip", "How to behave when generating identical assets from different prims"),
-			FNewMenuDelegate::CreateSP(this, &SUsdStage::FillAssetReuseSubMenu)
+			LOCTEXT("AssetReuse", "Asset sharing"),
+			LOCTEXT("AssetReuse_ToolTip", "How to behave when generating assets from identical prims"),
+			FNewMenuDelegate::CreateSP(this, &SUsdStage::FillShareAssetsSubMenu)
 		);
 
 		MenuBuilder.AddSubMenu(
@@ -1917,12 +1917,12 @@ void SUsdStage::FillCollapsingSubMenu(FMenuBuilder& MenuBuilder)
 	MenuBuilder.EndSection();
 }
 
-void SUsdStage::FillAssetReuseSubMenu(FMenuBuilder& MenuBuilder)
+void SUsdStage::FillShareAssetsSubMenu(FMenuBuilder& MenuBuilder)
 {
 	MenuBuilder.AddMenuEntry(
-		LOCTEXT("ReuseIdenticalAssets", "Reuse identical assets"),
+		LOCTEXT("ShareAssetsForIdenticalPrims_Text", "Share Assets for Identical Prims"),
 		LOCTEXT(
-			"ReuseIdenticalAssets_ToolTip",
+			"bShareAssetsForIdenticalPrims_ToolTip",
 			"If true, whenever two prims would have generated identical UAssets (like identical StaticMeshes or materials) then only one instance of "
 			"that asset is generated, and the asset is shared by the components generated for both prims. If false, we will generate a dedicated "
 			"asset for each prim."
@@ -1935,14 +1935,14 @@ void SUsdStage::FillAssetReuseSubMenu(FMenuBuilder& MenuBuilder)
 					if (AUsdStageActor* StageActor = GetStageActorOrCDO())
 					{
 						FScopedTransaction Transaction(FText::Format(
-							LOCTEXT("ReuseIdenticalAssetsTransaction", "Toggle bReuseIdenticalAssets on USD stage actor '{0}'"),
+							LOCTEXT("ShareAssetsForIdenticalPrims_Transaction", "Toggle bShareAssetsForIdenticalPrims on USD stage actor '{0}'"),
 							FText::FromString(StageActor->GetActorLabel())
 						));
 
 						// c.f. comment within AddKindToCollapseEntry just below
 						TGuardValue<bool> MaintainSelectionGuard(bUpdatingViewportSelection, true);
 
-						StageActor->SetReuseIdenticalAssets(!StageActor->bReuseIdenticalAssets);
+						StageActor->SetShareAssetsForIdenticalPrims(!StageActor->bShareAssetsForIdenticalPrims);
 						if (StageActor->IsTemplate())
 						{
 							StageActor->SaveConfig();
@@ -1956,7 +1956,7 @@ void SUsdStage::FillAssetReuseSubMenu(FMenuBuilder& MenuBuilder)
 				{
 					if (AUsdStageActor* StageActor = GetStageActorOrCDO())
 					{
-						return StageActor->bReuseIdenticalAssets;
+						return StageActor->bShareAssetsForIdenticalPrims;
 					}
 					return false;
 				}
diff --git a/Engine/Plugins/Importers/USDImporter/Source/USDStageEditor/Private/SUSDStage.h b/Engine/Plugins/Importers/USDImporter/Source/USDStageEditor/Private/SUSDStage.h
--- a/Engine/Plugins/Importers/USDImporter/Source/USDStageEditor/Private/SUSDStage.h
+++ b/Engine/Plugins/Importers/USDImporter/Source/USDStageEditor/Private/SUSDStage.h
@@ -94,7 +94,7 @@ protected:
 	void FillSubdivisionLevelSubMenu(FMenuBuilder& MenuBuilder);
 	void FillMetadataSubMenu(FMenuBuilder& MenuBuilder);
 	void FillCollapsingSubMenu(FMenuBuilder& MenuBuilder);
-	void FillAssetReuseSubMenu(FMenuBuilder& MenuBuilder);
+	void FillShareAssetsSubMenu(FMenuBuilder& MenuBuilder);
 	void FillInterpolationTypeSubMenu(FMenuBuilder& MenuBuilder);
 	void FillSelectionSubMenu(FMenuBuilder& MenuBuilder);
 	void FillNaniteThresholdSubMenu(FMenuBuilder& MenuBuilder);
diff --git a/Engine/Plugins/Importers/USDImporter/Source/USDStageEditorViewModels/Private/USDStageViewModel.cpp b/Engine/Plugins/Importers/USDImporter/Source/USDStageEditorViewModels/Private/USDStageViewModel.cpp
--- a/Engine/Plugins/Importers/USDImporter/Source/USDStageEditorViewModels/Private/USDStageViewModel.cpp
+++ b/Engine/Plugins/Importers/USDImporter/Source/USDStageEditorViewModels/Private/USDStageViewModel.cpp
@@ -328,7 +328,7 @@ void FUsdStageViewModel::ImportStage(const TCHAR* TargetContentFolder, UUsdStage
 		ImportContext.ImportOptions->MetadataOptions = StageActor->MetadataOptions;
 		ImportContext.ImportOptions->KindsToCollapse = StageActor->KindsToCollapse;
 		ImportContext.ImportOptions->bMergeIdenticalMaterialSlots = StageActor->bMergeIdenticalMaterialSlots;
-		ImportContext.ImportOptions->bReuseIdenticalAssets = StageActor->bReuseIdenticalAssets;
+		ImportContext.ImportOptions->bShareAssetsForIdenticalPrims = StageActor->bShareAssetsForIdenticalPrims;
 
 		ImportContext.bReadFromStageCache = true;	 // So that we import whatever the user has open right now, even if the file has changes
 
diff --git a/Engine/Plugins/Importers/USDImporter/Source/USDStageImporter/Private/USDStageImportOptions.cpp b/Engine/Plugins/Importers/USDImporter/Source/USDStageImporter/Private/USDStageImportOptions.cpp
--- a/Engine/Plugins/Importers/USDImporter/Source/USDStageImporter/Private/USDStageImportOptions.cpp
+++ b/Engine/Plugins/Importers/USDImporter/Source/USDStageImporter/Private/USDStageImportOptions.cpp
@@ -44,7 +44,7 @@ UUsdStageImportOptions::UUsdStageImportOptions(const FObjectInitializer& ObjectI
 
 	ExistingActorPolicy = EReplaceActorPolicy::Replace;
 	ExistingAssetPolicy = EReplaceAssetPolicy::Replace;
-	bReuseIdenticalAssets = true;
+	bShareAssetsForIdenticalPrims = true;
 
 	bPrimPathFolderStructure = false;
 	KindsToCollapse = (int32)(EUsdDefaultKind::Component | EUsdDefaultKind::Subcomponent);
@@ -126,7 +126,7 @@ void UsdUtils::AddAnalyticsAttributes(const UUsdStageImportOptions& Options, TAr
 	InOutAttributes.Emplace(TEXT("NumGroomInterpolationSettings"), LexToString(Options.GroomInterpolationSettings.Num()));
 	InOutAttributes.Emplace(TEXT("ReplaceActorPolicy"), LexToString((uint8)Options.ExistingActorPolicy));
 	InOutAttributes.Emplace(TEXT("ReplaceAssetPolicy"), LexToString((uint8)Options.ExistingAssetPolicy));
-	InOutAttributes.Emplace(TEXT("ReuseIdenticalAssets"), LexToString(Options.bReuseIdenticalAssets));
+	InOutAttributes.Emplace(TEXT("ShareAssetsForIdenticalPrims"), Options.bShareAssetsForIdenticalPrims);
 	InOutAttributes.Emplace(TEXT("PrimPathFolderStructure"), LexToString(Options.bPrimPathFolderStructure));
 	InOutAttributes.Emplace(TEXT("KindsToCollapse"), LexToString(Options.KindsToCollapse));
 	InOutAttributes.Emplace(TEXT("MergeIdenticalMaterialSlots"), LexToString(Options.bMergeIdenticalMaterialSlots));
diff --git a/Engine/Plugins/Importers/USDImporter/Source/USDStageImporter/Private/USDStageImporter.cpp b/Engine/Plugins/Importers/USDImporter/Source/USDStageImporter/Private/USDStageImporter.cpp
--- a/Engine/Plugins/Importers/USDImporter/Source/USDStageImporter/Private/USDStageImporter.cpp
+++ b/Engine/Plugins/Importers/USDImporter/Source/USDStageImporter/Private/USDStageImporter.cpp
@@ -1918,7 +1918,7 @@ void UUsdStageImporter::ImportFromFile(FUsdStageImportContext& ImportContext)
 	TranslationContext->ParentComponent = ImportContext.SceneActor ? ImportContext.SceneActor->GetRootComponent() : nullptr;
 	TranslationContext->KindsToCollapse = (EUsdDefaultKind)ImportContext.ImportOptions->KindsToCollapse;
 	TranslationContext->bMergeIdenticalMaterialSlots = ImportContext.ImportOptions->bMergeIdenticalMaterialSlots;
-	TranslationContext->bReuseIdenticalAssets = ImportContext.ImportOptions->bReuseIdenticalAssets;
+	TranslationContext->bShareAssetsForIdenticalPrims = ImportContext.ImportOptions->bShareAssetsForIdenticalPrims;
 	TranslationContext->bAllowInterpretingLODs = ImportContext.ImportOptions->bInterpretLODs;
 	TranslationContext->bAllowParsingSkeletalAnimations = ImportContext.ImportOptions->bImportGeometry
 														  && ImportContext.ImportOptions->bImportSkeletalAnimations;
@@ -2123,7 +2123,7 @@ bool UUsdStageImporter::ReimportSingleAsset(
 	TranslationContext->MetadataOptions = ImportContext.ImportOptions->MetadataOptions;
 	TranslationContext->KindsToCollapse = (EUsdDefaultKind)ImportContext.ImportOptions->KindsToCollapse;
 	TranslationContext->bMergeIdenticalMaterialSlots = ImportContext.ImportOptions->bMergeIdenticalMaterialSlots;
-	TranslationContext->bReuseIdenticalAssets = ImportContext.ImportOptions->bReuseIdenticalAssets;
+	TranslationContext->bShareAssetsForIdenticalPrims = ImportContext.ImportOptions->bShareAssetsForIdenticalPrims;
 	TranslationContext->bAllowInterpretingLODs = ImportContext.ImportOptions->bInterpretLODs;
 	TranslationContext->bAllowParsingSkeletalAnimations = ImportContext.ImportOptions->bImportGeometry
 														  && ImportContext.ImportOptions->bImportSkeletalAnimations;
diff --git a/Engine/Plugins/Importers/USDImporter/Source/USDStageImporter/Public/USDStageImportOptions.h b/Engine/Plugins/Importers/USDImporter/Source/USDStageImporter/Public/USDStageImportOptions.h
--- a/Engine/Plugins/Importers/USDImporter/Source/USDStageImporter/Public/USDStageImportOptions.h
+++ b/Engine/Plugins/Importers/USDImporter/Source/USDStageImporter/Public/USDStageImportOptions.h
@@ -153,24 +153,28 @@ public:
 	UPROPERTY(BlueprintReadWrite, config, EditAnywhere, Category = "USD options", meta = (EditCondition = bImportAtSpecificTimeCode))
 	float ImportTimeCode;
 
+	/**
+	 * If true, whenever two prims would have generated identical UAssets (like identical StaticMeshes or materials) then only one instance of
+	 * that asset is generated, and the asset is shared by the components generated for both prims.
+	 * If false, we will always generate a dedicated asset for each prim.
+	 */
+	UPROPERTY(BlueprintReadWrite, config, EditAnywhere, Category = "USD options")
+	bool bShareAssetsForIdenticalPrims;
+
 	/** Groom group interpolation settings */
 	UPROPERTY(EditAnywhere, config, BlueprintReadWrite, Category = "Groom")
 	TArray<FHairGroupsInterpolation> GroomInterpolationSettings;
 
 	/** What should happen when imported actors and components try to overwrite existing actors and components */
-	UPROPERTY(BlueprintReadWrite, config, EditAnywhere, Category = "Collision", meta = (EditCondition = bImportActors))
+	UPROPERTY(BlueprintReadWrite, config, EditAnywhere, Category = "Conflicts with Existing Objects", meta = (EditCondition = bImportActors))
 	EReplaceActorPolicy ExistingActorPolicy;
 
 	/** What should happen when imported assets try to overwrite existing assets */
-	UPROPERTY(BlueprintReadWrite, config, EditAnywhere, Category = "Collision")
+	UPROPERTY(BlueprintReadWrite, config, EditAnywhere, Category = "Conflicts with Existing Objects")
 	EReplaceAssetPolicy ExistingAssetPolicy;
 
-	/**
-	 * If true, whenever two prims would have generated identical UAssets (like identical StaticMeshes or materials) then only one instance of
-	 * that asset is generated, and the asset is shared by the components generated for both prims.
-	 * If false, we will always generate a dedicated asset for each prim.
-	 */
-	UPROPERTY(BlueprintReadWrite, config, EditAnywhere, Category = "Collision")
+	UE_DEPRECATED(5.5, "This property has been renamed to 'Share Assets for Identical Prims'")
+	UPROPERTY()
 	bool bReuseIdenticalAssets;
 
 	/**
diff --git a/Engine/Plugins/Importers/USDImporter/Source/USDUtilities/Private/USDConversionUtils.cpp b/Engine/Plugins/Importers/USDImporter/Source/USDUtilities/Private/USDConversionUtils.cpp
--- a/Engine/Plugins/Importers/USDImporter/Source/USDUtilities/Private/USDConversionUtils.cpp
+++ b/Engine/Plugins/Importers/USDImporter/Source/USDUtilities/Private/USDConversionUtils.cpp
@@ -1864,9 +1864,9 @@ bool UsdUtils::SetAssetUserData(UObject* Object, UUsdAssetUserData* AssetUserDat
 PRAGMA_ENABLE_DEPRECATION_WARNINGS
 
 #if USE_USD_SDK
-FString UsdUtils::GetAssetHashPrefix(const pxr::UsdPrim& PrimForAsset, bool bReuseIdenticalAssets)
+FString UsdUtils::GetAssetHashPrefix(const pxr::UsdPrim& PrimForAsset, bool bShareAssetsForIdenticalPrims)
 {
-	if (!PrimForAsset || bReuseIdenticalAssets)
+	if (!PrimForAsset || bShareAssetsForIdenticalPrims)
 	{
 		return FString{};
 	}
diff --git a/Engine/Plugins/Importers/USDImporter/Source/USDUtilities/Private/USDLightConversion.cpp b/Engine/Plugins/Importers/USDImporter/Source/USDUtilities/Private/USDLightConversion.cpp
--- a/Engine/Plugins/Importers/USDImporter/Source/USDUtilities/Private/USDLightConversion.cpp
+++ b/Engine/Plugins/Importers/USDImporter/Source/USDUtilities/Private/USDLightConversion.cpp
@@ -248,7 +248,7 @@ bool UsdToUnreal::ConvertDomeLight(
 	const pxr::UsdPrim& Prim,
 	USkyLightComponent& LightComponent,
 	UUsdAssetCache2* TexturesCache,
-	bool bReuseIdenticalAssets
+	bool bShareAssetsForIdenticalPrims
 )
 {
 	FScopedUsdAllocs UsdAllocs;
diff --git a/Engine/Plugins/Importers/USDImporter/Source/USDUtilities/Private/USDShadeConversion.cpp b/Engine/Plugins/Importers/USDImporter/Source/USDUtilities/Private/USDShadeConversion.cpp
--- a/Engine/Plugins/Importers/USDImporter/Source/USDUtilities/Private/USDShadeConversion.cpp
+++ b/Engine/Plugins/Importers/USDImporter/Source/USDUtilities/Private/USDShadeConversion.cpp
@@ -559,7 +559,7 @@ namespace UE
 				FParameterValue& OutValue,
 				const UMaterialInterface* Material,
 				UUsdAssetCache3* TexturesCache,
-				bool bReuseIdenticalAssets
+				bool bShareAssetsForIdenticalPrims
 			)
 			{
 				FScopedUsdAllocs UsdAllocs;
@@ -704,7 +704,7 @@ namespace UE
 							}
 
 							const FString
-								PrefixedTextureHash = UsdUtils::GetAssetHashPrefix(ShadeInput.GetPrim(), bReuseIdenticalAssets)
+								PrefixedTextureHash = UsdUtils::GetAssetHashPrefix(ShadeInput.GetPrim(), bShareAssetsForIdenticalPrims)
 													  + UsdUtils::GetTextureHash(TexturePath, bSRGB, CompressionSettings, AddressX, AddressY);
 
 							pxr::UsdAttribute FileInputAttr = FileInput.GetAttr();
@@ -831,7 +831,7 @@ namespace UE
 				FParameterValue& OutValue,
 				const UMaterialInterface* Material,
 				UUsdAssetCache3* TexturesCache,
-				bool bReuseIdenticalAssets
+				bool bShareAssetsForIdenticalPrims
 			)
 			{
 				FScopedUsdAllocs Allocs;
@@ -848,7 +848,7 @@ namespace UE
 				pxr::UsdShadeAttributeType AttributeType;
 				if (pxr::UsdShadeConnectableAPI::GetConnectedSource(Input.GetAttr(), &Source, &SourceName, &AttributeType))
 				{
-					if (!GetTextureParameterValue(Input, TEXTUREGROUP_WorldSpecular, OutValue, Material, TexturesCache, bReuseIdenticalAssets))
+					if (!GetTextureParameterValue(Input, TEXTUREGROUP_WorldSpecular, OutValue, Material, TexturesCache, bShareAssetsForIdenticalPrims))
 					{
 						// Check if we have a fallback input that we can use instead, since we don't have a valid texture value
 						if (const pxr::UsdShadeInput FallbackInput = Source.GetInput(UnrealIdentifiers::Fallback))
@@ -863,7 +863,7 @@ namespace UE
 
 						// Recurse because the attribute may just be pointing at some other attribute that has the data
 						// (e.g. when shader input is just "hoisted" and connected to the parent material input)
-						return GetFloatParameterValue(Source, SourceName, DefaultValue, OutValue, Material, TexturesCache, bReuseIdenticalAssets);
+						return GetFloatParameterValue(Source, SourceName, DefaultValue, OutValue, Material, TexturesCache, bShareAssetsForIdenticalPrims);
 					}
 				}
 				// No other node connected, so we must have some value
@@ -963,7 +963,7 @@ namespace UE
 				bool bIsNormalMap,
 				const UMaterialInterface* Material,
 				UUsdAssetCache3* TexturesCache,
-				bool bReuseIdenticalAssets
+				bool bShareAssetsForIdenticalPrims
 			)
 			{
 				FScopedUsdAllocs Allocs;
@@ -986,7 +986,7 @@ namespace UE
 							OutValue,
 							Material,
 							TexturesCache,
-							bReuseIdenticalAssets
+							bShareAssetsForIdenticalPrims
 						))
 					{
 						// Check whether this input receives its value through a connection to a
@@ -1037,7 +1037,7 @@ namespace UE
 							bIsNormalMap,
 							Material,
 							TexturesCache,
-							bReuseIdenticalAssets
+							bShareAssetsForIdenticalPrims
 						);
 					}
 				}
@@ -2156,7 +2156,7 @@ bool UsdToUnreal::ConvertMaterial(
 	UMaterialInstance& Material,
 	UUsdAssetCache3* TexturesCache,
 	const TCHAR* RenderContext,
-	bool bReuseIdenticalAssets
+	bool bShareAssetsForIdenticalPrims
 )
 {
 	FScopedUsdAllocs UsdAllocs;
@@ -2192,7 +2192,7 @@ bool UsdToUnreal::ConvertMaterial(
 			!bIsNormalMap,
 			&Material,
 			TexturesCache,
-			bReuseIdenticalAssets
+			bShareAssetsForIdenticalPrims
 		))
 	{
 		MaterialParameters.Add(TEXT("BaseColor"), ParameterValue);
@@ -2207,7 +2207,7 @@ bool UsdToUnreal::ConvertMaterial(
 			!bIsNormalMap,
 			&Material,
 			TexturesCache,
-			bReuseIdenticalAssets
+			bShareAssetsForIdenticalPrims
 		))
 	{
 		MaterialParameters.Add(TEXT("EmissiveColor"), ParameterValue);
@@ -2221,7 +2221,7 @@ bool UsdToUnreal::ConvertMaterial(
 			ParameterValue,
 			&Material,
 			TexturesCache,
-			bReuseIdenticalAssets
+			bShareAssetsForIdenticalPrims
 		))
 	{
 		MaterialParameters.Add(TEXT("Metallic"), ParameterValue);
@@ -2235,7 +2235,7 @@ bool UsdToUnreal::ConvertMaterial(
 			ParameterValue,
 			&Material,
 			TexturesCache,
-			bReuseIdenticalAssets
+			bShareAssetsForIdenticalPrims
 		))
 	{
 		MaterialParameters.Add(TEXT("Roughness"), ParameterValue);
@@ -2249,7 +2249,7 @@ bool UsdToUnreal::ConvertMaterial(
 			ParameterValue,
 			&Material,
 			TexturesCache,
-			bReuseIdenticalAssets
+			bShareAssetsForIdenticalPrims
 		))
 	{
 		MaterialParameters.Add(TEXT("Opacity"), ParameterValue);
@@ -2264,7 +2264,7 @@ bool UsdToUnreal::ConvertMaterial(
 			bIsNormalMap,
 			&Material,
 			TexturesCache,
-			bReuseIdenticalAssets
+			bShareAssetsForIdenticalPrims
 		))
 	{
 		MaterialParameters.Add(TEXT("Normal"), ParameterValue);
@@ -2278,7 +2278,7 @@ bool UsdToUnreal::ConvertMaterial(
 			ParameterValue,
 			&Material,
 			TexturesCache,
-			bReuseIdenticalAssets
+			bShareAssetsForIdenticalPrims
 		))
 	{
 		MaterialParameters.Add(TEXT("Refraction"), ParameterValue);
@@ -2292,7 +2292,7 @@ bool UsdToUnreal::ConvertMaterial(
 			ParameterValue,
 			&Material,
 			TexturesCache,
-			bReuseIdenticalAssets
+			bShareAssetsForIdenticalPrims
 		))
 	{
 		MaterialParameters.Add(TEXT("AmbientOcclusion"), ParameterValue);
@@ -2359,7 +2359,7 @@ bool UsdToUnreal::ConvertMaterial(
 	UMaterial& Material,
 	UUsdAssetCache3* TexturesCache,
 	const TCHAR* RenderContext,
-	bool bReuseIdenticalAssets
+	bool bShareAssetsForIdenticalPrims
 )
 {
 #if WITH_EDITOR
@@ -2405,7 +2405,7 @@ bool UsdToUnreal::ConvertMaterial(
 			!bIsNormalMap,
 			&Material,
 			TexturesCache,
-			bReuseIdenticalAssets
+			bShareAssetsForIdenticalPrims
 		))
 	{
 		BaseColorParameter = TempParameterValue;
@@ -2420,7 +2420,7 @@ bool UsdToUnreal::ConvertMaterial(
 			!bIsNormalMap,
 			&Material,
 			TexturesCache,
-			bReuseIdenticalAssets
+			bShareAssetsForIdenticalPrims
 		))
 	{
 		EmissiveParameter = TempParameterValue;
@@ -2434,7 +2434,7 @@ bool UsdToUnreal::ConvertMaterial(
 			TempParameterValue,
 			&Material,
 			TexturesCache,
-			bReuseIdenticalAssets
+			bShareAssetsForIdenticalPrims
 		))
 	{
 		MetallicParameter = TempParameterValue;
@@ -2448,7 +2448,7 @@ bool UsdToUnreal::ConvertMaterial(
 			TempParameterValue,
 			&Material,
 			TexturesCache,
-			bReuseIdenticalAssets
+			bShareAssetsForIdenticalPrims
 		))
 	{
 		RoughnessParameter = TempParameterValue;
@@ -2462,7 +2462,7 @@ bool UsdToUnreal::ConvertMaterial(
 			TempParameterValue,
 			&Material,
 			TexturesCache,
-			bReuseIdenticalAssets
+			bShareAssetsForIdenticalPrims
 		))
 	{
 		OpacityParameter = TempParameterValue;
@@ -2477,7 +2477,7 @@ bool UsdToUnreal::ConvertMaterial(
 			bIsNormalMap,
 			&Material,
 			TexturesCache,
-			bReuseIdenticalAssets
+			bShareAssetsForIdenticalPrims
 		))
 	{
 		NormalParameter = TempParameterValue;
@@ -2491,7 +2491,7 @@ bool UsdToUnreal::ConvertMaterial(
 		TempParameterValue,
 		&Material,
 		TexturesCache,
-		bReuseIdenticalAssets
+		bShareAssetsForIdenticalPrims
 	);
 	if (bHasRefractionValue || Material.BlendMode == BLEND_Translucent)	   // Force a 1.5 IOR if USD didn't specify a value, as it's USD's fallback
 																		   // value
@@ -2507,7 +2507,7 @@ bool UsdToUnreal::ConvertMaterial(
 			TempParameterValue,
 			&Material,
 			TexturesCache,
-			bReuseIdenticalAssets
+			bShareAssetsForIdenticalPrims
 		))
 	{
 		AmbientOcclusionParameter = TempParameterValue;
@@ -2678,7 +2678,7 @@ bool UsdToUnreal::ConvertShadeInputsToParameters(
 	UMaterialInstance& MaterialInstance,
 	UUsdAssetCache3* TexturesCache,
 	const TCHAR* RenderContext,
-	bool bReuseIdenticalAssets
+	bool bShareAssetsForIdenticalPrims
 )
 {
 	FScopedUsdAllocs UsdAllocs;
@@ -2760,7 +2760,7 @@ bool UsdToUnreal::ConvertShadeInputsToParameters(
 					ParameterValue,
 					&MaterialInstance,
 					TexturesCache,
-					bReuseIdenticalAssets
+					bShareAssetsForIdenticalPrims
 				))
 			{
 				UsdShadeConversionImpl::SetParameterValue(MaterialInstance, *DisplayName, ParameterValue, bForUsdPreviewSurface, Unused);
@@ -2775,7 +2775,7 @@ bool UsdToUnreal::ConvertShadeInputsToParameters(
 					 bIsNormalMap,
 					 &MaterialInstance,
 					 TexturesCache,
-					 bReuseIdenticalAssets
+					 bShareAssetsForIdenticalPrims
 				 ))
 		{
 			UsdShadeConversionImpl::SetParameterValue(MaterialInstance, *DisplayName, ParameterValue, bForUsdPreviewSurface, Unused);
@@ -2792,11 +2792,11 @@ bool UsdToUnreal::ConvertMaterial(
 	UMaterialInstance& Material,
 	UUsdAssetCache2* TexturesCache,
 	const TCHAR* RenderContext,
-	bool bReuseIdenticalAssets
+	bool bShareAssetsForIdenticalPrims
 )
 {
 	UUsdAssetCache3* NewCache = nullptr;
-	return UsdToUnreal::ConvertMaterial(UsdShadeMaterial, Material, NewCache, RenderContext, bReuseIdenticalAssets);
+	return UsdToUnreal::ConvertMaterial(UsdShadeMaterial, Material, NewCache, RenderContext, bShareAssetsForIdenticalPrims);
 }
 
 bool UsdToUnreal::ConvertMaterial(
@@ -2804,11 +2804,11 @@ bool UsdToUnreal::ConvertMaterial(
 	UMaterial& Material,
 	UUsdAssetCache2* TexturesCache,
 	const TCHAR* RenderContext,
-	bool bReuseIdenticalAssets
+	bool bShareAssetsForIdenticalPrims
 )
 {
 	UUsdAssetCache3* NewCache = nullptr;
-	return UsdToUnreal::ConvertMaterial(UsdShadeMaterial, Material, NewCache, RenderContext, bReuseIdenticalAssets);
+	return UsdToUnreal::ConvertMaterial(UsdShadeMaterial, Material, NewCache, RenderContext, bShareAssetsForIdenticalPrims);
 }
 
 bool UsdToUnreal::ConvertShadeInputsToParameters(
@@ -2816,11 +2816,11 @@ bool UsdToUnreal::ConvertShadeInputsToParameters(
 	UMaterialInstance& MaterialInstance,
 	UUsdAssetCache2* TexturesCache,
 	const TCHAR* RenderContext,
-	bool bReuseIdenticalAssets
+	bool bShareAssetsForIdenticalPrims
 )
 {
 	UUsdAssetCache3* NewCache = nullptr;
-	return UsdToUnreal::ConvertShadeInputsToParameters(UsdShadeMaterial, MaterialInstance, NewCache, RenderContext, bReuseIdenticalAssets);
+	return UsdToUnreal::ConvertShadeInputsToParameters(UsdShadeMaterial, MaterialInstance, NewCache, RenderContext, bShareAssetsForIdenticalPrims);
 }
 PRAGMA_ENABLE_DEPRECATION_WARNINGS
 
@@ -3478,7 +3478,7 @@ bool UsdUtils::IsMaterialTranslucent(const pxr::UsdShadeMaterial& UsdShadeMateri
 
 	const UMaterialInterface* Material = nullptr;
 	UUsdAssetCache3* TexturesCache = nullptr;
-	const bool bReuseIdenticalAssets = true;
+	const bool bShareAssetsForIdenticalPrims = true;
 	UsdShadeConversionImpl::FParameterValue ParameterValue;
 	bool bHasOpacityConnection = UsdShadeConversionImpl::GetFloatParameterValue(
 		Connectable,
@@ -3487,7 +3487,7 @@ bool UsdUtils::IsMaterialTranslucent(const pxr::UsdShadeMaterial& UsdShadeMateri
 		ParameterValue,
 		Material,
 		TexturesCache,
-		bReuseIdenticalAssets
+		bShareAssetsForIdenticalPrims
 	);
 
 	// Don't check if the texture is nullptr here as we won't actually parse it yet. If the variant has this type we know it's meant to be bound to a
diff --git a/Engine/Plugins/Importers/USDImporter/Source/USDUtilities/Public/USDConversionUtils.h b/Engine/Plugins/Importers/USDImporter/Source/USDUtilities/Public/USDConversionUtils.h
--- a/Engine/Plugins/Importers/USDImporter/Source/USDUtilities/Public/USDConversionUtils.h
+++ b/Engine/Plugins/Importers/USDImporter/Source/USDUtilities/Public/USDConversionUtils.h
@@ -361,11 +361,11 @@ namespace UsdUtils
 #if USE_USD_SDK
 	/**
 	 * Simple utility that generates a prefix that should be added to the hash when using it to cache/query an UsdAssetCache2 object.
-	 * The idea is to use this prefix alongside the hash whenever bReuseIdenticalAssets is false, so as to generate a new hash for each prim.
-	 * Whenever bReuseIdenticalAssets is true, the prefix will be the empty string.
-	 * Whenever bReuseIdenticalAssets is false, the prefix will be a SHA1 hash of the prim path and its stage identifier.
+	 * The idea is to use this prefix alongside the hash whenever bShareAssetsForIdenticalPrims is false, so as to generate a new hash for each prim.
+	 * Whenever bShareAssetsForIdenticalPrims is true, the prefix will be the empty string.
+	 * Whenever bShareAssetsForIdenticalPrims is false, the prefix will be a SHA1 hash of the prim path and its stage identifier.
 	 */
-	USDUTILITIES_API FString GetAssetHashPrefix(const pxr::UsdPrim& PrimForAsset, bool bReuseIdenticalAssets);
+	USDUTILITIES_API FString GetAssetHashPrefix(const pxr::UsdPrim& PrimForAsset, bool bShareAssetsForIdenticalPrims);
 #endif	  // #if USE_USD_SDK
 
 	/** Adds a reference on Prim to the layer at AbsoluteFilePath */
diff --git a/Engine/Plugins/Importers/USDImporter/Source/USDUtilities/Public/USDLightConversion.h b/Engine/Plugins/Importers/USDImporter/Source/USDUtilities/Public/USDLightConversion.h
--- a/Engine/Plugins/Importers/USDImporter/Source/USDUtilities/Public/USDLightConversion.h
+++ b/Engine/Plugins/Importers/USDImporter/Source/USDUtilities/Public/USDLightConversion.h
@@ -94,7 +94,7 @@ namespace UsdToUnreal
 		const pxr::UsdPrim& Prim,
 		USkyLightComponent& LightComponent,
 		UUsdAssetCache2* TexturesCache,
-		bool bReuseIdenticalAssets = true
+		bool bShareAssetsForIdenticalPrims = true
 	);
 	PRAGMA_ENABLE_DEPRECATION_WARNINGS
 
diff --git a/Engine/Plugins/Importers/USDImporter/Source/USDUtilities/Public/USDShadeConversion.h b/Engine/Plugins/Importers/USDImporter/Source/USDUtilities/Public/USDShadeConversion.h
--- a/Engine/Plugins/Importers/USDImporter/Source/USDUtilities/Public/USDShadeConversion.h
+++ b/Engine/Plugins/Importers/USDImporter/Source/USDUtilities/Public/USDShadeConversion.h
@@ -63,14 +63,14 @@ namespace UsdToUnreal
 		UMaterialInstance& Material,
 		UUsdAssetCache3* TexturesCache = nullptr,
 		const TCHAR* RenderContext = nullptr,
-		bool bReuseIdenticalAssets = true
+		bool bShareAssetsForIdenticalPrims = true
 	);
 	USDUTILITIES_API bool ConvertMaterial(
 		const pxr::UsdShadeMaterial& UsdShadeMaterial,
 		UMaterial& Material,
 		UUsdAssetCache3* TexturesCache = nullptr,
 		const TCHAR* RenderContext = nullptr,
-		bool bReuseIdenticalAssets = true
+		bool bShareAssetsForIdenticalPrims = true
 	);
 
 	/**
@@ -89,7 +89,7 @@ namespace UsdToUnreal
 		UMaterialInstance& MaterialInstance,
 		UUsdAssetCache3* TexturesCache,
 		const TCHAR* RenderContext = nullptr,
-		bool bReuseIdenticalAssets = true
+		bool bShareAssetsForIdenticalPrims = true
 	);
 
 	PRAGMA_DISABLE_DEPRECATION_WARNINGS
@@ -99,7 +99,7 @@ namespace UsdToUnreal
 		UMaterialInstance& Material,
 		UUsdAssetCache2* TexturesCache = nullptr,
 		const TCHAR* RenderContext = nullptr,
-		bool bReuseIdenticalAssets = true
+		bool bShareAssetsForIdenticalPrims = true
 	);
 	UE_DEPRECATED(5.5, "Use the other overload that receives an UUsdAssetCache3 object instead")
 	USDUTILITIES_API bool ConvertMaterial(
@@ -107,7 +107,7 @@ namespace UsdToUnreal
 		UMaterial& Material,
 		UUsdAssetCache2* TexturesCache = nullptr,
 		const TCHAR* RenderContext = nullptr,
-		bool bReuseIdenticalAssets = true
+		bool bShareAssetsForIdenticalPrims = true
 	);
 	UE_DEPRECATED(5.5, "Use the other overload that receives an UUsdAssetCache3 object instead")
 	USDUTILITIES_API bool ConvertShadeInputsToParameters(
@@ -115,7 +115,7 @@ namespace UsdToUnreal
 		UMaterialInstance& MaterialInstance,
 		UUsdAssetCache2* TexturesCache,
 		const TCHAR* RenderContext = nullptr,
-		bool bReuseIdenticalAssets = true
+		bool bShareAssetsForIdenticalPrims = true
 	);
 	PRAGMA_ENABLE_DEPRECATION_WARNINGS
 }
