USD: Lower some log messages down to 'verbose' level as they can lead to unnecessary spam.

diff --git a/Engine/Plugins/Importers/USDImporter/Source/USDClasses/Private/USDAssetCache3.cpp b/Engine/Plugins/Importers/USDImporter/Source/USDClasses/Private/USDAssetCache3.cpp
--- a/Engine/Plugins/Importers/USDImporter/Source/USDClasses/Private/USDAssetCache3.cpp
+++ b/Engine/Plugins/Importers/USDImporter/Source/USDClasses/Private/USDAssetCache3.cpp
@@ -266,7 +266,7 @@ namespace UE::USDAssetCache3::Private
 						}
 						UE_LOG(
 							LogUsd,
-							Log,
+							Verbose,
 							TEXT("Not trying to clean up '%s' because it is referenced by '%s'"),
 							*ReferencedObject->GetPathName(),
 							*Referencer->GetPathName()
@@ -323,7 +323,7 @@ namespace UE::USDAssetCache3::Private
 
 			for (UObject* DeletableAsset : DeletableAssets)
 			{
-				UE_LOG(LogUsd, Log, TEXT("Deleting '%s'"), *DeletableAsset->GetPathName());
+				UE_LOG(LogUsd, Verbose, TEXT("Deleting '%s'"), *DeletableAsset->GetPathName());
 
 				Progress.EnterProgressFrame();
 
@@ -345,7 +345,7 @@ namespace UE::USDAssetCache3::Private
 #else
 			for (UObject* DeletableAsset : DeletableAssets)
 			{
-				UE_LOG(LogUsd, Log, TEXT("Deleting '%s'"), *DeletableAsset->GetPathName());
+				UE_LOG(LogUsd, Verbose, TEXT("Deleting '%s'"), *DeletableAsset->GetPathName());
 
 				// These are essentially the internals of ObjectTools::DeleteSingleObject that actually do the deletion
 				DeletableAsset->MarkPackageDirty();
@@ -367,7 +367,7 @@ namespace UE::USDAssetCache3::Private
 			double ElapsedSeconds = FPlatformTime::ToSeconds64(FPlatformTime::Cycles64() - StartTime);
 			UE_LOG(
 				LogUsd,
-				Log,
+				Verbose,
 				TEXT("Deleted %d out of %d assets in %.3f s (including GC and transaction reset)"),
 				DeletableAssets.Num(),
 				ObjectsToDelete.Num(),
@@ -957,7 +957,7 @@ void UUsdAssetCache3::DeleteUnreferencedAssets(bool bShowConfirmation)
 			{
 				UE_LOG(
 					LogUsd,
-					Log,
+					Verbose,
 					TEXT("Not trying to clean up '%s' because the path doesn't resolve, or resolve to a saved asset"),
 					*AssetPathString
 				);
@@ -971,7 +971,7 @@ void UUsdAssetCache3::DeleteUnreferencedAssets(bool bShowConfirmation)
 			{
 				// We never want to delete assets that the user manually added to the asset cache
 				// (Only assets added via GetOrCreateCachedAsset/CacheAsset are considered Deletable)
-				UE_LOG(LogUsd, Log, TEXT("Not trying to clean up '%s' because it hasn't been set as deletable"), *AssetPathString);
+				UE_LOG(LogUsd, Verbose, TEXT("Not trying to clean up '%s' because it hasn't been set as deletable"), *AssetPathString);
 				continue;
 			}
 
@@ -980,7 +980,7 @@ void UUsdAssetCache3::DeleteUnreferencedAssets(bool bShowConfirmation)
 			{
 				if (FoundReferencers->Num() > 0)
 				{
-					UE_LOG(LogUsd, Log, TEXT("Not trying to clean up '%s' because it has object referencers"), *AssetPathString);
+					UE_LOG(LogUsd, Verbose, TEXT("Not trying to clean up '%s' because it has object referencers"), *AssetPathString);
 					continue;
 				}
 			}
@@ -1411,7 +1411,7 @@ void UUsdAssetCache3::TryCachingAssetFromAssetUserData(const FAssetData& Existin
 	{
 		UE_LOG(
 			LogUsd,
-			Log,
+			Verbose,
 			TEXT("Loading existing asset '%s' to check if it can be automatically added to the asset cache '%s'"),
 			*ExistingAssetPathStr,
 			*GetPathName()
@@ -1427,7 +1427,7 @@ void UUsdAssetCache3::TryCachingAssetFromAssetUserData(const FAssetData& Existin
 			{
 				UE_LOG(
 					LogUsd,
-					Log,
+					Verbose,
 					TEXT("Automatically caching asset '%s' into asset cache '%s' with hash '%s'"),
 					*ExistingAssetPathStr,
 					*GetPathName(),
diff --git a/Engine/Plugins/Importers/USDImporter/Source/USDSchemas/Private/USDGeomMeshTranslator.cpp b/Engine/Plugins/Importers/USDImporter/Source/USDSchemas/Private/USDGeomMeshTranslator.cpp
--- a/Engine/Plugins/Importers/USDImporter/Source/USDSchemas/Private/USDGeomMeshTranslator.cpp
+++ b/Engine/Plugins/Importers/USDImporter/Source/USDSchemas/Private/USDGeomMeshTranslator.cpp
@@ -631,7 +631,7 @@ namespace UsdGeomMeshTranslatorImpl
 			Options |= EComputeNTBsFlags::Normals;
 			UE_LOG(
 				LogUsd,
-				Log,
+				Verbose,
 				TEXT("%f%% of the normals from Mesh prim '%s' are invalid or unusable. This is at or above the threshold of '%f%%' (configurable via "
 					 "the cvar '%s'), so normals will be discarded and fully recomputed. Note that when the cvar "
 					 "'USD.Subdiv.IgnoreNormalsWhenSubdividing' is true it is expected for subdivision meshes to have their normals discarded."),
@@ -645,7 +645,7 @@ namespace UsdGeomMeshTranslatorImpl
 		{
 			UE_LOG(
 				LogUsd,
-				Log,
+				Verbose,
 				TEXT("%f%% of the normals from Mesh prim '%s' are invalid or unusable. This is below the threshold of '%f%%' (configurable via the "
 					 "cvar '%s'), so the invalid normals will be repaired."),
 				InvalidNormalFraction * 100.0f,
