USD: Prevent race condition when multiple cooker processes try to patch the plugInfo.json files simultaneously.

diff --git a/Engine/Plugins/Importers/USDImporter/Source/USDClasses/Private/USDClassesModule.cpp b/Engine/Plugins/Importers/USDImporter/Source/USDClasses/Private/USDClassesModule.cpp
--- a/Engine/Plugins/Importers/USDImporter/Source/USDClasses/Private/USDClassesModule.cpp
+++ b/Engine/Plugins/Importers/USDImporter/Source/USDClasses/Private/USDClassesModule.cpp
@@ -86,6 +86,13 @@ namespace UE::USDClasses::Private
 
 void IUsdClassesModule::UpdatePlugInfoFiles(const FString& PluginDirectory, const FString& TargetDllFolder)
 {
+	// Prevent cooker worker processes from trying to patch the plugInfo.json files: Only the director
+	// process should do that
+	if (UE::GetMultiprocessId() != 0)
+	{
+		return;
+	}
+
 	// Traverse all USD plugins
 	TArray<FString> JsonPaths;
 	const bool bFiles = true;
