USD: Replace the UsdAssetCache's usage of AssetActions with AssetDefinitions.

diff --git a/Engine/Plugins/Importers/USDImporter/Source/USDClassesEditor/Private/AssetDefinition_USDAssetCache.cpp b/Engine/Plugins/Importers/USDImporter/Source/USDClassesEditor/Private/AssetDefinition_USDAssetCache.cpp
new file mode 100644
--- /dev/null
+++ b/Engine/Plugins/Importers/USDImporter/Source/USDClassesEditor/Private/AssetDefinition_USDAssetCache.cpp
@@ -0,0 +1,35 @@
+// Copyright Epic Games, Inc. All Rights Reserved.
+
+#include "AssetDefinition_USDAssetCache.h"
+
+#include "USDAssetCache3.h"
+
+#define LOCTEXT_NAMESPACE "AssetTypeActions"
+
+FText UAssetDefinition_UsdAssetCache::GetAssetDisplayName() const
+{
+	return LOCTEXT("AssetTypeActions_AssetCache", "USD Asset Cache");
+}
+
+FLinearColor UAssetDefinition_UsdAssetCache::GetAssetColor() const
+{
+	return FColor(32, 145, 208);
+}
+
+TSoftClassPtr<UObject> UAssetDefinition_UsdAssetCache::GetAssetClass() const
+{
+	return UUsdAssetCache3::StaticClass();
+}
+
+bool UAssetDefinition_UsdAssetCache::CanImport() const
+{
+	return false;
+}
+
+TConstArrayView<FAssetCategoryPath> UAssetDefinition_UsdAssetCache::GetAssetCategories() const
+{
+	static const auto Categories = {EAssetCategoryPaths::Misc};
+	return Categories;
+}
+
+#undef LOCTEXT_NAMESPACE
diff --git a/Engine/Plugins/Importers/USDImporter/Source/USDClassesEditor/Private/AssetDefinition_USDAssetCache.h b/Engine/Plugins/Importers/USDImporter/Source/USDClassesEditor/Private/AssetDefinition_USDAssetCache.h
new file mode 100644
--- /dev/null
+++ b/Engine/Plugins/Importers/USDImporter/Source/USDClassesEditor/Private/AssetDefinition_USDAssetCache.h
@@ -0,0 +1,23 @@
+// Copyright Epic Games, Inc. All Rights Reserved.
+
+#pragma once
+
+#include "AssetDefinitionDefault.h"
+#include "CoreMinimal.h"
+
+#include "AssetDefinition_USDAssetCache.generated.h"
+
+UCLASS()
+class UAssetDefinition_UsdAssetCache : public UAssetDefinitionDefault
+{
+	GENERATED_BODY()
+
+public:
+	// UAssetDefinition Begin
+	virtual FText GetAssetDisplayName() const override;
+	virtual FLinearColor GetAssetColor() const override;
+	virtual TSoftClassPtr<UObject> GetAssetClass() const override;
+	virtual bool CanImport() const override;
+	virtual TConstArrayView<FAssetCategoryPath> GetAssetCategories() const override;
+	// UAssetDefinition End
+};
diff --git a/Engine/Plugins/Importers/USDImporter/Source/USDClassesEditor/Private/USDAssetCacheAssetActions.cpp b/Engine/Plugins/Importers/USDImporter/Source/USDClassesEditor/Private/USDAssetCacheAssetActions.cpp
deleted file mode 100644
diff --git a/Engine/Plugins/Importers/USDImporter/Source/USDClassesEditor/Private/USDAssetCacheAssetActions.h b/Engine/Plugins/Importers/USDImporter/Source/USDClassesEditor/Private/USDAssetCacheAssetActions.h
deleted file mode 100644
diff --git a/Engine/Plugins/Importers/USDImporter/Source/USDClassesEditor/Private/USDClassesEditorModule.cpp b/Engine/Plugins/Importers/USDImporter/Source/USDClassesEditor/Private/USDClassesEditorModule.cpp
--- a/Engine/Plugins/Importers/USDImporter/Source/USDClassesEditor/Private/USDClassesEditorModule.cpp
+++ b/Engine/Plugins/Importers/USDImporter/Source/USDClassesEditor/Private/USDClassesEditorModule.cpp
@@ -3,9 +3,7 @@
 #include "USDClassesEditorModule.h"
 
 #include "USDAssetCache2.h"
-#include "USDAssetCacheAssetActions.h"
 
-#include "AssetToolsModule.h"
 #include "Modules/ModuleManager.h"
 
 #define LOCTEXT_NAMESPACE "USDClassesEditorModule"
@@ -29,24 +27,11 @@ class FUsdClassesEditorModule : public IUsdClassesEditorModule
 public:
 	virtual void StartupModule() override
 	{
-		// Register asset actions for the AssetCache asset
-		IAssetTools& AssetTools = FModuleManager::LoadModuleChecked<FAssetToolsModule>(TEXT("AssetTools")).Get();
-		AssetCacheAssetActions = MakeShared<FUsdAssetCacheAssetActions>();
-		AssetTools.RegisterAssetTypeActions(AssetCacheAssetActions.ToSharedRef());
 	}
 
 	virtual void ShutdownModule() override
 	{
-		// Unregister asset actions for the AssetCache asset
-		if (FAssetToolsModule* AssetToolsModule = FModuleManager::GetModulePtr<FAssetToolsModule>(TEXT("AssetTools")))
-		{
-			IAssetTools& AssetTools = AssetToolsModule->Get();
-			AssetTools.UnregisterAssetTypeActions(AssetCacheAssetActions.ToSharedRef());
-		}
 	}
-
-private:
-	TSharedPtr<IAssetTypeActions> AssetCacheAssetActions;
 };
 
 IMPLEMENT_MODULE(FUsdClassesEditorModule, USDClassesEditor);
diff --git a/Engine/Plugins/Importers/USDImporter/Source/USDClassesEditor/USDClassesEditor.Build.cs b/Engine/Plugins/Importers/USDImporter/Source/USDClassesEditor/USDClassesEditor.Build.cs
--- a/Engine/Plugins/Importers/USDImporter/Source/USDClassesEditor/USDClassesEditor.Build.cs
+++ b/Engine/Plugins/Importers/USDImporter/Source/USDClassesEditor/USDClassesEditor.Build.cs
@@ -12,6 +12,7 @@ namespace UnrealBuildTool.Rules
 			PublicDependencyModuleNames.AddRange(
 				new string[]
 				{
+					"AssetDefinition",
 					"Core",
 					"CoreUObject",
 					"Engine",
