USDCore, USDImporter: move UE Python example scripts back to USDImporter

diff --git a/Engine/Plugins/Importers/USDImporter/Resources/PythonExamples/collect_metadata.py b/Engine/Plugins/Importers/USDImporter/Resources/PythonExamples/collect_metadata.py
new file mode 100644
--- /dev/null
+++ b/Engine/Plugins/Importers/USDImporter/Resources/PythonExamples/collect_metadata.py
@@ -0,0 +1,90 @@
+"""
+This script describes how the metadata from USD prims can be collected and manipulated on
+the Unreal side, after fully importing a stage or opening one with a stage actor. Just
+remember to set the desired metadata options on the stage actor itself if using one, as
+it will not collect metadata by default (e.g. `stage_actor.set_collect_metadata(True)`).
+
+In general the metadata will be collected into UsdAssetUserData objects, that will be added
+to the generated assets (and components, if the import option is set).
+
+All asset types that we generate (with the exeption of UPhysicsAssets) can contain
+UsdAssetUserData objects and will receive USD metadata from their source prims when being
+generated.
+
+The metadata values will always be stringified and stored within the UsdAssetUserData as
+strings, as there is no real way of storing anything close to a "dynamic type" or variant
+type (like VtValues) with UPROPERTIES. Keep in mind that unreal.UsdConversionLibrary contains
+functions that can help converting these values to and from strings into Unreal types (the
+very last step of this script shows an example of this). These functions can be used from
+Python or Editor Utility Blueprints.
+"""
+
+import os
+import unreal
+
+DESTINATION_CONTENT_PATH = r"/Game/Imported/"
+
+this_directory = os.path.dirname(os.path.realpath(__file__))
+full_file_path = os.path.join(this_directory, "metadata_cube.usda")
+
+# Setup import options
+options = unreal.UsdStageImportOptions()
+options.import_actors = True
+options.import_geometry = True
+options.import_materials = True
+options.metadata_options.collect_metadata = True
+options.metadata_options.collect_from_entire_subtrees = True  # Collapsed assets will contain metadata from collapsed children
+options.metadata_options.collect_on_components = True  # Spawned components will have Asset User Data too
+options.metadata_options.blocked_prefix_filters = ["customData:int"]  # All metadata entries starting with "customData:int" will be blocked
+options.metadata_options.invert_filters = False  # When this is false, the prefix filters select stuff to block, and allows everything else.
+												 # When this is true, the prefix filters select stuff to allow, and block everything else
+
+# Import scene
+task = unreal.AssetImportTask()
+task.set_editor_property('filename', full_file_path)
+task.set_editor_property('destination_path', DESTINATION_CONTENT_PATH)
+task.set_editor_property('automated', True)
+task.set_editor_property('options', options)
+task.set_editor_property('replace_existing', True)
+asset_tools = unreal.AssetToolsHelpers.get_asset_tools()
+asset_tools.import_asset_tasks([task])
+asset_paths = task.get_editor_property("imported_object_paths")
+success = asset_paths and len(asset_paths) > 0
+if not success:
+	exit()
+
+# Find our imported static mesh
+asset_paths = unreal.EditorAssetLibrary.list_assets(DESTINATION_CONTENT_PATH)
+for asset_path in asset_paths:
+	asset = unreal.load_asset(asset_path)
+	if isinstance(asset, unreal.StaticMesh):
+		static_mesh = asset
+		break
+
+# Get the USD asset user data from the static mesh
+aud = unreal.UsdConversionLibrary.get_usd_asset_user_data(static_mesh)
+
+# Iterate through all the collected metadata
+for stage_identifier, metadata in aud.stage_identifier_to_metadata.items():
+	for prim_path, prim_metadata in metadata.prim_path_to_metadata.items():
+		for key, value in prim_metadata.metadata.items():
+			print(f"{stage_identifier=} {prim_path=}, {key=}, {value=}")
+
+# Get whatever USD generated for the stage identifier
+stage_identifier = list(aud.stage_identifier_to_metadata.keys())[0]
+
+# You can omit the stage identifier and prim path if the asset user data contains exactly one entry for either
+has_int_value = unreal.UsdConversionLibrary.has_metadata_field(aud, "customData:intValue")
+assert(not has_int_value)  # We filtered entries starting with "customData:int", so we won't get the "intValue" entry
+
+# Get a particular metadata value directly without iterating through the structs (there are additional functions within UsdConversionLibrary)
+metadata_value = unreal.UsdConversionLibrary.get_metadata_field(aud, "customData:nestedMap:someColor", stage_identifier, "/Cube")
+print(metadata_value.stringified_value)  # Prints "(1, 0.5, 0.3)". Note that stringified_value is always just a string, however
+print(metadata_value.type_name)  # Prints "float3"
+
+# Can use the library's utility functions to unstringify stringified_value directly into UE types
+vector = unreal.UsdConversionLibrary.unstringify_as_float3(metadata_value.stringified_value)
+assert(isinstance(vector, unreal.Vector))
+print(vector)  # Prints "<Struct 'Vector' (0x0000026B026371A0) {x: 1.000000, y: 0.500000, z: 0.300000}>"
+
+
diff --git a/Engine/Plugins/Importers/USDImporter/Resources/PythonExamples/create_and_export_metadata.py b/Engine/Plugins/Importers/USDImporter/Resources/PythonExamples/create_and_export_metadata.py
new file mode 100644
--- /dev/null
+++ b/Engine/Plugins/Importers/USDImporter/Resources/PythonExamples/create_and_export_metadata.py
@@ -0,0 +1,109 @@
+"""
+This script describes how to create brand new metadata for an Unreal asset, then export it
+to USD and have the metadata end up on the exported prim.
+
+Of course, it is also possible to directly export to USD assets that were originally imported
+from USD with collected metadata. That metadata will be written back out to the exported prim,
+and in most cases should end up identical to the metadata from the originally imported prim.
+
+Like for the importing case on collect_metadata.py, for the exporting case the metadata should
+be stored inside UsdAssetUserData objects, which are added to the UObjects to export. These
+can be persisted and permanently saved with the UObjects, or just temporarily added before
+export as we'll do here.
+
+Note that we support exporting metadata for all asset types that we support exporting to USD,
+including even ULevels and ULevelSequences. For levels and level sequences the metadata will
+be exported to the top-level "Root" prim, that is set as the default prim on the exported layers.
+
+When exporting levels, note that you can use `unreal.UsdConversionLibrary.set_usd_asset_user_data()`
+in order to set your UsdAssetUserData on the ULevel objects from Python. It is also important to
+stress that the ULevel UObjects can contain metadata, but the UWorld UObjects that own them cannot.
+If you have a UWorld, you can use `levels = unreal.EditorLevelUtils.get_levels(world)` to try and get
+at the persistent level owned by the UWorld, and stash your metadata there.
+
+Additionally, when exporting levels, note that you can use the export option
+"export_component_metadata" to export metadata held by individual SceneComponents on that level.
+Note that metadata exported in this manner ends up on the layers corresponding to the levels
+themselves, while metadata stored on assets ends up on the layers corresponding to each individual
+asset. A benefit of this separation is the fact that when e.g. a StaticMeshComponent contains metadata
+to export, but also its StaticMesh contains its own set of metadata to export, the two sets
+will be exported to their separate layers, but composed together with regular USD composition
+whenever the stage is opened. The component's metadata will provide the stronger opinion in that case.
+"""
+
+import os
+import unreal
+
+ASSET_CONTENT_PATH = "/Engine/BasicShapes/Cube.Cube"
+EXPORT_FILE_PATH = r"[REDACTED_USER_HOME] username>/Desktop/ExportTest/output.usda"
+
+asset = unreal.load_asset(ASSET_CONTENT_PATH)
+assert(asset)
+
+# When exporting, the stage identifier is not important: We will export all entries one after the other
+# When exporting, if we have a single prim path then its actual value doesn't matter: The output prim will
+# receive all entries from that single prim path regardless of what it is
+aud = unreal.UsdAssetUserData()
+success = unreal.UsdConversionLibrary.set_metadata_field(
+	aud,
+	key="customData:someColorValue",
+	value=unreal.UsdConversionLibrary.stringify_as_float3(unreal.Vector(1.0, 0.5, 0.3)),
+	value_type_name="float3",
+	stage_identifier="StageIdentifier",
+	prim_path="/SomePrimPath",
+	trigger_property_change_events=False
+)
+assert(success)
+
+# Set the UsdAssetUserData on the asset we're about to export
+success = unreal.UsdConversionLibrary.set_usd_asset_user_data(asset, aud)
+assert(success)
+
+# Setup export options
+options = unreal.StaticMeshExporterUSDOptions()
+options.mesh_asset_options.material_baking_options.textures_dir = unreal.DirectoryPath(os.path.join(os.path.dirname(EXPORT_FILE_PATH), "Textures"))
+options.metadata_options.export_asset_info = True  # Exports generic info about the Unreal asset, can be turned off independently of the metadata export
+options.metadata_options.export_asset_metadata = True
+options.metadata_options.blocked_prefix_filters = []  # We can use filters on export too, but we won't for this sample.
+options.metadata_options.invert_filters = False 	  # Check collect_metadata.py for an example of filter usage
+
+# Export the asset
+task = unreal.AssetExportTask()
+task.set_editor_property('object', asset)
+task.set_editor_property('filename', EXPORT_FILE_PATH)
+task.set_editor_property('automated', True)
+task.set_editor_property('options', options)
+task.set_editor_property('exporter', unreal.StaticMeshExporterUsd())
+task.set_editor_property('replace_identical', True)
+success = unreal.Exporter.run_asset_export_task(task)
+assert(success)
+
+# The exported stage should look something like this:
+"""
+#usda 1.0
+(
+    defaultPrim = "Cube"
+    metersPerUnit = 0.01
+    upAxis = "Z"
+)
+
+def Mesh "Cube" (
+    assetInfo = {
+        asset identifier = @[REDACTED_USER_HOME] username>/Desktop/ExportTest/output.usda@
+        string name = "Cube"
+        dictionary unreal = {
+            string assetType = "StaticMesh"
+            string contentPath = "/Engine/BasicShapes/Cube.Cube"
+            string engineVersion = "<your engine version>"
+            string exportTime = "2000.11.26-10.33.01"
+        }
+        string version = "895C90D17BDB5F192671D6911E2114627986D8A3"
+    }
+    customData = {
+        float3 someColorValue = (1, 0.5, 0.3)           <---------- exported metadata
+    }
+)
+{
+	...
+}
+"""
\ No newline at end of file
diff --git a/Engine/Plugins/Importers/USDImporter/Resources/PythonExamples/enable_control_rig.py b/Engine/Plugins/Importers/USDImporter/Resources/PythonExamples/enable_control_rig.py
new file mode 100644
--- /dev/null
+++ b/Engine/Plugins/Importers/USDImporter/Resources/PythonExamples/enable_control_rig.py
@@ -0,0 +1,41 @@
+"""
+This script describes how to enable Control Rig for the components generated for a particular prim when
+imported in UE.
+
+After the script is executed, the stage can be saved and the connection details will persist on the USD
+files themselves. Opening the stage again will automatically re-establish the connection.
+
+To use this example, follow these steps:
+- Enable the "USD Importer" and "Live Link" plugins in UE;
+- Open the USD Stage editor window by going to Window -> Virtual Production -> USD Stage;
+- Open your desired USD Stage by going to File -> Open on the USD Stage editor and picking a USD file;
+- Edit the script below, replacing the prim path, ControlRig path and other attributes with your desired values;
+- Run the script, either by copy-pasting it into the UE Python console, or by pasting the path to this
+  file onto the same console.
+"""
+
+from pxr import Usd, UsdUtils, Sdf
+
+stage = UsdUtils.StageCache().Get().GetAllStages()[0]
+prim = stage.GetPrimAtPath("/MySkelRootPrim")
+
+schema = Usd.SchemaRegistry.GetTypeFromSchemaTypeName("ControlRigAPI")
+prim.ApplyAPI(schema)
+
+# Optional: If you want to also set the schema attribute values, you can do this:
+with Sdf.ChangeBlock():
+    path_attr = prim.GetAttribute("unreal:controlRig:controlRigPath")
+    path_attr.Set("/Game/SkeletalMeshes/MyRig.MyRig")
+
+    # Set this to True if you want to generate an FKControlRig instead of using an existing
+    # Control Rig blueprint asset
+    use_fk_attr = prim.GetAttribute("unreal:controlRig:useFKControlRig")
+    use_fk_attr.Set(False)
+
+    # Enable this to have the Control Rig baking process automatically prune unnecessary keys
+    reduce_keys_attr = prim.GetAttribute("unreal:controlRig:reduceKeys")
+    reduce_keys_attr.Set(False)
+
+    tolerance_attr = prim.GetAttribute("unreal:controlRig:reductionTolerance")
+    tolerance_attr.Set(0.001)
+
diff --git a/Engine/Plugins/Importers/USDImporter/Resources/PythonExamples/enable_live_link.py b/Engine/Plugins/Importers/USDImporter/Resources/PythonExamples/enable_live_link.py
new file mode 100644
--- /dev/null
+++ b/Engine/Plugins/Importers/USDImporter/Resources/PythonExamples/enable_live_link.py
@@ -0,0 +1,36 @@
+"""
+This script describes how to enable Live Link for the components generated for a particular prim when
+imported in UE.
+
+After the script is executed, the stage can be saved and the connection details will persist on the USD
+files themselves. Opening the stage again will automatically re-establish the connection.
+
+To use this example, follow these steps:
+- Enable the "USD Importer" and "Live Link" plugins in UE;
+- Open the USD Stage editor window by going to Window -> Virtual Production -> USD Stage;
+- Open your desired USD Stage by going to File -> Open on the USD Stage editor and picking a USD file;
+- Edit the script below, replacing the prim path, AnimBlueprint path and live link subject names with
+  your desired values;
+- Run the script, either by copy-pasting it into the UE Python console, or by pasting the path to this
+  file onto the same console.
+"""
+
+from pxr import Usd, UsdUtils, Sdf
+
+stage = UsdUtils.StageCache().Get().GetAllStages()[0]
+prim = stage.GetPrimAtPath("/cube_default")
+
+schema = Usd.SchemaRegistry.GetTypeFromSchemaTypeName("LiveLinkAPI")
+prim.ApplyAPI(schema)
+
+# Optional: If you want to also set the schema attribute values, you can do this:
+with Sdf.ChangeBlock():
+    anim_bp_attr = prim.GetAttribute("unreal:liveLink:animBlueprintPath")
+    anim_bp_attr.Set("/Game/MyBlueprints/MyAnimBP.MyAnimBP")
+
+    subject_name_attr = prim.GetAttribute("unreal:liveLink:subjectName")
+    subject_name_attr.Set("cube")
+
+    enabled_attr = prim.GetAttribute("unreal:liveLink:enabled")
+    enabled_attr.Set(True)
+
diff --git a/Engine/Plugins/Importers/USDImporter/Resources/PythonExamples/export_static_mesh.py b/Engine/Plugins/Importers/USDImporter/Resources/PythonExamples/export_static_mesh.py
new file mode 100644
--- /dev/null
+++ b/Engine/Plugins/Importers/USDImporter/Resources/PythonExamples/export_static_mesh.py
@@ -0,0 +1,21 @@
+import unreal
+
+OUTPUT_FILENAME = r"C:\MyFolder\output.usda"
+INPUT_ASSET_CONTENT_PATH = r"/Game/MyStaticMesh"
+
+asset = unreal.load_asset(INPUT_ASSET_CONTENT_PATH)
+
+options = unreal.StaticMeshExporterUSDOptions()
+options.stage_options.meters_per_unit = 0.02
+options.mesh_asset_options.use_payload = True
+options.mesh_asset_options.payload_format = "usda"
+options.mesh_asset_options.bake_materials = True
+
+task = unreal.AssetExportTask()
+task.set_editor_property('object', asset)
+task.set_editor_property('filename', OUTPUT_FILENAME)
+task.set_editor_property('automated', True)
+task.set_editor_property('options', options)
+task.set_editor_property('exporter', unreal.StaticMeshExporterUsd())
+
+unreal.Exporter.run_asset_export_task(task)
\ No newline at end of file
diff --git a/Engine/Plugins/Importers/USDImporter/Resources/PythonExamples/import_specific_prims.py b/Engine/Plugins/Importers/USDImporter/Resources/PythonExamples/import_specific_prims.py
new file mode 100644
--- /dev/null
+++ b/Engine/Plugins/Importers/USDImporter/Resources/PythonExamples/import_specific_prims.py
@@ -0,0 +1,68 @@
+"""
+This script describes how to import a USD Stage into actors and assets, also optionally specifying specific prim
+paths to import.
+
+The `prims_to_import` property can be left at its default value of ["/"] (a list with just the "/" prim path in it)
+to import the entire stage.
+
+The provided paths in `prims_to_import` are used as a USD population mask when opening the stage, and as such are
+subject to the rules described here: https://graphics.pixar.com/usd/release/api/class_usd_stage_population_mask.html
+
+As an example, consider the following USD stage:
+
+        #usda 1.0
+        def Xform "ParentA"
+        {
+            def Xform "ChildA"
+            {
+            }
+
+            def Xform "ChildB"
+            {
+            }
+        }
+        def Xform "ParentB"
+        {
+            def Xform "ChildC"
+            {
+            }
+
+            def Xform "ChildD"
+            {
+            }
+        }
+
+In general, the main thing to keep in mind is that if "/ParentA" is within `prims_to_import`, ParentA and *all* of its
+children will be imported. As a consequence, having both "/ParentA" and "/ParentA/ChildA" on the list is reduntant, as
+"/ParentA" will already lead to "/ParentA/ChildA" being imported, as previously mentioned.
+"""
+
+import unreal
+
+ROOT_LAYER_FILENAME = r"C:\MyFolder\my_scene.usda"
+DESTINATION_CONTENT_PATH = r"/Game/Imported/"
+
+options = unreal.UsdStageImportOptions()
+options.import_actors = True
+options.import_geometry = True
+options.import_skeletal_animations = True
+options.import_level_sequences = True
+options.import_materials = True
+options.prim_path_folder_structure = False
+options.prims_to_import = [
+    "/ParentA",  # This will import ParentA, ChildA and ChildB
+    "/ParentB/ChildC"  # This will import ParentB and ChildC only (and *not* import ChildD)
+]
+
+task = unreal.AssetImportTask()
+task.set_editor_property('filename', ROOT_LAYER_FILENAME)
+task.set_editor_property('destination_path', DESTINATION_CONTENT_PATH)
+task.set_editor_property('automated', True)
+task.set_editor_property('options', options)
+task.set_editor_property('replace_existing', True)
+
+asset_tools = unreal.AssetToolsHelpers.get_asset_tools()
+asset_tools.import_asset_tasks([task])
+
+asset_paths = task.get_editor_property("imported_object_paths")
+success = asset_paths and len(asset_paths) > 0
\ No newline at end of file
diff --git a/Engine/Plugins/Importers/USDImporter/Resources/PythonExamples/stage_actor_and_stage_editor.py b/Engine/Plugins/Importers/USDImporter/Resources/PythonExamples/stage_actor_and_stage_editor.py
new file mode 100644
--- /dev/null
+++ b/Engine/Plugins/Importers/USDImporter/Resources/PythonExamples/stage_actor_and_stage_editor.py
@@ -0,0 +1,65 @@
+"""
+This script describes the basic usage of the USD Stage Actor and USD Stage Editor.
+"""
+
+import os
+import unreal
+from pxr import Usd, UsdUtils, UsdGeom, Gf
+
+ROOT_LAYER_FILENAME = r"C:\MyFolder\my_scene.usda"
+DESTINATION_CONTENT_PATH = r"/Game/Imported/"
+
+# Stage Actors can be spawned on the UE level like this.
+# If the Stage Editor is open when this happens, it will automatically attach to the last actor spawned.
+editor_actor_subsystem = unreal.get_editor_subsystem(unreal.EditorActorSubsystem)
+stage_actor = editor_actor_subsystem.spawn_actor_from_class(unreal.UsdStageActor, unreal.Vector())
+
+# You can open a stage and interact with it by manipulating the Stage Actor directly:
+stage_actor.set_editor_property('root_layer', unreal.FilePath(ROOT_LAYER_FILENAME))
+
+# Any stage opened by a Stage Actor will be opened within the UsdUtils' stage cache.
+# This means you can retrieve the same stage via Python from it:
+stage = None
+for s in UsdUtils.StageCache.Get().GetAllStages():
+    if s.GetRootLayer().GetDisplayName() == os.path.basename(ROOT_LAYER_FILENAME):
+         stage = s
+         break
+
+# Interacting with the stage via Python will update the UE level automatically:
+prim = stage.GetPrimAtPath("/Cup")
+xform_api = UsdGeom.XformCommonAPI(prim)
+trans, _, _, _, _ = xform_api.GetXformVectors(Usd.TimeCode.Default())
+delta = Gf.Vec3d(100, 100, 100)
+xform_api.SetTranslate(trans + delta)
+
+# You can also toggle variants in this way, and both the USD Stage Editor and the UE level will update automatically:
+prim = stage.GetPrimAtPath("/PrimWithALODVariantSet")
+variant_sets = prim.GetVariantSets()
+if variant_sets.HasVariantSet("LOD"):
+    var_set = variant_sets.GetVariantSet("LOD")
+    for var in var_set.GetVariantNames():
+        print(f"Setting variant set '{var_set.GetName()}' with variant '{var}'")
+        var_set.SetVariantSelection(var)
+
+# Alternatively it is possible to interact with the USD Stage Editor via Python.
+# Here are some examples of operations you can do. Please consult the documentation for the full list
+unreal.UsdStageEditorLibrary.open_stage_editor()
+unreal.UsdStageEditorLibrary.file_open(ROOT_LAYER_FILENAME)
+other_stage_actor = unreal.UsdStageEditorLibrary.get_attached_stage_actor()
+unreal.UsdStageEditorLibrary.set_attached_stage_actor(stage_actor)
+selected_prims = unreal.UsdStageEditorLibrary.get_selected_prim_paths()
+
+# You can even trigger the Actions->Import button from here, as a way of importing
+# the current state of the opened stage:
+options = unreal.UsdStageImportOptions()
+options.import_actors = True
+options.import_geometry = True
+options.import_skeletal_animations = True
+options.import_level_sequences = True
+options.import_materials = True
+options.prim_path_folder_structure = False
+options.prims_to_import = [
+    "/ParentA",
+    "/ParentB/ChildC"
+]
+unreal.UsdStageEditorLibrary.actions_import(DESTINATION_CONTENT_PATH, options)
\ No newline at end of file
diff --git a/Engine/Plugins/Runtime/USDCore/Resources/PythonExamples/collect_metadata.py b/Engine/Plugins/Runtime/USDCore/Resources/PythonExamples/collect_metadata.py
deleted file mode 100644
diff --git a/Engine/Plugins/Runtime/USDCore/Resources/PythonExamples/create_and_export_metadata.py b/Engine/Plugins/Runtime/USDCore/Resources/PythonExamples/create_and_export_metadata.py
deleted file mode 100644
diff --git a/Engine/Plugins/Runtime/USDCore/Resources/PythonExamples/enable_control_rig.py b/Engine/Plugins/Runtime/USDCore/Resources/PythonExamples/enable_control_rig.py
deleted file mode 100644
diff --git a/Engine/Plugins/Runtime/USDCore/Resources/PythonExamples/enable_live_link.py b/Engine/Plugins/Runtime/USDCore/Resources/PythonExamples/enable_live_link.py
deleted file mode 100644
diff --git a/Engine/Plugins/Runtime/USDCore/Resources/PythonExamples/export_static_mesh.py b/Engine/Plugins/Runtime/USDCore/Resources/PythonExamples/export_static_mesh.py
deleted file mode 100644
diff --git a/Engine/Plugins/Runtime/USDCore/Resources/PythonExamples/import_specific_prims.py b/Engine/Plugins/Runtime/USDCore/Resources/PythonExamples/import_specific_prims.py
deleted file mode 100644
diff --git a/Engine/Plugins/Runtime/USDCore/Resources/PythonExamples/stage_actor_and_stage_editor.py b/Engine/Plugins/Runtime/USDCore/Resources/PythonExamples/stage_actor_and_stage_editor.py
deleted file mode 100644
