Fix error compiling USDGeomMeshTranslator.cpp

diff --git a/Engine/Plugins/Importers/USDImporter/Source/USDSchemas/Private/USDGeomMeshTranslator.cpp b/Engine/Plugins/Importers/USDImporter/Source/USDSchemas/Private/USDGeomMeshTranslator.cpp
--- a/Engine/Plugins/Importers/USDImporter/Source/USDSchemas/Private/USDGeomMeshTranslator.cpp
+++ b/Engine/Plugins/Importers/USDImporter/Source/USDSchemas/Private/USDGeomMeshTranslator.cpp
@@ -816,7 +816,7 @@ namespace UsdGeomMeshTranslatorImpl
 		}
 
 #if RHI_RAYTRACING
-		if (IsRayTracingAllowed() && bSupportRayTracing)
+		if (IsRayTracingAllowed() && SaticMesh.bSupportRayTracing)
 		{
 			StaticMesh.GetRenderData()->InitializeRayTracingRepresentationFromRenderingLODs();
 		}
