USD: Fix invalidated prim references when flipping through LODs during import.

diff --git a/Engine/Plugins/Importers/USDImporter/Source/USDStageImporter/Private/USDStageImporter.cpp b/Engine/Plugins/Importers/USDImporter/Source/USDStageImporter/Private/USDStageImporter.cpp
--- a/Engine/Plugins/Importers/USDImporter/Source/USDStageImporter/Private/USDStageImporter.cpp
+++ b/Engine/Plugins/Importers/USDImporter/Source/USDStageImporter/Private/USDStageImporter.cpp
@@ -388,6 +388,7 @@ namespace UE::USDStageImporter::Private
 		USceneComponent* Component = nullptr;
 
 		// Spawn components and/or actors for this prim
+		UE::FSdfPath PrimPath = Prim.GetPrimPath();
 		if (TSharedPtr<FUsdSchemaTranslator> SchemaTranslator = UsdSchemasModule.GetTranslatorRegistry()
 																	.CreateTranslatorForSchema(TranslationContext.AsShared(), UE::FUsdTyped(Prim)))
 		{
@@ -396,6 +397,10 @@ namespace UE::USDStageImporter::Private
 			bExpandChilren = !SchemaTranslator->CollapsesChildren(ECollapsingType::Components);
 		}
 
+		// Refresh our prim reference because FUsdGeomXformableTranslator::CreateComponents can potentially flip through
+		// LODs when setting material overrides
+		Prim = ImportContext.Stage.GetPrimAtPath(PrimPath);
+
 		// In USD if a parent prim has animated visibility, it will affect the entire subtree. In UE this doesn't
 		// happen by default, so if our visibility is animated (or if we're forced to animate visibility from a parent prim),
 		// then we must ensure that we generate visibility tracks for our child prims that don't collapse either, so that
