USD: Fix exported Spawnable tracks never being able to hide/despawn the objects.

diff --git a/Engine/Plugins/Importers/USDImporter/Source/USDExporter/Private/LevelSequenceExporterUSD.cpp b/Engine/Plugins/Importers/USDImporter/Source/USDExporter/Private/LevelSequenceExporterUSD.cpp
--- a/Engine/Plugins/Importers/USDImporter/Source/USDExporter/Private/LevelSequenceExporterUSD.cpp
+++ b/Engine/Plugins/Importers/USDImporter/Source/USDExporter/Private/LevelSequenceExporterUSD.cpp
@@ -23,6 +23,8 @@
 #include "UsdWrappers/UsdStage.h"
 
 #include "AssetExportTask.h"
+#include "Bindings/MovieSceneReplaceableDirectorBlueprintBinding.h"
+#include "Bindings/MovieSceneSpawnableDirectorBlueprintBinding.h"
 #include "Compilation/MovieSceneCompiledDataManager.h"
 #include "Components/SkeletalMeshComponent.h"
 #include "Editor.h"
@@ -52,8 +54,6 @@
 #include "Tracks/MovieSceneSpawnTrack.h"
 #include "Tracks/MovieSceneSubTrack.h"
 #include "UObject/UObjectGlobals.h"
-#include "Bindings/MovieSceneSpawnableDirectorBlueprintBinding.h"
-#include "Bindings/MovieSceneReplaceableDirectorBlueprintBinding.h"
 
 #define LOCTEXT_NAMESPACE "LevelSequenceExporterUSD"
 
@@ -220,9 +220,18 @@ namespace UE::LevelSequenceExporterUSD::Private
 				}
 			}
 
+			// We keep track of spawned objects here on our derived class, but the base classes expect the Register to have
+			// an entry while the spawnable is spawned, and to not have one when it is not spawned, so here we must synchronize it
+			if (Object)
+			{
+				ESpawnOwnership SpawnOwnership = ESpawnOwnership::InnerSequence;
+				FMovieSceneSpawnRegisterKey Key(TemplateID, Guid, BindingIndex);
+				Register.Add(Key, FSpawnedObject(Guid, *Object, SpawnOwnership));
+			}
 			// Don't even have anything we can reuse: We need to spawn a brand new instance of this spawnable
-			if (!Object)
+			else
 			{
+				// SpawnObject will add an entry into the Register for us
 				Object = FMovieSceneSpawnRegister::SpawnObject(Guid, MovieScene, TemplateID, SharedPlaybackState, BindingIndex);
 				UE_LOG(
 					LogUsd,
@@ -292,6 +301,8 @@ namespace UE::LevelSequenceExporterUSD::Private
 
 		virtual void DestroySpawnedObject(UObject& Object, UMovieSceneSpawnableBindingBase* CustomSpawnableBinding) override
 		{
+			// We don't have to clean up the Register here, the caller to DestroySpawnedObject will do that
+
 			if (bDestroyingJustHides)
 			{
 				USceneComponent* Component = nullptr;
@@ -396,38 +407,6 @@ namespace UE::LevelSequenceExporterUSD::Private
 			}
 		}
 
-		// FLevelSequenceHidingSpawnRegister is a bit of a hack and just hides it's spawned actors instead
-		// of deleting them (when we want it to do so). Unfortunately, the base FMovieSceneSpawnRegister part
-		// will still nevertheless clear it's Register entry for the spawnable when deleting (even if just hiding),
-		// and there's nothing we can do to prevent it. This means we can't call FindSpawnedObject and must use our
-		// own GetExistingSpawn and data members
-		UObject* GetExistingSpawn(const UMovieSceneSequence& RootSequence, FMovieSceneSequenceID SequenceID, const FSpawnedInstanceKey& InstanceKey)
-		{
-			int32 SpawnableIndex = INDEX_NONE;
-			if (TMap<FMovieSceneSequenceID, TMap<FSpawnedInstanceKey, int32>>* SequenceIDsToSpawns = RootSequenceToSpawnableInstanceIndices.Find(
-					&RootSequence
-				))
-			{
-				if (TMap<FSpawnedInstanceKey, int32>* Spawns = SequenceIDsToSpawns->Find(SequenceID))
-				{
-					if (int32* Index = Spawns->Find(InstanceKey))
-					{
-						SpawnableIndex = *Index;
-					}
-				}
-			}
-
-			if (SpawnableIndex != INDEX_NONE)
-			{
-				if (TArray<UObject*>* Instances = SpawnableInstances.Find(InstanceKey))
-				{
-					return (*Instances)[SpawnableIndex];
-				}
-			}
-
-			return nullptr;
-		}
-
 	private:
 		// Ensures all of our new spawns have unique names
 		TSet<FString> UsedActorLabels;
@@ -925,7 +904,7 @@ namespace UE::LevelSequenceExporterUSD::Private
 						{
 							if (BindingReference.CustomBinding->WillSpawnObject(Context.Sequencer->GetSharedPlaybackState()))
 							{
-								BoundObject = Context.SpawnRegister->GetExistingSpawn(*RootSequence, SequenceInstance, FSpawnedInstanceKey(Guid, BindingIndex++));
+								BoundObject = Context.SpawnRegister->FindSpawnedObject(Guid, SequenceInstance, BindingIndex++).Get();
 								if (!BoundObject)
 								{
 									// This should never happen as we preemptively spawn everything:
@@ -940,7 +919,10 @@ namespace UE::LevelSequenceExporterUSD::Private
 								}
 							}
 
-							if (UMovieSceneSpawnableDirectorBlueprintBinding* SpawnableDirectorBlueprintBinding = Cast<UMovieSceneSpawnableDirectorBlueprintBinding>(BindingReference.CustomBinding->AsSpawnable(Context.Sequencer->GetSharedPlaybackState())))
+							if (UMovieSceneSpawnableDirectorBlueprintBinding* SpawnableDirectorBlueprintBinding = Cast<
+									UMovieSceneSpawnableDirectorBlueprintBinding>(
+									BindingReference.CustomBinding->AsSpawnable(Context.Sequencer->GetSharedPlaybackState())
+								))
 							{
 								if (SpawnableDirectorBlueprintBinding->DynamicBinding.Function)
 								{
@@ -948,11 +930,15 @@ namespace UE::LevelSequenceExporterUSD::Private
 								}
 							}
 
-							if (UMovieSceneReplaceableDirectorBlueprintBinding* ReplaceableDirectorBlueprintBinding = Cast<UMovieSceneReplaceableDirectorBlueprintBinding>(BindingReference.CustomBinding))
+							if (UMovieSceneReplaceableDirectorBlueprintBinding* ReplaceableDirectorBlueprintBinding = Cast<
+									UMovieSceneReplaceableDirectorBlueprintBinding>(BindingReference.CustomBinding))
 							{
 								if (ReplaceableDirectorBlueprintBinding->DynamicBinding.Function)
 								{
-									DynamicBindings.Add(FSpawnedInstanceKey(Guid, BindingIndex), &ReplaceableDirectorBlueprintBinding->DynamicBinding);
+									DynamicBindings.Add(
+										FSpawnedInstanceKey(Guid, BindingIndex),
+										&ReplaceableDirectorBlueprintBinding->DynamicBinding
+									);
 								}
 							}
 						}
