USD: Fix too many and incorrect compatible materials being generated in some scenarios.

diff --git a/Engine/Plugins/Importers/USDImporter/Source/USDSchemas/Private/MeshTranslationImpl.cpp b/Engine/Plugins/Importers/USDImporter/Source/USDSchemas/Private/MeshTranslationImpl.cpp
--- a/Engine/Plugins/Importers/USDImporter/Source/USDSchemas/Private/MeshTranslationImpl.cpp
+++ b/Engine/Plugins/Importers/USDImporter/Source/USDSchemas/Private/MeshTranslationImpl.cpp
@@ -561,41 +561,64 @@ TMap<const UsdUtils::FUsdPrimMaterialSlot*, UMaterialInterface*> MeshTranslation
 				}
 				case UsdUtils::EPrimAssignmentType::MaterialPrim:
 				{
-					UMaterialInterface* OneSidedMat = nullptr;
-					bool bOneSidedMatIsInstanceOfReferencePreviewSurface = false;
-
 					UE::FSdfPath MaterialPrimPath{*Slot.MaterialSource};
 
-					TArray<UMaterialInterface*> ExistingMaterials = InfoCache.GetAssetsForPrim<UMaterialInterface>(MaterialPrimPath);
+					bool bMaterialIsDirectReference = false;
 
+					// Here we have to pick the "best" material to use as reference, in case we need compatible/TwoSided versions.
+					// They are returned from InfoCache.GetAssetsForPrim in the most recent to least recent order, so
+					// in general we want to pick the first ones we find that match our criteria (as the older assets may be leftover from
+					// before we resynced something)
+					TArray<UMaterialInterface*> ExistingMaterials = InfoCache.GetAssetsForPrim<UMaterialInterface>(MaterialPrimPath);
 					for (UMaterialInterface* ExistingMaterial : ExistingMaterials)
 					{
 						const bool bExistingIsTwoSided = ExistingMaterial->IsTwoSided();
+						const bool bSidednessMatches = Slot.bMeshIsDoubleSided == bExistingIsTwoSided;
+
+						// Prefer sticking with a material instance that has as parent one of our reference materials.
+						// The idea here being that we have two approaches when making TwoSided and compatible
+						// materials: A) Make the material compatible first, and then a TwoSided version of the
+						// compatible; B) Make the material TwoSided first, and then a compatible version of the
+						// TwoSided; We're going to chose B), for the reason that at runtime we can only make a material
+						// TwoSided if it is an instance of our reference materials (as we can't manually change the
+						// material base property overrides at runtime)
+						//
+						// Note that we may end up with MaterialX/MDL materials in here, so not being a direct reference
+						// doesn't mean it's just another one of our material instances... It could be that an e.g.
+						// non-UsdPreviewSurface MDL UMaterial is the best reference we can find
+						bool bExistingIsDirectReference = false;
+						if (UMaterialInstance* ExistingInstance = Cast<UMaterialInstance>(ExistingMaterial))
+						{
+							bExistingIsDirectReference = UsdUnreal::MaterialUtils::IsReferencePreviewSurfaceMaterial(ExistingInstance->Parent);
+						}
 
-						if (!bExistingIsTwoSided)
+						if (bSidednessMatches)
 						{
-							// Prefer sticking with a material instance that has as parent one of our reference materials.
-							// The idea here being that we have two approaches when making TwoSided and compatible
-							// materials: A) Make the material compatible first, and then a TwoSided version of the
-							// compatible; B) Make the material TwoSided first, and then a compatible version of the
-							// TwoSided; We're going to chose B), for the reason that at runtime we can only make a material
-							// TwoSided if it is an instance of our reference materials (as we can't manually change the
-							// material base property overrides at runtime)
-							UMaterialInstance* ExistingInstance = Cast<UMaterialInstance>(ExistingMaterial);
-							const bool bExistingIsInstanceOfReferencePreviewSurface = ExistingInstance
-																					  && UsdUnreal::MaterialUtils::IsReferencePreviewSurfaceMaterial(
-																						  ExistingInstance->Parent
-																					  );
-							if (!OneSidedMat || (!bOneSidedMatIsInstanceOfReferencePreviewSurface && bExistingIsInstanceOfReferencePreviewSurface))
+							Material = ExistingMaterial;
+
+							if (bExistingIsDirectReference)
 							{
-								OneSidedMat = ExistingMaterial;
-								bOneSidedMatIsInstanceOfReferencePreviewSurface = bExistingIsInstanceOfReferencePreviewSurface;
+								// This is a perfect match, we don't need to keep looking
+								break;
 							}
 						}
+						else if (Slot.bMeshIsDoubleSided && !bExistingIsTwoSided)
+						{
+							// Keep track of this one-sided material to turn it into TwoSided later
 
-						if (Slot.bMeshIsDoubleSided == bExistingIsTwoSided)
+							// Prefer the one that is a direct preview surface reference if we have one already
+							if (!Material || (!bMaterialIsDirectReference && bExistingIsDirectReference))
+							{
+								Material = ExistingMaterial;
+								bMaterialIsDirectReference = bExistingIsDirectReference;
+							}
+						}
+						else // if (!Slot.bMeshIsDoubleSided && bExistingIsTwoSided)
 						{
-							Material = ExistingMaterial;
+							// We can ignore this case: If we're searching for a one sided material and just ran into
+							// an existing two-sided one we should just keep iterating: If a two-sided material is within
+							// ExistingMaterials, it's one-sided reference material *must* also be in there, so we'll find
+							// something eventually
 						}
 					}
 
@@ -607,28 +630,14 @@ TMap<const UsdUtils::FUsdPrimMaterialSlot*, UMaterialInterface*> MeshTranslation
 
 					// Need to create a two-sided material on-demand, *before* we make it compatible:
 					// This because at runtime we can't just set the base property overrides, and just instead create a new
-					// MIC based on the TwoSided reference material, and the compatible material should be a MIC of that MIC
-					if (Slot.bMeshIsDoubleSided && !Material)
+					// MID based on the TwoSided reference material, and the compatible material should be a MID of that MID
+					if (Slot.bMeshIsDoubleSided && Material && !Material->IsTwoSided())
 					{
-						// By now we parsed all materials so we must have the single-sided version of this material
-						if (!OneSidedMat)
-						{
-							UE_LOG(
-								LogUsd,
-								Warning,
-								TEXT(
-									"Failed to generate a two-sided material from the material prim at path '%s' as no single-sided material was generated for it."
-								),
-								*Slot.MaterialSource
-							);
-							continue;
-						}
-
-						const FString PrefixedOneSidedHash = AssetCache.GetHashForAsset(OneSidedMat);
+						const FString PrefixedOneSidedHash = AssetCache.GetHashForAsset(Material);
 						const FString PrefixedTwoSidedHash = PrefixedOneSidedHash + UnrealIdentifiers::TwoSidedMaterialSuffix;
 
 						UMaterialInterface* TwoSidedMat = UE::MeshTranslationImplInternal::Private::GetOrCreateTwoSidedVersionOfMaterial(
-							OneSidedMat,
+							Material,
 							PrefixedTwoSidedHash,
 							AssetCache
 						);
@@ -638,7 +647,7 @@ TMap<const UsdUtils::FUsdPrimMaterialSlot*, UMaterialInterface*> MeshTranslation
 							// Update AssetUserData whether we generated a new material or reused one from the asset cache
 							{
 								UUsdMaterialAssetUserData* OneSidedUserData = UsdUnreal::ObjectUtils::GetAssetUserData<UUsdMaterialAssetUserData>(
-									OneSidedMat
+									Material
 								);
 								ensure(OneSidedUserData);
 
@@ -667,16 +676,18 @@ TMap<const UsdUtils::FUsdPrimMaterialSlot*, UMaterialInterface*> MeshTranslation
 									"Failed to generate a two-sided material from the material prim at path '%s'. Falling back to using the single-sided material '%s' instead."
 								),
 								*Slot.MaterialSource,
-								*OneSidedMat->GetPathName()
+								*Material->GetPathName()
 							);
-							Material = OneSidedMat;
 							PrefixedMaterialHash = PrefixedOneSidedHash;
 						}
 					}
 
 					if (Material)
 					{
-						// Ping one-sided material as active
+						// Mark that we used this Material. We don't have to worry about our one-sided material because
+						// if we have one, it will be the two-sided's reference material, and we collect reference materials
+						// when collecting asset dependencies for import (which is the only mechanism that uses this
+						// TouchAsset/ActiveAssets stuff)
 						AssetCache.TouchAssetPath(Material);
 
 						InfoCache.LinkAssetToPrim(UE::FSdfPath{*Slot.MaterialSource}, Material);
