USD: Fix small bug introduced with 32532081 causing USD control rig bakes to produce keyframes at wrong times.

diff --git a/Engine/Plugins/Importers/USDImporter/Source/USDStage/Private/USDLevelSequenceHelper.cpp b/Engine/Plugins/Importers/USDImporter/Source/USDStage/Private/USDLevelSequenceHelper.cpp
--- a/Engine/Plugins/Importers/USDImporter/Source/USDStage/Private/USDLevelSequenceHelper.cpp
+++ b/Engine/Plugins/Importers/USDImporter/Source/USDStage/Private/USDLevelSequenceHelper.cpp
@@ -442,15 +442,11 @@ namespace UsdLevelSequenceHelperImpl
 			Track->SetTrackName(FName(*ObjectName));
 			Track->SetDisplayName(FText::FromString(ObjectName));
 
-			// We want the baked keyframe times to match exactly the source animation keyframe times. The way that
-			// LoadAnimSequenceIntoThisSection uses this however indicates it's meant to be relative to the start
-			// of the playback range, which we do here
-			ControlRigSectionStartFrame -= UE::MovieScene::DiscreteInclusiveLower(MovieScene->GetPlaybackRange());
-
 			const bool bResetControls = true;
+			const FFrameNumber SequenceStart{0};
 			ParamSection->LoadAnimSequenceIntoThisSection(
 				AnimSequence,
-				FFrameNumber(0),
+				SequenceStart,
 				MovieScene,
 				SkeletalMeshComp,
 				bReduceKeys,
@@ -738,9 +734,11 @@ private:
 
 	struct FPrimTwinBindings
 	{
-		// clang fix for std::is_default_constructible_v 
+		// clang fix for std::is_default_constructible_v
 		// returning false in inlined code of outer class
-		FPrimTwinBindings() {}
+		FPrimTwinBindings()
+		{
+		}
 
 		ULevelSequence* Sequence = nullptr;
 
@@ -3863,7 +3861,10 @@ FGuid FUsdLevelSequenceHelperImpl::GetOrCreateComponentBinding(
 	// Make sure we always bind the parent actor too
 	if (AActor* Actor = ComponentToBind.GetOwner())
 	{
-		TSharedRef<const UE::MovieScene::FSharedPlaybackState> SharedPlaybackState = MovieSceneHelpers::CreateTransientSharedPlaybackState(Actor, &Sequence);
+		TSharedRef<const UE::MovieScene::FSharedPlaybackState> SharedPlaybackState = MovieSceneHelpers::CreateTransientSharedPlaybackState(
+			Actor,
+			&Sequence
+		);
 
 		ActorBinding = Sequence.FindBindingFromObject(Actor, SharedPlaybackState);
 		if (!ActorBinding.IsValid())
