USD: Fixed crash in standalone game when loading a level with a USDStageActor.

diff --git a/Engine/Plugins/Importers/USDImporter/Source/USDStage/Private/USDLevelSequenceHelper.cpp b/Engine/Plugins/Importers/USDImporter/Source/USDStage/Private/USDLevelSequenceHelper.cpp
--- a/Engine/Plugins/Importers/USDImporter/Source/USDStage/Private/USDLevelSequenceHelper.cpp
+++ b/Engine/Plugins/Importers/USDImporter/Source/USDStage/Private/USDLevelSequenceHelper.cpp
@@ -274,7 +274,7 @@ namespace UsdLevelSequenceHelperImpl
 	TSharedPtr<ISequencer> GetOpenedSequencerForLevelSequence(ULevelSequence* LevelSequence)
 	{
 		const bool bFocusIfOpen = false;
-		IAssetEditorInstance* AssetEditor = GEditor->GetEditorSubsystem<UAssetEditorSubsystem>()->FindEditorForAsset(LevelSequence, bFocusIfOpen);
+		IAssetEditorInstance* AssetEditor = GEditor ? GEditor->GetEditorSubsystem<UAssetEditorSubsystem>()->FindEditorForAsset(LevelSequence, bFocusIfOpen) : nullptr;
 		ILevelSequenceEditorToolkit* LevelSequenceEditor = static_cast<ILevelSequenceEditorToolkit*>(AssetEditor);
 		return LevelSequenceEditor ? LevelSequenceEditor->GetSequencer() : nullptr;
 	}
diff --git a/Engine/Plugins/Importers/USDImporter/Source/USDStage/Private/USDStageActor.cpp b/Engine/Plugins/Importers/USDImporter/Source/USDStage/Private/USDStageActor.cpp
--- a/Engine/Plugins/Importers/USDImporter/Source/USDStage/Private/USDStageActor.cpp
+++ b/Engine/Plugins/Importers/USDImporter/Source/USDStage/Private/USDStageActor.cpp
@@ -3601,7 +3601,7 @@ void AUsdStageActor::ReloadAnimations()
 	if (!IsTemplate())
 	{
 #if WITH_EDITOR
-		UAssetEditorSubsystem* AssetEditorSubsystem = GEditor->GetEditorSubsystem<UAssetEditorSubsystem>();
+		UAssetEditorSubsystem* AssetEditorSubsystem = GEditor ? GEditor->GetEditorSubsystem<UAssetEditorSubsystem>() : nullptr;
 		const bool bSequencerWasOpened = AssetEditorSubsystem && AssetEditorSubsystem->FindEditorsForAssetAndSubObjects(LevelSequence).Num() > 0;
 #endif	  // WITH_EDITOR
 
