USD: Use Universal Object Locators for persistent bindings.

diff --git a/Engine/Plugins/Importers/USDImporter/Source/USDSchemas/Private/USDInfoCache.cpp b/Engine/Plugins/Importers/USDImporter/Source/USDSchemas/Private/USDInfoCache.cpp
--- a/Engine/Plugins/Importers/USDImporter/Source/USDSchemas/Private/USDInfoCache.cpp
+++ b/Engine/Plugins/Importers/USDImporter/Source/USDSchemas/Private/USDInfoCache.cpp
@@ -1687,7 +1687,7 @@ TArray<TWeakObjectPtr<UObject>> FUsdInfoCache::GetAllAssetsForPrim(const UE::FSd
 	return {};
 }
 
-TArray<UE::FSdfPath> FUsdInfoCache::GetPrimsForAsset(UObject* Asset) const
+TArray<UE::FSdfPath> FUsdInfoCache::GetPrimsForAsset(const UObject* Asset) const
 {
 	if (!Asset)
 	{
diff --git a/Engine/Plugins/Importers/USDImporter/Source/USDSchemas/Public/USDInfoCache.h b/Engine/Plugins/Importers/USDImporter/Source/USDSchemas/Public/USDInfoCache.h
--- a/Engine/Plugins/Importers/USDImporter/Source/USDSchemas/Public/USDInfoCache.h
+++ b/Engine/Plugins/Importers/USDImporter/Source/USDSchemas/Public/USDInfoCache.h
@@ -134,7 +134,7 @@ public:
 		return CastAssets;
 	}
 
-	TArray<UE::FSdfPath> GetPrimsForAsset(UObject* Asset) const;
+	TArray<UE::FSdfPath> GetPrimsForAsset(const UObject* Asset) const;
 	TMap<UE::FSdfPath, TArray<TWeakObjectPtr<UObject>>> GetAllAssetPrimLinks() const;
 
 private:
diff --git a/Engine/Plugins/Importers/USDImporter/Source/USDStage/Private/USDLocatorEditors.cpp b/Engine/Plugins/Importers/USDImporter/Source/USDStage/Private/USDLocatorEditors.cpp
new file mode 100644
--- /dev/null
+++ b/Engine/Plugins/Importers/USDImporter/Source/USDStage/Private/USDLocatorEditors.cpp
@@ -0,0 +1,205 @@
+// Copyright Epic Games, Inc. All Rights Reserved.
+
+#if WITH_EDITOR
+
+#include "USDLocatorEditors.h"
+
+#include "USDLocatorFragments.h"
+#include "USDStageActor.h"
+
+#include "DragAndDrop/ActorDragDropOp.h"
+#include "GameFramework/Actor.h"
+#include "IUniversalObjectLocatorCustomization.h"
+#include "SceneOutlinerDragDrop.h"
+#include "UniversalObjectLocator.h"
+#include "UniversalObjectLocatorEditor.h"
+#include "UniversalObjectLocatorFragmentTypeHandle.h"
+#include "Widgets/Input/SEditableTextBox.h"
+
+#define LOCTEXT_NAMESPACE "USDLocatorEditors"
+
+namespace UE::UniversalObjectLocator
+{
+	class SUsdPrimLocatorEditorUI : public SCompoundWidget
+	{
+		SLATE_BEGIN_ARGS(SUsdPrimLocatorEditorUI)
+		{
+		}
+		SLATE_END_ARGS()
+
+		void Construct(const FArguments& InArgs, TSharedPtr<IFragmentEditorHandle> InHandle)
+		{
+			WeakHandle = InHandle;
+
+			// clang-format off
+			ChildSlot
+			[
+				SNew(SHorizontalBox)
+				+SHorizontalBox::Slot()
+				.Padding(FMargin(8.0f, 4.0f))
+				.VAlign(VAlign_Center)
+				.AutoWidth()
+				[
+					SNew(STextBlock)
+					.Font(FAppStyle::GetFontStyle("PropertyWindow.NormalFont"))
+					.Text(LOCTEXT("PrimPathLabel", "Prim path:"))
+				]
+
+				+SHorizontalBox::Slot()
+				.FillWidth(1.0)
+				[
+					SNew(SEditableTextBox)
+					.MinDesiredWidth(300.0f)
+					.Text_Lambda([this]() -> FText
+					{
+						if (TSharedPtr<IFragmentEditorHandle> Handle = WeakHandle.Pin())
+						{
+							const FUniversalObjectLocatorFragment& Fragment = Handle->GetFragment();
+							if (const FUsdPrimLocatorFragment* CastFragment = Fragment.GetPayloadAs(FUsdPrimLocatorFragment::FragmentType))
+							{
+								return FText::FromString(CastFragment->PrimPath);
+							}
+						}
+
+						return FText::GetEmpty();
+					})
+					.Font(FAppStyle::GetFontStyle("PropertyWindow.NormalFont"))
+					.OnTextCommitted_Lambda([this](const FText& NewText, ETextCommit::Type CommitType)
+					{
+						if (TSharedPtr<IFragmentEditorHandle> Handle = WeakHandle.Pin())
+						{
+							const FUniversalObjectLocatorFragment& OldFragment = Handle->GetFragment();
+							const FUsdPrimLocatorFragment* OldCastFragment = OldFragment.GetPayloadAs(FUsdPrimLocatorFragment::FragmentType);
+
+							FUniversalObjectLocatorFragment NewFragment{FUsdPrimLocatorFragment::FragmentType};
+							FUsdPrimLocatorFragment* NewCastFragment = NewFragment.GetPayloadAs(FUsdPrimLocatorFragment::FragmentType);
+
+							NewCastFragment->PrimPath = NewText.ToString();
+							NewCastFragment->bPreferComponent = OldCastFragment->bPreferComponent;
+
+							Handle->SetValue(NewFragment);
+						}
+					})
+				]
+			];
+			// clang-format on
+		}
+
+		TWeakPtr<IFragmentEditorHandle> WeakHandle;
+	};
+
+	ELocatorFragmentEditorType FUsdPrimLocatorEditor::GetLocatorFragmentEditorType() const
+	{
+		return ELocatorFragmentEditorType::Relative;
+	}
+
+	bool FUsdPrimLocatorEditor::IsAllowedInContext(FName InContextName) const
+	{
+		return true;
+	}
+
+	// Note: I'm not exactly sure what is supposed to be dragged and dropped here, as any drag operation seems to
+	// close the right-click binding menu anyway? I can't click and drag the bindings themselves either...
+	bool FUsdPrimLocatorEditor::IsDragSupported(TSharedPtr<FDragDropOperation> DragOperation, UObject* Context) const
+	{
+		UObject* Result = ResolveDragOperation(DragOperation, Context);
+		return Result != nullptr;
+	}
+
+	UObject* FUsdPrimLocatorEditor::ResolveDragOperation(TSharedPtr<FDragDropOperation> DragOperation, UObject* Context) const
+	{
+		TSharedPtr<FActorDragDropOp> ActorDrag;
+
+		if (DragOperation->IsOfType<FSceneOutlinerDragDropOp>())
+		{
+			FSceneOutlinerDragDropOp* SceneOutlinerOp = static_cast<FSceneOutlinerDragDropOp*>(DragOperation.Get());
+			ActorDrag = SceneOutlinerOp->GetSubOp<FActorDragDropOp>();
+		}
+		else if (DragOperation->IsOfType<FActorDragDropOp>())
+		{
+			ActorDrag = StaticCastSharedPtr<FActorDragDropOp>(DragOperation);
+		}
+
+		if (ActorDrag)
+		{
+			for (const TWeakObjectPtr<AActor>& WeakActor : ActorDrag->Actors)
+			{
+				if (AActor* Actor = WeakActor.Get())
+				{
+					return FUsdPrimLocatorFragment::GetAttachParentStageActor(Actor);
+				}
+			}
+		}
+
+		return nullptr;
+	}
+
+	TSharedPtr<SWidget> FUsdPrimLocatorEditor::MakeEditUI(const FEditUIParameters& InParameters)
+	{
+		return SNew(SUsdPrimLocatorEditorUI, InParameters.Handle);
+	}
+
+	FText FUsdPrimLocatorEditor::GetDisplayText(const FUniversalObjectLocatorFragment* InFragment) const
+	{
+		if (InFragment != nullptr)
+		{
+			ensure(InFragment->GetFragmentTypeHandle() == FUsdPrimLocatorFragment::FragmentType);
+
+			if (const FUsdPrimLocatorFragment* CastFragment = InFragment->GetPayloadAs(FUsdPrimLocatorFragment::FragmentType))
+			{
+				return FText::FromString(FPaths::GetBaseFilename(CastFragment->PrimPath));
+			}
+		}
+
+		return LOCTEXT("UsdPrimLocatorEditorDisplayText", "USD Prim");
+	}
+
+	FText FUsdPrimLocatorEditor::GetDisplayTooltip(const FUniversalObjectLocatorFragment* InFragment) const
+	{
+		if (InFragment != nullptr)
+		{
+			ensure(InFragment->GetFragmentTypeHandle() == FUsdPrimLocatorFragment::FragmentType);
+
+			if (const FUsdPrimLocatorFragment* CastFragment = InFragment->GetPayloadAs(FUsdPrimLocatorFragment::FragmentType))
+			{
+				return FText::Format(LOCTEXT("UsdPrimLocatorEditorTooltip", "A path to prim '{0}'"), FText::FromString(CastFragment->PrimPath));
+			}
+		}
+
+		return LOCTEXT("UsdPrimLocatorEditorTooltip", "A path to a prim on a USD stage");
+	}
+
+	FSlateIcon FUsdPrimLocatorEditor::GetDisplayIcon(const FUniversalObjectLocatorFragment* InFragment) const
+	{
+		return FSlateIcon(FAppStyle::GetAppStyleSetName(), "LevelEditor.Tabs.USDStage");
+	}
+
+	UClass* FUsdPrimLocatorEditor::ResolveClass(const FUniversalObjectLocatorFragment& InFragment, UObject* InContext) const
+	{
+		if (UClass* Class = ILocatorFragmentEditor::ResolveClass(InFragment, InContext))
+		{
+			return Class;
+		}
+
+		if (const FUsdPrimLocatorFragment* CastFragment = InFragment.GetPayloadAs(FUsdPrimLocatorFragment::FragmentType))
+		{
+			if (!CastFragment->bPreferComponent)
+			{
+				return AActor::StaticClass();
+			}
+		}
+
+		return USceneComponent::StaticClass();
+	}
+
+	FUniversalObjectLocatorFragment FUsdPrimLocatorEditor::MakeDefaultLocatorFragment() const
+	{
+		FUniversalObjectLocatorFragment NewFragment(FUsdPrimLocatorFragment::FragmentType);
+		return NewFragment;
+	}
+
+}	 // namespace UE::UniversalObjectLocator
+
+#undef LOCTEXT_NAMESPACE
+
+#endif	  // WITH_EDITOR
diff --git a/Engine/Plugins/Importers/USDImporter/Source/USDStage/Private/USDLocatorEditors.h b/Engine/Plugins/Importers/USDImporter/Source/USDStage/Private/USDLocatorEditors.h
new file mode 100644
--- /dev/null
+++ b/Engine/Plugins/Importers/USDImporter/Source/USDStage/Private/USDLocatorEditors.h
@@ -0,0 +1,31 @@
+// Copyright Epic Games, Inc. All Rights Reserved.
+
+#pragma once
+
+#if WITH_EDITOR
+
+#include "CoreTypes.h"
+#include "UniversalObjectLocatorEditor.h"
+
+namespace UE::UniversalObjectLocator
+{
+	class FUsdPrimLocatorEditor : public ILocatorFragmentEditor
+	{
+	public:
+		// Begin ILocatorFragmentEditor interface
+		ELocatorFragmentEditorType GetLocatorFragmentEditorType() const override;
+		bool IsAllowedInContext(FName InContextName) const override;
+		bool IsDragSupported(TSharedPtr<FDragDropOperation> DragOperation, UObject* Context) const override;
+		UObject* ResolveDragOperation(TSharedPtr<FDragDropOperation> DragOperation, UObject* Context) const override;
+		TSharedPtr<SWidget> MakeEditUI(const FEditUIParameters& InParameters) override;
+		FText GetDisplayText(const FUniversalObjectLocatorFragment* InFragment) const override;
+		FText GetDisplayTooltip(const FUniversalObjectLocatorFragment* InFragment) const override;
+		FSlateIcon GetDisplayIcon(const FUniversalObjectLocatorFragment* InFragment) const override;
+		UClass* ResolveClass(const FUniversalObjectLocatorFragment& InFragment, UObject* InContext) const override;
+		FUniversalObjectLocatorFragment MakeDefaultLocatorFragment() const override;
+		// End ILocatorFragmentEditor interface
+	};
+
+}	 // namespace UE::UniversalObjectLocator
+
+#endif	  // WITH_EDITOR
diff --git a/Engine/Plugins/Importers/USDImporter/Source/USDStage/Private/USDLocatorFragments.cpp b/Engine/Plugins/Importers/USDImporter/Source/USDStage/Private/USDLocatorFragments.cpp
new file mode 100644
--- /dev/null
+++ b/Engine/Plugins/Importers/USDImporter/Source/USDStage/Private/USDLocatorFragments.cpp
@@ -0,0 +1,145 @@
+// Copyright Epic Games, Inc. All Rights Reserved.
+
+#include "USDLocatorFragments.h"
+
+#include "USDStageActor.h"
+
+#include "UniversalObjectLocatorFragmentTypeHandle.h"
+#include "UniversalObjectLocatorInitializeParams.h"
+#include "UniversalObjectLocatorInitializeResult.h"
+#include "UniversalObjectLocatorResolveParams.h"
+#include "UniversalObjectLocatorStringParams.h"
+
+UE::UniversalObjectLocator::TFragmentTypeHandle<FUsdPrimLocatorFragment> FUsdPrimLocatorFragment::FragmentType;
+
+namespace UE::USDLocatorFragments::Private
+{
+	const static FString ActorLocatorPrefix = TEXT("!actor!");
+}	 // namespace UE::USDLocatorFragments::Private
+
+UE::UniversalObjectLocator::FResolveResult FUsdPrimLocatorFragment::Resolve(const UE::UniversalObjectLocator::FResolveParams& Params) const
+{
+	using namespace UE::UniversalObjectLocator;
+
+	UObject* Result = nullptr;
+
+	if (AUsdStageActor* StageActorContext = Cast<AUsdStageActor>(Params.Context))
+	{
+		Result = StageActorContext->GetGeneratedComponent(PrimPath);
+
+		// To to resolve to the actor if we should
+		if (!bPreferComponent)
+		{
+			if (USceneComponent* Component = Cast<USceneComponent>(Result))
+			{
+				if (AActor* OwnerActor = Component->GetOwner())
+				{
+					if (Component == OwnerActor->GetRootComponent())
+					{
+						Result = OwnerActor;
+					}
+				}
+			}
+		}
+	}
+
+	return FResolveResultData(Result);
+}
+
+void FUsdPrimLocatorFragment::ToString(FStringBuilderBase& OutStringBuilder) const
+{
+	if (!bPreferComponent)
+	{
+		OutStringBuilder.Append(UE::USDLocatorFragments::Private::ActorLocatorPrefix);
+	}
+	OutStringBuilder.Append(PrimPath);
+}
+
+UE::UniversalObjectLocator::FParseStringResult FUsdPrimLocatorFragment::TryParseString(
+	FStringView InString,
+	const UE::UniversalObjectLocator::FParseStringParams& Params
+)
+{
+	using namespace UE::USDLocatorFragments::Private;
+
+	bPreferComponent = true;
+	if (InString.StartsWith(ActorLocatorPrefix))
+	{
+		bPreferComponent = false;
+		InString.RemovePrefix(ActorLocatorPrefix.Len());
+	}
+
+	PrimPath = InString;
+	return UE::UniversalObjectLocator::FParseStringResult().Success();
+}
+
+UE::UniversalObjectLocator::FInitializeResult FUsdPrimLocatorFragment::Initialize(const UE::UniversalObjectLocator::FInitializeParams& InParams)
+{
+	using namespace UE::UniversalObjectLocator;
+	using namespace UE::USDLocatorFragments::Private;
+
+	const UObject* Object = InParams.Object;
+	const AUsdStageActor* StageActor = GetAttachParentStageActor(Object);
+
+	if (StageActor)
+	{
+		FString FoundPrimPath = StageActor->GetSourcePrimPath(Object);
+		if (FoundPrimPath.IsEmpty())
+		{
+			return FInitializeResult::Failure();
+		}
+
+		PrimPath = FoundPrimPath;
+		bPreferComponent = Object->IsA<USceneComponent>();
+	}
+
+	return FInitializeResult::Relative(StageActor);
+}
+
+uint32 FUsdPrimLocatorFragment::ComputePriority(const UObject* ObjectToReference, const UObject* Context)
+{
+	using namespace UE::USDLocatorFragments::Private;
+
+	// Actors/components spawned by the stage actor are always transient
+	if (!ObjectToReference->HasAnyFlags(RF_Transient))
+	{
+		return 0;
+	}
+
+	// If we're inside a AUsdStageActor's attachment hierarchy, let's assume it's one of the actor's spawned objects
+	if (GetAttachParentStageActor(ObjectToReference))
+	{
+		return 3000;
+	}
+
+	return 0;
+}
+
+const AUsdStageActor* FUsdPrimLocatorFragment::GetAttachParentStageActor(const UObject* Object)
+{
+	const AActor* Actor = Cast<const AActor>(Object);
+	if (!Actor)
+	{
+		if (const USceneComponent* Component = Cast<const USceneComponent>(Object))
+		{
+			Actor = Component->GetOwner();
+		}
+	}
+
+	while (Actor)
+	{
+		if (const AUsdStageActor* StageActor = Cast<AUsdStageActor>(Actor))
+		{
+			return StageActor;
+		}
+
+		Actor = Actor->GetAttachParentActor();
+	}
+
+	return nullptr;
+}
+
+AUsdStageActor* FUsdPrimLocatorFragment::GetAttachParentStageActor(UObject* Object)
+{
+	return const_cast<AUsdStageActor*>(GetAttachParentStageActor(const_cast<const UObject*>(Object)));
+}
diff --git a/Engine/Plugins/Importers/USDImporter/Source/USDStage/Private/USDStageActor.cpp b/Engine/Plugins/Importers/USDImporter/Source/USDStage/Private/USDStageActor.cpp
--- a/Engine/Plugins/Importers/USDImporter/Source/USDStage/Private/USDStageActor.cpp
+++ b/Engine/Plugins/Importers/USDImporter/Source/USDStage/Private/USDStageActor.cpp
@@ -591,200 +591,6 @@ struct FUsdStageActorImpl
 #endif	  // WITH_EDITOR
 	}
 
-	static void SetupDynamicBinding(
-		const FString& PrimPath,
-		FGuid BindingID,
-		ULevelSequence* Sequence,
-		const FString& DefaultActorFilter,
-		TSharedPtr<ISequencer> Sequencer
-	)
-	{
-#if WITH_EDITOR
-		UMovieScene* MovieScene = Sequence->GetMovieScene();
-		if (!MovieScene)
-		{
-			return;
-		}
-
-		// Make sure we trigger the creation of the DirectorBlueprint before we add our binding, or else
-		// FMovieSceneDynamicBindingUtils::ResolveDynamicBinding may fail. Note: The regular code path for dynamic
-		// bindings via the UI ends up calling these exact same lines from
-		// FMovieSceneDirectorBlueprintEndpointCustomization::PopulateQuickBindSubMenu
-		FMovieSceneSequenceEditor* SequenceEditor = FMovieSceneSequenceEditor::Find(Sequence);
-		if (!SequenceEditor)
-		{
-			return;
-		}
-		UBlueprint* DirectorBlueprint = SequenceEditor->GetOrCreateDirectorBlueprint(Sequence);
-		if (!DirectorBlueprint)
-		{
-			return;
-		}
-
-		UMovieSceneReplaceableDirectorBlueprintBinding* NewCustomBinding = Cast<UMovieSceneReplaceableDirectorBlueprintBinding>(
-			UMovieSceneReplaceableDirectorBlueprintBinding::StaticClass()->GetDefaultObject<UMovieSceneCustomBinding>()->CreateNewCustomBinding(
-				nullptr,
-				*MovieScene
-			)
-		);
-
-		if (!NewCustomBinding)
-		{
-			return;
-		}
-
-		FMovieSceneDynamicBinding& DynamicBinding = NewCustomBinding->DynamicBinding;
-		// Setup the Dynamic Binding
-		{
-			DynamicBinding = FMovieSceneDynamicBinding{};
-			DynamicBinding.CompiledFunctionName = NAME_None;
-			DynamicBinding.ResolveParamsPinName = TEXT("Params");
-			DynamicBinding.Function = UUsdDynamicBindingResolverLibrary::StaticClass()->FindFunctionByName(
-				GET_FUNCTION_NAME_CHECKED(UUsdDynamicBindingResolverLibrary, ResolveWithStageActor)
-			);
-			if (DynamicBinding.Function)
-			{
-				DynamicBinding.ResolveParamsProperty = DynamicBinding.Function->FindPropertyByName(DynamicBinding.ResolveParamsPinName);
-			}
-
-			// Store a path to this very actor on the binding, so that it can find us later and ask how to resolve a particular
-			// prim path
-			FMovieSceneDynamicBindingPayloadVariable& ActorPathVariable = DynamicBinding.PayloadVariables.FindOrAdd(TEXT("StageActorIDNam"
-																														 "eFilter"));
-			ActorPathVariable.Value = DefaultActorFilter;
-
-			FMovieSceneDynamicBindingPayloadVariable& RootLayerVariable = DynamicBinding.PayloadVariables.FindOrAdd(TEXT("RootLayerFilter"));
-			RootLayerVariable.Value = FString{};	// No root layer filter by default for more flexibility
-
-			FMovieSceneDynamicBindingPayloadVariable& PrimPathVariable = DynamicBinding.PayloadVariables.FindOrAdd(TEXT("PrimPath"));
-			PrimPathVariable.Value = PrimPath;
-		}
-
-		Sequence->Modify();
-		MovieScene->Modify();
-
-		UE::Sequencer::FCreateBindingParams CreateBindingParams;
-		CreateBindingParams.ReplacementGuid = BindingID;
-		CreateBindingParams.BindingIndex = 0;
-		CreateBindingParams.BindingNameOverride = NewCustomBinding->GetDesiredBindingName();
-		CreateBindingParams.CustomBinding = NewCustomBinding;
-
-		FSequencerUtilities::CreateOrReplaceBinding(Sequencer.ToSharedRef(), nullptr, CreateBindingParams);
-
-		// Setup the Blueprint function ("endpoint") on the Sequence's director Blueprint.
-		// In the future we could also just setup a single endpoint for all dynamic bindings, as they're all identical here. However, even then
-		// I think we should allow the option of having this system emit a separatate endpoint for each binding like it does now, as the user
-		// may want to modify these dynamic binding functions for some reason (maybe add a prim path suffix, or retrieve the prim path from
-		// some other system, etc.).
-		// Reference: FMovieSceneDynamicBindingCustomization::GenerateEndpointDefinition
-		UK2Node_FunctionEntry* NewEndpoint = nullptr;
-		{
-			FMovieSceneDirectorBlueprintEndpointDefinition EndpointDefinition;
-			EndpointDefinition.EndpointType = EMovieSceneDirectorBlueprintEndpointType::Function;
-
-			// Use the same dummy function the regular UI-created Dynamic Bindings use to setup the endpoint signature, in case it changes
-			static const FName SampleResolveBindingFuncName{TEXT("SampleResolveBinding")};
-			UClass* EndpointUtilClass = UMovieSceneDynamicBindingEndpointUtil::StaticClass();
-			EndpointDefinition.EndpointSignature = EndpointUtilClass->FindFunctionByName(SampleResolveBindingFuncName);
-			if (!EndpointDefinition.EndpointSignature)
-			{
-				return;
-			}
-
-			FMovieScenePossessable* Possessable = Sequence->GetMovieScene()->FindPossessable(BindingID);
-			if (!Possessable)
-			{
-				return;
-			}
-			EndpointDefinition.EndpointName = Possessable->GetName() + TEXT("_DynamicBinding");
-
-			DirectorBlueprint->Modify();
-
-			// Creates the new function on the LevelSequence's director blueprint
-			NewEndpoint = FMovieSceneDirectorBlueprintUtils::CreateFunctionEndpoint(DirectorBlueprint, EndpointDefinition);
-			if (!NewEndpoint)
-			{
-				return;
-			}
-			NewEndpoint->MetaData.bCallInEditor = true;
-		}
-
-		// Create a function call node within that endpoint that calls our UFunction
-		// Reference: FMovieSceneDirectorBlueprintEndpointCustomization::HandleQuickBindActionSelected
-		UBlueprintFunctionNodeSpawner* Spawner = UBlueprintFunctionNodeSpawner::Create(DynamicBinding.Function);
-		FBlueprintActionMenuItem Action{Spawner};
-		UEdGraphPin* ThenPin = NewEndpoint->FindPin(UEdGraphSchema_K2::PN_Then, EGPD_Output);
-		FVector2D NodePosition(NewEndpoint->NodePosX + 400.f, NewEndpoint->NodePosY + 100.0f);
-		UEdGraphNode* NewNode = Action.PerformAction(NewEndpoint->GetGraph(), ThenPin, NodePosition);
-
-		// Connect the function entry point to our function call node if it has an exec pin
-		UEdGraphPin* NewNodeExecPin = NewNode->FindPin(UEdGraphSchema_K2::PN_Execute, EGPD_Input);
-		if (ThenPin && NewNodeExecPin)
-		{
-			ThenPin->MakeLinkTo(NewNodeExecPin);
-		}
-
-		// Connect the function call result pin to the endpoint's result node
-		TArray<UK2Node_FunctionResult*> ResultNodes;
-		NewNode->GetGraph()->GetNodesOfClass(ResultNodes);
-		if (ResultNodes.Num() > 0)
-		{
-			// If there is a result node, move it past the endpoint call
-			ResultNodes[0]->NodePosX = NodePosition.X + 400.f;
-
-			// Connect the function call node to the endpoint result node if it has an exec pin
-			UEdGraphPin* NewNodeThenPin = NewNode->FindPin(UEdGraphSchema_K2::PN_Execute, EGPD_Input);
-			UEdGraphPin* ResultExecPin = ResultNodes[0]->FindPin(UEdGraphSchema_K2::PN_Execute, EGPD_Input);
-			if (NewNodeThenPin && ResultExecPin)
-			{
-				NewNodeThenPin->MakeLinkTo(ResultExecPin);
-			}
-
-			// If the new node has a return value, and if the endpoint has one too, try to connect them together.
-			UEdGraphPin* OutputPin = ResultNodes[0]->FindPin(UEdGraphSchema_K2::PN_ReturnValue, EGPD_Input);
-			UEdGraphPin* NewNodeReturnValuePin = NewNode->FindPin(UEdGraphSchema_K2::PN_ReturnValue, EGPD_Output);
-			if (OutputPin && NewNodeReturnValuePin)
-			{
-				// Connect the nodes.
-				NewNodeReturnValuePin->MakeLinkTo(OutputPin);
-			}
-		}
-
-		// Connect additional parameter pins from the endpoint to the function call node
-		// Reference: FMovieSceneDirectorBlueprintEndpointCustomization::SetEndpoint
-		TSet<FName> NonPayloadPins;
-		const FString* WorldContextParamName = DynamicBinding.Function->FindMetaData(FBlueprintMetadata::MD_WorldContext);
-		if (WorldContextParamName)
-		{
-			NonPayloadPins.Add(FName(*WorldContextParamName));
-		}
-		if (UK2Node_EditablePinBase* EditableNode = Cast<UK2Node_EditablePinBase>(NewEndpoint))
-		{
-			for (UEdGraphPin* PayloadPin : NewNode->Pins)
-			{
-				if (PayloadPin && PayloadPin->LinkedTo.Num() == 0 && PayloadPin->Direction == EGPD_Input
-					&& PayloadPin->PinType.PinCategory != UEdGraphSchema_K2::PC_Exec && PayloadPin->PinName != UEdGraphSchema_K2::PN_Self
-					&& !NonPayloadPins.Contains(PayloadPin->PinName))
-				{
-					// Pins for ref parameters for functions default to bIsReference but the payload cannot be by
-					// reference.
-					PayloadPin->PinType.bIsReference = false;
-
-					UEdGraphPin* NewPin = EditableNode->CreateUserDefinedPin(PayloadPin->PinName, PayloadPin->PinType, EGPD_Output);
-					if (NewNode != NewEndpoint && NewPin)
-					{
-						NewPin->MakeLinkTo(PayloadPin);
-					}
-				}
-			}
-		}
-
-		FMovieSceneDynamicBindingUtils::SetEndpoint(MovieScene, &DynamicBinding, NewEndpoint);
-		FMovieSceneDynamicBindingUtils::EnsureBlueprintExtensionCreated(Sequence, DirectorBlueprint);
-		FKismetEditorUtilities::CompileBlueprint(DirectorBlueprint);
-#endif	  // WITH_EDITOR
-	}
-
 	static void GetDescendantMovieSceneSequences(UMovieSceneSequence* InSequence, TSet<UMovieSceneSequence*>& OutAllSequences)
 	{
 		if (InSequence == nullptr || OutAllSequences.Contains(InSequence))
@@ -1198,25 +1004,6 @@ AUsdStageActor::AUsdStageActor()
 			);
 		}
 
-		TFunction<void(TSharedRef<ISequencer>)> SubscribeToSequencerEvents = [this](TSharedRef<ISequencer> Sequencer)
-		{
-			Sequencer->OnActorAddedToSequencer().AddUObject(this, &AUsdStageActor::OnActorAddedToSequencer, Sequencer.ToWeakPtr());
-			Sequencer->OnMovieSceneDataChanged().AddUObject(this, &AUsdStageActor::OnMovieSceneDataChanged, Sequencer.ToWeakPtr());
-		};
-
-		// Make sure we are bound to the OnActorAddedToSequencer of every sequencer that exists
-		IUsdStageModule& UsdStageModule = FModuleManager::Get().LoadModuleChecked<IUsdStageModule>(TEXT("UsdStage"));
-		for (const TWeakPtr<ISequencer>& ExistingSequencer : UsdStageModule.GetExistingSequencers())
-		{
-			if (TSharedPtr<ISequencer> PinnedSequencer = ExistingSequencer.Pin())
-			{
-				SubscribeToSequencerEvents(PinnedSequencer.ToSharedRef());
-			}
-		}
-		ISequencerModule& SequencerModule = FModuleManager::Get().LoadModuleChecked<ISequencerModule>(TEXT("Sequencer"));
-		OnSequencerCreatedHandle = SequencerModule.RegisterOnSequencerCreated(FOnSequencerCreated::FDelegate::CreateLambda(SubscribeToSequencerEvents)
-		);
-
 		FCoreUObjectDelegates::OnObjectPropertyChanged.AddUObject(this, &AUsdStageActor::OnObjectPropertyChanged);
 
 		// Also prevent standalone from doing this
@@ -2976,7 +2763,7 @@ ULevelSequence* AUsdStageActor::GetLevelSequence()
 	return LevelSequence;
 }
 
-USceneComponent* AUsdStageActor::GetGeneratedComponent(const FString& PrimPath)
+USceneComponent* AUsdStageActor::GetGeneratedComponent(const FString& PrimPath) const
 {
 	const UE::FUsdStage& CurrentStage = static_cast<const AUsdStageActor*>(this)->GetUsdStage();
 	if (!CurrentStage || StageState != EUsdStageState::OpenedAndLoaded)
@@ -3005,7 +2792,7 @@ USceneComponent* AUsdStageActor::GetGeneratedComponent(const FString& PrimPath)
 	return nullptr;
 }
 
-TArray<UObject*> AUsdStageActor::GetGeneratedAssets(const FString& PrimPath)
+TArray<UObject*> AUsdStageActor::GetGeneratedAssets(const FString& PrimPath) const
 {
 	const UE::FUsdStage& CurrentStage = static_cast<const AUsdStageActor*>(this)->GetUsdStage();
 	if (!CurrentStage || StageState != EUsdStageState::OpenedAndLoaded)
@@ -3042,11 +2829,24 @@ TArray<UObject*> AUsdStageActor::GetGeneratedAssets(const FString& PrimPath)
 	return Assets;
 }
 
-FString AUsdStageActor::GetSourcePrimPath(UObject* Object)
+FString AUsdStageActor::GetSourcePrimPath(const UObject* Object) const
 {
-	if (USceneComponent* Component = Cast<USceneComponent>(Object))
+	UUsdPrimTwin* RootUsdPrimTwin = GetRootPrimTwin();
+
+	const USceneComponent* Component = Cast<const USceneComponent>(Object);
+	if (!Component)
 	{
-		if (UUsdPrimTwin* UsdPrimTwin = GetRootPrimTwin()->Find(Component))
+		// We always bind the root component and actor itself to the same prim anyway, so let's
+		// just decay to the component in case we've been given an actor
+		if (const AActor* Actor = Cast<AActor>(Object))
+		{
+			Component = Actor->GetRootComponent();
+		}
+	}
+
+	if (Component && RootUsdTwin)
+	{
+		if (UUsdPrimTwin* UsdPrimTwin = RootUsdPrimTwin->Find(Component))
 		{
 			return UsdPrimTwin->PrimPath;
 		}
@@ -3566,6 +3366,11 @@ UUsdPrimTwin* AUsdStageActor::GetRootPrimTwin()
 	return RootUsdTwin;
 }
 
+UUsdPrimTwin* AUsdStageActor::GetRootPrimTwin() const
+{
+	return RootUsdTwin;
+}
+
 void AUsdStageActor::Refresh() const
 {
 	OnTimeChanged.Broadcast();
@@ -4016,23 +3821,6 @@ void AUsdStageActor::BeginDestroy()
 			TransBuffer->OnRedo().Remove(OnRedoHandle);
 		}
 
-		if (ISequencerModule* SequencerModule = FModuleManager::Get().GetModulePtr<ISequencerModule>(TEXT("Sequencer")))
-		{
-			SequencerModule->UnregisterOnSequencerCreated(OnSequencerCreatedHandle);
-
-			if (IUsdStageModule* UsdStageModule = FModuleManager::Get().GetModulePtr<IUsdStageModule>(TEXT("UsdStage")))
-			{
-				for (const TWeakPtr<ISequencer>& ExistingSequencer : UsdStageModule->GetExistingSequencers())
-				{
-					if (TSharedPtr<ISequencer> PinnedSequencer = ExistingSequencer.Pin())
-					{
-						PinnedSequencer->OnActorAddedToSequencer().RemoveAll(this);
-						PinnedSequencer->OnMovieSceneDataChanged().RemoveAll(this);
-					}
-				}
-			}
-		}
-
 		GEngine->OnLevelActorDeleted().RemoveAll(this);
 		FCoreUObjectDelegates::OnObjectsReplaced.RemoveAll(this);
 	}
@@ -4245,226 +4033,6 @@ void AUsdStageActor::UpdateSpawnedObjectsTransientFlag(bool bTransient)
 	GetRootPrimTwin()->Iterate(UpdateTransient, bRecursive);
 }
 
-void AUsdStageActor::OnActorAddedToSequencer(AActor* NewActor, const FGuid Guid, TWeakPtr<ISequencer> WeakSequencer)
-{
-#if WITH_EDITOR
-
-	if (bIsAddingBinding)
-	{
-		// Prevent re-entrancy
-		return;
-	}
-	if (!NewActor || !NewActor->HasAnyFlags(RF_Transient))
-	{
-		return;
-	}
-
-	UUsdPrimTwin* UsdPrimTwin = GetRootPrimTwin()->Find(NewActor->GetRootComponent());
-	if (!UsdPrimTwin)
-	{
-		return;
-	}
-	const FString& PrimPath = UsdPrimTwin->PrimPath;
-
-	TSharedPtr<ISequencer> PinnedSequencer = WeakSequencer.Pin();
-	if (!PinnedSequencer)
-	{
-		return;
-	}
-
-	ULevelSequence* Sequence = Cast<ULevelSequence>(PinnedSequencer->GetRootMovieSceneSequence());
-	if (!Sequence)
-	{
-		return;
-	}
-
-	// We never need dynamic bindings on our (or other stage actor's) transient level sequences
-	if (Sequence == LevelSequence || Sequence->HasAnyFlags(RF_Transient))
-	{
-		return;
-	}
-
-	UMovieScene* MovieScene = Sequence->GetMovieScene();
-	if (!MovieScene)
-	{
-		return;
-	}
-
-	FMovieScenePossessable* Possessable = MovieScene->FindPossessable(Guid);
-	if (!Possessable)
-	{
-		return;
-	}
-
-
-	// The UI action to add the actor to the sequencer should have spawned it's transaction already, but just in case it
-	// hasn't let's do that, as we'll end up modifying UObjects in here
-	FScopedTransaction Transaction(
-		FText::Format(LOCTEXT("ActorAddedToSequencer", "Bound USD actor '{0}' to the Sequencer"), FText::FromString(NewActor->GetActorLabel()))
-	);
-	Sequence->Modify();
-	MovieScene->Modify();
-
-	bIsAddingBinding = true;
-	FUsdStageActorImpl::SetupDynamicBinding(PrimPath, Guid, Sequence, *GetName(), PinnedSequencer);
-	bIsAddingBinding = false;
-#endif	  // WITH_EDITOR
-}
-
-void AUsdStageActor::OnMovieSceneDataChanged(EMovieSceneDataChangeType ChangeType, TWeakPtr<ISequencer> WeakSequencer)
-{
-#if WITH_EDITOR
-	if (ChangeType != EMovieSceneDataChangeType::MovieSceneStructureItemAdded
-		&& ChangeType != EMovieSceneDataChangeType::MovieSceneStructureItemsChanged)
-	{
-		return;
-	}
-
-	TSharedPtr<ISequencer> PinnedSequencer = WeakSequencer.Pin();
-	if (!PinnedSequencer)
-	{
-		return;
-	}
-
-	ULevelSequence* Sequence = Cast<ULevelSequence>(PinnedSequencer->GetRootMovieSceneSequence());
-	if (!Sequence)
-	{
-		return;
-	}
-
-	// We never need dynamic bindings on our (or other stage actor's) transient level sequences
-	if (Sequence == LevelSequence || Sequence->HasAnyFlags(RF_Transient))
-	{
-		return;
-	}
-
-	UUsdLevelSequenceAssetUserData* AssetUserData = Sequence->GetAssetUserData<UUsdLevelSequenceAssetUserData>();
-	if (!AssetUserData)
-	{
-		AssetUserData = NewObject<UUsdLevelSequenceAssetUserData>(Sequence);
-		Sequence->AddAssetUserData(AssetUserData);
-	}
-
-	const FGuid& CurrentSignature = Sequence->GetSignature();
-	if (CurrentSignature == AssetUserData->LastCheckedSignature)
-	{
-		return;
-	}
-	AssetUserData->LastCheckedSignature = CurrentSignature;
-
-	UMovieScene* MovieScene = Sequence->GetMovieScene();
-	if (!MovieScene)
-	{
-		return;
-	}
-
-	// The UI action to add the actor to the sequencer should have spawned it's transaction already, but just in case it
-	// hasn't let's do that, as we'll end up modifying UObjects in here
-	FScopedTransaction Transaction(LOCTEXT("ComponentAddedToSequencer", "Added USD component bindings to the Sequencer"));
-	Sequence->Modify();
-	MovieScene->Modify();
-
-	TFunction<UObject*(const FMovieScenePossessable&)> LocateBoundObject = nullptr;
-	LocateBoundObject = [&LocateBoundObject, Sequence, MovieScene, PinnedSequencer](const FMovieScenePossessable& Possessable) -> UObject*
-	{
-		const FGuid& Guid = Possessable.GetGuid();
-		const FGuid& ParentGuid = Possessable.GetParent();
-
-		UObject* ParentContext = nullptr;
-		if (ParentGuid.IsValid())
-		{
-			if (FMovieScenePossessable* ParentPossessable = MovieScene->FindPossessable(ParentGuid))
-			{
-				ParentContext = LocateBoundObject(*ParentPossessable);
-			}
-		}
-
-		TArrayView<TWeakObjectPtr<>> Objects = PinnedSequencer->FindBoundObjects(Guid, PinnedSequencer->GetFocusedTemplateID());
-
-		if (Objects.Num() > 0)
-		{
-			return Objects[0].Get();
-		}
-
-		return nullptr;
-	};
-
-	int32 PossessableCount = MovieScene->GetPossessableCount();
-	for (int32 Index = 0; Index < PossessableCount; ++Index)
-	{
-		FMovieScenePossessable& Possessable = MovieScene->GetPossessable(Index);
-
-		// We only handle components on this code path, and component bindings always have a parent actor binding
-		if (!Possessable.GetParent().IsValid())
-		{
-			continue;
-		}
-
-		const FGuid& BindingGuid = Possessable.GetGuid();
-
-		const FMovieSceneBindingReferences* BindingReferences = Sequence->GetBindingReferences();
-		if (!BindingReferences)
-		{
-			continue;
-		}
-
-		// If the binding already has another dynamic binding let's not touch it,
-		// regardless of whether we set that dynamic binding up or the user did
-		if (Algo::AnyOf(
-				BindingReferences->GetReferences(Possessable.GetGuid()),
-				[](const FMovieSceneBindingReference& Reference)
-				{
-					if (UMovieSceneSpawnableDirectorBlueprintBinding* SpawnableBinding = Cast<UMovieSceneSpawnableDirectorBlueprintBinding>(
-							Reference.CustomBinding
-						))
-					{
-						if (SpawnableBinding->DynamicBinding.Function != nullptr)
-						{
-							return true;
-						}
-					}
-					if (UMovieSceneReplaceableDirectorBlueprintBinding* ReplaceableBinding = Cast<UMovieSceneReplaceableDirectorBlueprintBinding>(
-							Reference.CustomBinding
-						))
-					{
-						if (ReplaceableBinding->DynamicBinding.Function != nullptr)
-						{
-							return true;
-						}
-					}
-					return false;
-				}
-			))
-		{
-			continue;
-		}
-
-		if (AssetUserData->HandledBindingGuids.Contains(BindingGuid))
-		{
-			continue;
-		}
-
-		UActorComponent* BoundComponent = Cast<UActorComponent>(LocateBoundObject(Possessable));
-		if (!BoundComponent)
-		{
-			continue;
-		}
-
-		FString SourcePrimPath = GetSourcePrimPath(BoundComponent);
-		if (SourcePrimPath.IsEmpty())
-		{
-			continue;
-		}
-
-		// Only mark it as handled if we found the right stage actor for this component. Other stage actors that don't
-		// own this component will all return empty strings for SourcePrimPath
-		AssetUserData->HandledBindingGuids.Add(BindingGuid);
-
-		FUsdStageActorImpl::SetupDynamicBinding(SourcePrimPath, BindingGuid, Sequence, *GetName(), PinnedSequencer);
-	}
-#endif	  // WITH_EDITOR
-}
-
 void AUsdStageActor::OnUsdPrimTwinDestroyed(const UUsdPrimTwin& UsdPrimTwin)
 {
 	PrimsToAnimate.Remove(UsdPrimTwin.PrimPath);
diff --git a/Engine/Plugins/Importers/USDImporter/Source/USDStage/Private/USDStageModule.cpp b/Engine/Plugins/Importers/USDImporter/Source/USDStage/Private/USDStageModule.cpp
--- a/Engine/Plugins/Importers/USDImporter/Source/USDStage/Private/USDStageModule.cpp
+++ b/Engine/Plugins/Importers/USDImporter/Source/USDStage/Private/USDStageModule.cpp
@@ -2,6 +2,7 @@
 
 #include "USDStageModule.h"
 
+#include "USDLocatorFragments.h"
 #include "USDMemory.h"
 #include "USDStageActor.h"
 #include "USDStageActorCustomization.h"
@@ -9,10 +10,14 @@
 #include "Engine/World.h"
 #include "EngineUtils.h"
 #include "Interfaces/IPluginManager.h"
+#include "IUniversalObjectLocatorModule.h"
 #include "Modules/ModuleManager.h"
 
 #if WITH_EDITOR
+#include "USDLocatorEditors.h"
+
 #include "ISequencerModule.h"
+#include "IUniversalObjectLocatorEditorModule.h"
 #include "PropertyEditorModule.h"
 #endif	  // WITH_EDITOR
 
@@ -21,7 +26,19 @@ class FUsdStageModule : public IUsdStageModule
 public:
 	virtual void StartupModule() override
 	{
+		using namespace UE::UniversalObjectLocator;
+
+		IUniversalObjectLocatorModule& UolModule = FModuleManager::Get().LoadModuleChecked<IUniversalObjectLocatorModule>("UniversalObjectLocator");
+		FFragmentTypeParameters Parameters{"usdprim", NSLOCTEXT("USDStageModule", "UsdPrimLocatorFragment", "UsdPrim")};
+		Parameters.PrimaryEditorType = "UsdPrim";
+		FUsdPrimLocatorFragment::FragmentType = UolModule.RegisterFragmentType<FUsdPrimLocatorFragment>(Parameters);
+
 #if WITH_EDITOR
+		IUniversalObjectLocatorEditorModule& UolEditorModule = FModuleManager::LoadModuleChecked<IUniversalObjectLocatorEditorModule>(
+			"UniversalObjectLocatorEditor"
+		);
+		UolEditorModule.RegisterLocatorEditor("UsdPrim", MakeShared<FUsdPrimLocatorEditor>());
+
 		LLM_SCOPE_BYTAG(Usd);
 
 		FPropertyEditorModule& PropertyModule = FModuleManager::LoadModuleChecked<FPropertyEditorModule>(TEXT("PropertyEditor"));
@@ -54,7 +71,21 @@ public:
 
 	virtual void ShutdownModule() override
 	{
+		using namespace UE::UniversalObjectLocator;
+
+		if (IUniversalObjectLocatorModule* UolModule = FModuleManager::Get().GetModulePtr<IUniversalObjectLocatorModule>("UniversalObjectLocator"))
+		{
+			UolModule->UnregisterFragmentType(FUsdPrimLocatorFragment::FragmentType);
+		}
+
 #if WITH_EDITOR
+		if (IUniversalObjectLocatorEditorModule* UolEditorModule = FModuleManager::Get().GetModulePtr<IUniversalObjectLocatorEditorModule>(
+				"UniversalObjectLocatorEditor"
+			))
+		{
+			UolEditorModule->UnregisterLocatorEditor("UsdPrim");
+		}
+
 		FPropertyEditorModule& PropertyModule = FModuleManager::LoadModuleChecked<FPropertyEditorModule>(TEXT("PropertyEditor"));
 		PropertyModule.UnregisterCustomClassLayout(TEXT("UsdStageActor"));
 
diff --git a/Engine/Plugins/Importers/USDImporter/Source/USDStage/Public/USDLocatorFragments.h b/Engine/Plugins/Importers/USDImporter/Source/USDStage/Public/USDLocatorFragments.h
new file mode 100644
--- /dev/null
+++ b/Engine/Plugins/Importers/USDImporter/Source/USDStage/Public/USDLocatorFragments.h
@@ -0,0 +1,57 @@
+// Copyright Epic Games, Inc. All Rights Reserved.
+
+#pragma once
+
+#include "CoreTypes.h"
+#include "UniversalObjectLocatorFwd.h"
+
+#include "USDLocatorFragments.generated.h"
+
+class AUsdStageActor;
+
+/**
+ * 32 Bytes (40 in-editor).
+ */
+USTRUCT()
+struct FUsdPrimLocatorFragment
+{
+	GENERATED_BODY()
+
+	UPROPERTY()
+	bool bPreferComponent;
+
+	UPROPERTY()
+	FString PrimPath;
+
+	USDSTAGE_API static UE::UniversalObjectLocator::TFragmentTypeHandle<FUsdPrimLocatorFragment> FragmentType;
+
+	friend uint32 GetTypeHash(const FUsdPrimLocatorFragment& A)
+	{
+		return GetTypeHash(A.PrimPath) ^ GetTypeHash(A.bPreferComponent);
+	}
+
+	friend bool operator==(const FUsdPrimLocatorFragment& A, const FUsdPrimLocatorFragment& B)
+	{
+		return A.bPreferComponent == B.bPreferComponent && A.PrimPath == B.PrimPath;
+	}
+
+	UE::UniversalObjectLocator::FResolveResult Resolve(const UE::UniversalObjectLocator::FResolveParams& Params) const;
+	UE::UniversalObjectLocator::FInitializeResult Initialize(const UE::UniversalObjectLocator::FInitializeParams& InParams);
+
+	void ToString(FStringBuilderBase& OutStringBuilder) const;
+	UE::UniversalObjectLocator::FParseStringResult TryParseString(FStringView InString, const UE::UniversalObjectLocator::FParseStringParams& Params);
+
+	static uint32 ComputePriority(const UObject* Object, const UObject* Context);
+
+	static const AUsdStageActor* GetAttachParentStageActor(const UObject* Object);
+	static AUsdStageActor* GetAttachParentStageActor(UObject* Object);
+};
+
+template<>
+struct TStructOpsTypeTraits<FUsdPrimLocatorFragment> : public TStructOpsTypeTraitsBase2<FUsdPrimLocatorFragment>
+{
+	enum
+	{
+		WithIdenticalViaEquality = true,
+	};
+};
diff --git a/Engine/Plugins/Importers/USDImporter/Source/USDStage/Public/USDStageActor.h b/Engine/Plugins/Importers/USDImporter/Source/USDStage/Public/USDStageActor.h
--- a/Engine/Plugins/Importers/USDImporter/Source/USDStage/Public/USDStageActor.h
+++ b/Engine/Plugins/Importers/USDImporter/Source/USDStage/Public/USDStageActor.h
@@ -242,7 +242,7 @@ public:
 	 * invalid.
 	 */
 	UFUNCTION(BlueprintCallable, Category = "USD", meta = (CallInEditor = "true"))
-	USDSTAGE_API USceneComponent* GetGeneratedComponent(const FString& PrimPath);
+	USDSTAGE_API USceneComponent* GetGeneratedComponent(const FString& PrimPath) const;
 
 	/**
 	 * Gets the transient assets that were generated for a prim with a given prim path. Likely one asset (e.g. UStaticMesh), but can be multiple
@@ -251,7 +251,7 @@ public:
 	 * @return The corresponding generated assets. May be empty if path is invalid or if that prim led to no generated assets.
 	 */
 	UFUNCTION(BlueprintCallable, Category = "USD", meta = (CallInEditor = "true"))
-	USDSTAGE_API TArray<UObject*> GetGeneratedAssets(const FString& PrimPath);
+	USDSTAGE_API TArray<UObject*> GetGeneratedAssets(const FString& PrimPath) const;
 
 	/**
 	 * Gets the path to the prim that was parsed to generate the given `Object`.
@@ -260,7 +260,7 @@ public:
 	 * @return The path to the source prim, e.g. "/root_prim/some_prim". May be empty in case we couldn't find the source prim.
 	 */
 	UFUNCTION(BlueprintCallable, Category = "USD", meta = (CallInEditor = "true"))
-	USDSTAGE_API FString GetSourcePrimPath(UObject* Object);
+	USDSTAGE_API FString GetSourcePrimPath(const UObject* Object) const;
 
 	// Creates a brand new, memory-only USD stage and opens it
 	UFUNCTION(BlueprintCallable, Category = "USD", meta = (CallInEditor = "true"))
@@ -360,6 +360,7 @@ protected:
 	void AnimatePrims();
 
 	UUsdPrimTwin* GetRootPrimTwin();
+	UUsdPrimTwin* GetRootPrimTwin() const;
 	UUsdPrimTwin* GetOrCreatePrimTwin(const UE::FSdfPath& UsdPrimPath);
 
 	UUsdPrimTwin* ExpandPrim(const UE::FUsdPrim& Prim, bool bResync, FUsdSchemaTranslationContext& TranslationContext);
@@ -378,9 +379,6 @@ protected:
 
 	void UpdateSpawnedObjectsTransientFlag(bool bTransient);
 
-	void OnActorAddedToSequencer(AActor* NewActor, const FGuid Guid, TWeakPtr<ISequencer> WeakSequencer);
-	void OnMovieSceneDataChanged(EMovieSceneDataChangeType ChangeType, TWeakPtr<ISequencer> WeakSequencer);
-
 #if WITH_EDITOR
 	void OnBeginPIE(bool bIsSimulating);
 	void OnPostPIEStarted(bool bIsSimulating);
diff --git a/Engine/Plugins/Importers/USDImporter/Source/USDStage/USDStage.Build.cs b/Engine/Plugins/Importers/USDImporter/Source/USDStage/USDStage.Build.cs
--- a/Engine/Plugins/Importers/USDImporter/Source/USDStage/USDStage.Build.cs
+++ b/Engine/Plugins/Importers/USDImporter/Source/USDStage/USDStage.Build.cs
@@ -40,6 +40,7 @@ namespace UnrealBuildTool.Rules
 					"Slate",
 					"SlateCore",
 					"StaticMeshDescription",
+					"UniversalObjectLocator",
 					"UnrealUSDWrapper",
 					"USDClasses",
 					"USDUtilities",
@@ -53,7 +54,6 @@ namespace UnrealBuildTool.Rules
 					{
 						"BlueprintGraph", // For setting up the Sequencer Dynamic Binding blueprint graphs
 						"ControlRigDeveloper",
-						"RigVMDeveloper",
 						"DeveloperToolSettings",
 						"EditorStyle", // For the font style on the stage actor customization
 						"InputCore", // For keyboard control on the widget in the stage actor customization
@@ -61,7 +61,10 @@ namespace UnrealBuildTool.Rules
 						"LevelSequenceEditor",
 						"MovieSceneTools",
 						"PropertyEditor", // For the stage actor's details customization
+						"RigVMDeveloper",
+						"SceneOutliner", // For drag and drop operations on the Universal Object Locator editor
 						"Sequencer",
+						"UniversalObjectLocatorEditor",
 						"UnrealEd",
 						"USDClassesEditor",
 					}
diff --git a/Engine/Plugins/MovieScene/LevelSequenceEditor/Source/LevelSequenceEditor/Public/LevelSequenceEditorSubsystem.h b/Engine/Plugins/MovieScene/LevelSequenceEditor/Source/LevelSequenceEditor/Public/LevelSequenceEditorSubsystem.h
--- a/Engine/Plugins/MovieScene/LevelSequenceEditor/Source/LevelSequenceEditor/Public/LevelSequenceEditorSubsystem.h
+++ b/Engine/Plugins/MovieScene/LevelSequenceEditor/Source/LevelSequenceEditor/Public/LevelSequenceEditorSubsystem.h
@@ -56,7 +56,7 @@ struct FMovieSceneBindingPropertyInfo
 	GENERATED_BODY()
 
 	// Locator for the entry
-	UPROPERTY(EditAnywhere, Category = "Default", meta=(AllowedLocators="Actor", DisplayName="Actor"))
+	UPROPERTY(EditAnywhere, Category = "Default", meta=(AllowedLocators="Actor, UsdPrim", DisplayName="Actor"))
 	FUniversalObjectLocator Locator;
 
 	// Flags for how to resolve the locator
diff --git a/Engine/Plugins/Runtime/USDCore/Source/USDClasses/Private/USDObjectUtils.cpp b/Engine/Plugins/Runtime/USDCore/Source/USDClasses/Private/USDObjectUtils.cpp
--- a/Engine/Plugins/Runtime/USDCore/Source/USDClasses/Private/USDObjectUtils.cpp
+++ b/Engine/Plugins/Runtime/USDCore/Source/USDClasses/Private/USDObjectUtils.cpp
@@ -263,10 +263,6 @@ TSubclassOf<UUsdAssetUserData> UsdUnreal::ObjectUtils::GetAssetUserDataClassForO
 	{
 		return UUsdSparseVolumeTextureAssetUserData::StaticClass();
 	}
-	else if (ObjectClass->IsChildOf(ULevelSequence::StaticClass()))
-	{
-		return UUsdLevelSequenceAssetUserData::StaticClass();
-	}
 	else if (ObjectClass->ImplementsInterface(UInterface_AssetUserData::StaticClass()))
 	{
 		// Only return UUsdAssetUserData in case the object can hold asset user data, otherwise we'd get a warning
diff --git a/Engine/Plugins/Runtime/USDCore/Source/USDClasses/Public/USDAssetUserData.h b/Engine/Plugins/Runtime/USDCore/Source/USDClasses/Public/USDAssetUserData.h
--- a/Engine/Plugins/Runtime/USDCore/Source/USDClasses/Public/USDAssetUserData.h
+++ b/Engine/Plugins/Runtime/USDCore/Source/USDClasses/Public/USDAssetUserData.h
@@ -122,7 +122,8 @@ public:
  * We need this as part of a mechanism to automatically repair those bindings when they break if we close/reload the stage.
  */
 UCLASS(BlueprintType)
-class USDCLASSES_API UUsdLevelSequenceAssetUserData : public UAssetUserData
+class UE_DEPRECATED(5.5, "This class is no longer used, as AUsdStageActors now use Sequencer Universal Object Locators for persistent bindings")
+	USDCLASSES_API UUsdLevelSequenceAssetUserData : public UAssetUserData
 {
 	GENERATED_BODY()
 
@@ -173,4 +174,4 @@ public:
 	 */
 	UPROPERTY(EditAnywhere, BlueprintReadWrite, Category = "USD")
 	TArray<FString> TimeSamplePaths;
-};
\ No newline at end of file
+};
