USD: Add option to repopulate LevelSequence animations;

Generate inherited visibility tracks also when opening the stage;
Refactor handling of visibility inheritance when opening the stage to not need another full subtree traversal;
Remove unnecessary usage of PrimsToAnimate as a "visited set" when deciding whether to Add/Remove prims from the LevelSequenceHelper;
Remove some unnecessary updates when handling inherited visibilty tracks.

diff --git a/Engine/Plugins/Importers/USDImporter/Source/USDSchemas/Public/USDSchemaTranslator.h b/Engine/Plugins/Importers/USDImporter/Source/USDSchemas/Public/USDSchemaTranslator.h
--- a/Engine/Plugins/Importers/USDImporter/Source/USDSchemas/Public/USDSchemaTranslator.h
+++ b/Engine/Plugins/Importers/USDImporter/Source/USDSchemas/Public/USDSchemaTranslator.h
@@ -187,6 +187,9 @@ struct USDSCHEMAS_API FUsdSchemaTranslationContext : public TSharedFromThis<FUsd
 	/** True if we're a context created by the USDStageImporter to fully import to persistent assets and actors */
 	bool bIsImporting = false;
 
+	/** True if we're just re-adding animations onto the LevelSequence, and not creating/updating components */
+	bool bIsJustRepopulatingLevelSequence = false;
+
 	/**
 	 * True if we're building the InfoCache assigned to this context. This usually means we shouldn't query it for information, and should instead
 	 * compute it manually so that it can be cached
diff --git a/Engine/Plugins/Importers/USDImporter/Source/USDStage/Private/USDLevelSequenceHelper.cpp b/Engine/Plugins/Importers/USDImporter/Source/USDStage/Private/USDLevelSequenceHelper.cpp
--- a/Engine/Plugins/Importers/USDImporter/Source/USDStage/Private/USDLevelSequenceHelper.cpp
+++ b/Engine/Plugins/Importers/USDImporter/Source/USDStage/Private/USDLevelSequenceHelper.cpp
@@ -1803,6 +1803,7 @@ void FUsdLevelSequenceHelperImpl::AddCommonTracks(const UUsdPrimTwin& PrimTwin,
 			// Collect all the time samples we'll need to sample our visibility at (USD has inherited visibilities, so every time
 			// a parent has a key, we need to recompute the child visibility at that moment too)
 			TArray<double> TotalVisibilityTimeSamples;
+			VisibilityAttribute.GetTimeSamples(TotalVisibilityTimeSamples);
 
 			// If we're adding a visibility track because a parent has visibility animations, we want to write our baked visibility
 			// tracks on the same layer as the first one of our parents that actually has animated visibility.
@@ -1810,7 +1811,7 @@ void FUsdLevelSequenceHelperImpl::AddCommonTracks(const UUsdPrimTwin& PrimTwin,
 			// many parents and specs on multiple layers, but this is hopefully at least *a* reasonable answer.
 			UE::FUsdAttribute FirstAnimatedVisibilityParentAttr;
 
-			if ((VisibilityAttribute.GetTimeSamples(TotalVisibilityTimeSamples) && TotalVisibilityTimeSamples.Num() > 0) || bForceVisibilityTracks)
+			if (bForceVisibilityTracks)
 			{
 				// TODO: Improve this, as this is extremely inefficient since we'll be parsing this tree for the root down and repeatedly
 				// redoing this one child at a time...
diff --git a/Engine/Plugins/Importers/USDImporter/Source/USDStage/Private/USDStageActor.cpp b/Engine/Plugins/Importers/USDImporter/Source/USDStage/Private/USDStageActor.cpp
--- a/Engine/Plugins/Importers/USDImporter/Source/USDStage/Private/USDStageActor.cpp
+++ b/Engine/Plugins/Importers/USDImporter/Source/USDStage/Private/USDStageActor.cpp
@@ -1063,13 +1063,8 @@ AUsdStageActor::AUsdStageActor()
 								TEXT("Reloading animations because layer '%s' was reloaded"),
 								*LayerToChangeList.Key.GetIdentifier()
 							);
-							ReloadAnimations();
-
-							// Make sure our PrimsToAnimate and the LevelSequenceHelper are kept in sync, because we'll use PrimsToAnimate to
-							// check whether we need to call LevelSequenceHelper::AddPrim within AUsdStageActor::ExpandPrim. Without this reset
-							// our prims would already be in here by the time we're checking if we need to add tracks or not, and we wouldn't re-add
-							// the tracks
-							PrimsToAnimate.Reset();
+							RegenerateLevelSequence();
+							RepopulateLevelSequence();
 							return;
 						}
 
@@ -1084,8 +1079,8 @@ AUsdStageActor::AUsdStageActor()
 									TEXT("Reloading animations because layer '%s' was added or removed"),
 									*LayerToChangeList.Key.GetIdentifier()
 								);
-								ReloadAnimations();
-								PrimsToAnimate.Reset();
+								RegenerateLevelSequence();
+								RepopulateLevelSequence();
 								return;
 							}
 						}
@@ -1798,6 +1793,13 @@ void AUsdStageActor::OnUsdObjectsChanged(const UsdUtils::FObjectChangesByPath& I
 		}
 	}
 
+	// Recreate our LevelSequences before we regenerate components and want to add bindings/tracks
+	// back onto it
+	if (bNeedsAnimationReload)
+	{
+		RegenerateLevelSequence();
+	}
+
 	RefreshStageTask.EnterProgressFrame();
 	FScopedSlowTask RegenerateComponentsTask(PrimsToUpdate.Num(), LOCTEXT("RegeneratingComponents", "Regenerating components"));
 	for (const TPair<UE::FSdfPath, bool>& PrimChangedInfo : PrimsToUpdate)
@@ -1830,12 +1832,6 @@ void AUsdStageActor::OnUsdObjectsChanged(const UsdUtils::FObjectChangesByPath& I
 		OnPrimChanged.Broadcast(PrimChangedInfo.Key.GetString(), PrimChangedInfo.Value);
 	}
 
-	if (bNeedsAnimationReload)
-	{
-		ReloadAnimations();
-		PrimsToAnimate.Reset();
-	}
-
 	if (bHasResync)
 	{
 		FUsdStageActorImpl::RepairExternalSequencerBindings();
@@ -1930,7 +1926,12 @@ UUsdPrimTwin* AUsdStageActor::GetOrCreatePrimTwin(const UE::FSdfPath& UsdPrimPat
 	return UsdPrimTwin;
 }
 
-UUsdPrimTwin* AUsdStageActor::ExpandPrim(const UE::FUsdPrim& Prim, bool bResync, FUsdSchemaTranslationContext& TranslationContext)
+UUsdPrimTwin* AUsdStageActor::ExpandPrim(
+	const UE::FUsdPrim& Prim,
+	bool bResync,
+	FUsdSchemaTranslationContext& TranslationContext,
+	TOptional<bool> bParentHasAnimatedVisibility
+)
 {
 	UUsdPrimTwin* UsdPrimTwin = nullptr;
 #if USE_USD_SDK
@@ -1952,45 +1953,82 @@ UUsdPrimTwin* AUsdStageActor::ExpandPrim(const UE::FUsdPrim& Prim, bool bResync,
 
 	bool bExpandChildren = true;
 
-	IUsdSchemasModule& UsdSchemasModule = FModuleManager::Get().LoadModuleChecked<IUsdSchemasModule>(TEXT("USDSchemas"));
-
-	if (TSharedPtr<FUsdSchemaTranslator> SchemaTranslator = UsdSchemasModule.GetTranslatorRegistry()
-																.CreateTranslatorForSchema(TranslationContext.AsShared(), UE::FUsdTyped(Prim)))
+	if (!TranslationContext.bIsJustRepopulatingLevelSequence)
 	{
-		if (bResync && !UsdPrimTwin->SceneComponent.IsValid())
-		{
-			UsdPrimTwin->SceneComponent = SchemaTranslator->CreateComponents();
-		}
-		else
+		IUsdSchemasModule& UsdSchemasModule = FModuleManager::Get().LoadModuleChecked<IUsdSchemasModule>(TEXT("USDSchemas"));
+		if (TSharedPtr<FUsdSchemaTranslator> SchemaTranslator = UsdSchemasModule.GetTranslatorRegistry()
+																	.CreateTranslatorForSchema(TranslationContext.AsShared(), UE::FUsdTyped(Prim)))
 		{
-			USceneComponent* TwinSceneComponent = UsdPrimTwin->SceneComponent.Get();
-
-			ObjectsToWatch.Remove(TwinSceneComponent);
-			if (Prim.IsA(TEXT("Camera")))
+			if (bResync && !UsdPrimTwin->SceneComponent.IsValid())
 			{
-				if (ACineCameraActor* CameraActor = Cast<ACineCameraActor>(SceneComponent->GetOwner()))
-				{
-					ObjectsToWatch.Remove(CameraActor->GetCineCameraComponent());
-				}
+				UsdPrimTwin->SceneComponent = SchemaTranslator->CreateComponents();
 			}
-			else if (Prim.IsA(TEXT("PointInstancer")))
+			else
 			{
-				TSet<FString> PrototypePaths = FUsdStageActorImpl::GetPointInstancerPrototypes(Prim);
+				USceneComponent* TwinSceneComponent = UsdPrimTwin->SceneComponent.Get();
 
-				for (const TObjectPtr<USceneComponent>& Child : TwinSceneComponent->GetAttachChildren())
+				ObjectsToWatch.Remove(TwinSceneComponent);
+				if (Prim.IsA(TEXT("Camera")))
+				{
+					if (ACineCameraActor* CameraActor = Cast<ACineCameraActor>(SceneComponent->GetOwner()))
+					{
+						ObjectsToWatch.Remove(CameraActor->GetCineCameraComponent());
+					}
+				}
+				else if (Prim.IsA(TEXT("PointInstancer")))
 				{
-					UHierarchicalInstancedStaticMeshComponent* HISMComponent = Cast<UHierarchicalInstancedStaticMeshComponent>(Child.Get());
-					ObjectsToWatch.Remove(HISMComponent);
+					TSet<FString> PrototypePaths = FUsdStageActorImpl::GetPointInstancerPrototypes(Prim);
+
+					for (const TObjectPtr<USceneComponent>& Child : TwinSceneComponent->GetAttachChildren())
+					{
+						UHierarchicalInstancedStaticMeshComponent* HISMComponent = Cast<UHierarchicalInstancedStaticMeshComponent>(Child.Get());
+						ObjectsToWatch.Remove(HISMComponent);
+					}
 				}
+				SchemaTranslator->UpdateComponents(TwinSceneComponent);
+			}
+
+			bExpandChildren = !SchemaTranslator->CollapsesChildren(ECollapsingType::Components);
+		}
+	}
+
+	// Check for parents with animated visibility.
+	//
+	// When opening the stage we'll propagate this down already, but we may be just updating a random prim in the middle
+	// of the hierarchy from an update notice, so we may need to check our parents right here.
+	// After we do this here we can propagate this value to our children though
+	if (!bParentHasAnimatedVisibility.IsSet())
+	{
+		bool bHasVisibilityAnimationParent = false;
+		UE::FUsdPrim ParentPrim = Prim.GetParent();
+		while (ParentPrim && !ParentPrim.IsPseudoRoot())
+		{
+			if (UsdUtils::HasAnimatedVisibility(ParentPrim))
+			{
+				bHasVisibilityAnimationParent = true;
+				break;
 			}
-			SchemaTranslator->UpdateComponents(TwinSceneComponent);
+
+			ParentPrim = ParentPrim.GetParent();
 		}
 
-		bExpandChildren = !SchemaTranslator->CollapsesChildren(ECollapsingType::Components);
+		bParentHasAnimatedVisibility = bHasVisibilityAnimationParent;
 	}
 
 	if (bExpandChildren)
 	{
+		// Unfortunately if we have animated visibility we need to be ready to update the visibility
+		// of all components that we spawned for child prims whenever this prim's visibility updates.
+		// We can't just have this prim's FUsdGeomXformableTranslator::UpdateComponents ->
+		// -> UsdToUnreal::ConvertXformable call use SetHiddenInGame recursively, because we may have
+		// child prims that are themselves also invisible, and so their own subtrees should be invisible
+		// even if this prim goes visible. Also keep in mind that technically we'll always update each
+		// prim in the order that they are within PrimsToAnimate, but that order is not strictly enforced
+		// to be e.g. a breadth first traversal on the prim tree or anything like this, so these updates
+		// need to be order-independent, which means we really should add the entire subtree to the list
+		// and have UpdateComponents called on all components.
+		bParentHasAnimatedVisibility = bParentHasAnimatedVisibility.GetValue() || UsdUtils::HasAnimatedVisibility(Prim);
+
 		USceneComponent* ContextParentComponent = TranslationContext.ParentComponent;
 
 		if (UsdPrimTwin->SceneComponent.IsValid())
@@ -2005,12 +2043,12 @@ UUsdPrimTwin* AUsdStageActor::ExpandPrim(const UE::FUsdPrim& Prim, bool bResync,
 
 		for (const UE::FUsdPrim& ChildPrim : PrimChildren)
 		{
-			ExpandPrim(ChildPrim, bResync, TranslationContext);
+			ExpandPrim(ChildPrim, bResync, TranslationContext, bParentHasAnimatedVisibility);
 		}
 	}
 
 	USceneComponent* TwinSceneComponent = UsdPrimTwin->SceneComponent.Get();
-	if (TwinSceneComponent)
+	if (TwinSceneComponent && !TranslationContext.bIsJustRepopulatingLevelSequence)
 	{
 #if WITH_EDITOR
 		TwinSceneComponent->PostEditChange();
@@ -2069,73 +2107,39 @@ UUsdPrimTwin* AUsdStageActor::ExpandPrim(const UE::FUsdPrim& Prim, bool bResync,
 	}
 
 	// Check if the prim should have Sequencer tracks or not
-	bool bIsAnimated = false;
-	if (USkeletalMeshComponent* SkeletalMeshComponent = Cast<USkeletalMeshComponent>(TwinSceneComponent))
-	{
-		if (SkeletalMeshComponent->AnimationData.AnimToPlay)
-		{
-			// We know we're animated if we have an animation of course
-			bIsAnimated = true;
-		}
-	}
-	if (!bIsAnimated)
-	{
-		// Always consider SpatialAudio prims as animated so that we can create LevelSequence tracks for the audio itself.
-		// We exclusively handle the audio stuff via the Sequencer and LevelSequence tracks because there's no way to play audio via Time animation,
-		// and the audio component is not meant to be a fully featured audio player with start/end play times and animated volume controls. In
-		// other words, if we placed our SoundWave asset on the component, the audio component would instantly play it when going into PIE, which is
-		// not what we want. The audio component and actor are only really used for their transforms on the level whenever we'ret trying to play
-		// spatial audio
-		bIsAnimated = UsdUtils::IsAnimated(Prim) || Prim.IsA(TEXT("SpatialAudio"));
-	}
+	bool bIsAnimated = bParentHasAnimatedVisibility.GetValue();
 
-	// Create Sequencer tracks for the prim
-	bool bHasAnimatedBounds = false;
-	if (bIsAnimated)
+	// We know we're animated if we have skeletal animation of course
+	if (!bIsAnimated)
 	{
-		if (!PrimsToAnimate.Contains(UsdPrimTwin->PrimPath))
-		{
-			// Unfortunately if we have animated visibility we need to be ready to update the visibility
-			// of all components that we spawned for child prims whenever this prim's visibility updates.
-			// We can't just have this prim's FUsdGeomXformableTranslator::UpdateComponents ->
-			// -> UsdToUnreal::ConvertXformable call use SetHiddenInGame recursively, because we may have
-			// child prims that are themselves also invisible, and so their own subtrees should be invisible
-			// even if this prim goes visible. Also keep in mind that technically we'll always update each
-			// prim in the order that they are within PrimsToAnimate, but that order is not strictly enforced
-			// to be e.g. a breadth first traversal on the prim tree or anything like this, so these updates
-			// need to be order-independent, which means we really should add the entire subtree to the list
-			// and have UpdateComponents called on all components.
-			if (UsdUtils::HasAnimatedVisibility(Prim))
+		if (USkeletalMeshComponent* SkeletalMeshComponent = Cast<USkeletalMeshComponent>(TwinSceneComponent))
+		{
+			if (SkeletalMeshComponent->AnimationData.AnimToPlay)
 			{
-				const bool bRecursive = true;
-				UsdPrimTwin->Iterate(
-					[this](UUsdPrimTwin& Twin)
-					{
-						if (!PrimsToAnimate.Contains(Twin.PrimPath))
-						{
-							PrimsToAnimate.Add(Twin.PrimPath);
-							LevelSequenceHelper.AddPrim(Twin);
-						}
-					},
-					bRecursive
-				);
+				bIsAnimated = true;
 			}
-
-			PrimsToAnimate.Add(UsdPrimTwin->PrimPath);
-			LevelSequenceHelper.AddPrim(*UsdPrimTwin);
 		}
 	}
-	else if (EUsdDrawMode DrawMode = UsdUtils::GetAppliedDrawMode(Prim); DrawMode != EUsdDrawMode::Default)
+
+	// Always consider SpatialAudio prims as animated so that we can create LevelSequence tracks for the audio itself.
+	// We exclusively handle the audio stuff via the Sequencer and LevelSequence tracks because there's no way to play audio via Time animation,
+	// and the audio component is not meant to be a fully featured audio player with start/end play times and animated volume controls. In
+	// other words, if we placed our SoundWave asset on the component, the audio component would instantly play it when going into PIE, which is
+	// not what we want. The audio component and actor are only really used for their transforms on the level whenever we'ret trying to play
+	// spatial audio
+	bIsAnimated = bIsAnimated || Prim.IsA(TEXT("SpatialAudio")) || UsdUtils::IsAnimated(Prim);
+
+	TOptional<bool> bHasAnimatedBounds;
+	if (EUsdDrawMode DrawMode = UsdUtils::GetAppliedDrawMode(Prim); DrawMode != EUsdDrawMode::Default)
 	{
 		const bool bUseExtentsHint = true;
 		const bool bIgnoreVisibility = false;
 		bHasAnimatedBounds = UsdUtils::HasAnimatedBounds(Prim, static_cast<EUsdPurpose>(PurposesToLoad), bUseExtentsHint, bIgnoreVisibility);
 
-		if (bHasAnimatedBounds)
+		const bool bDefault = false;
+		if (bHasAnimatedBounds.Get(bDefault))
 		{
-			const bool bForceVisibilityTrack = false;
-			LevelSequenceHelper.AddPrim(*UsdPrimTwin, bForceVisibilityTrack, bHasAnimatedBounds);
-			PrimsToAnimate.Add(UsdPrimTwin->PrimPath);
+			bIsAnimated = true;
 
 			// Mark the component as animated right away because HasAnimatedBounds is expensive to call and
 			// we don't want to have to re-do it when creating the component
@@ -2146,7 +2150,13 @@ UUsdPrimTwin* AUsdStageActor::ExpandPrim(const UE::FUsdPrim& Prim, bool bResync,
 		}
 	}
 
-	if (!bIsAnimated && !bHasAnimatedBounds && PrimsToAnimate.Contains(UsdPrimTwin->PrimPath))
+	if (bIsAnimated)
+	{
+		const bool bForceVisibilityTracks = bParentHasAnimatedVisibility.GetValue();
+		LevelSequenceHelper.AddPrim(*UsdPrimTwin, bForceVisibilityTracks, bHasAnimatedBounds);
+		PrimsToAnimate.Add(UsdPrimTwin->PrimPath);
+	}
+	else
 	{
 		PrimsToAnimate.Remove(UsdPrimTwin->PrimPath);
 		LevelSequenceHelper.RemovePrim(*UsdPrimTwin);
@@ -2187,15 +2197,12 @@ UUsdPrimTwin* AUsdStageActor::ExpandPrim(const UE::FUsdPrim& Prim, bool bResync,
 			// considered animated too, so lets add it to the proper locations. This will also ensure that
 			// we can close the sequencer after creating a new animation in this way and see it animate on
 			// the level
-			if (!bIsAnimated)
-			{
-				PrimsToAnimate.Add(UsdPrimTwin->PrimPath);
-				LevelSequenceHelper.AddPrim(*UsdPrimTwin);
+			PrimsToAnimate.Add(UsdPrimTwin->PrimPath);
+			LevelSequenceHelper.AddPrim(*UsdPrimTwin);
 
-				// Prevent register/unregister spam when calling FUsdGeomXformableTranslator::UpdateComponents later
-				// during sequencer animation (which can cause the Sequencer UI to glitch out a bit)
-				UsdPrimTwin->SceneComponent->SetMobility(EComponentMobility::Movable);
-			}
+			// Prevent register/unregister spam when calling FUsdGeomXformableTranslator::UpdateComponents later
+			// during sequencer animation (which can cause the Sequencer UI to glitch out a bit)
+			UsdPrimTwin->SceneComponent->SetMobility(EComponentMobility::Movable);
 		}
 	}
 
@@ -3154,13 +3161,7 @@ void AUsdStageActor::LoadUsdStage(bool bOpenIfNeeded)
 		InfoCache = MakeShared<FUsdInfoCache>();
 	}
 
-	ReloadAnimations();
-
-	// Make sure our PrimsToAnimate and the LevelSequenceHelper are kept in sync, because we'll use PrimsToAnimate to
-	// check whether we need to call LevelSequenceHelper::AddPrim within AUsdStageActor::ExpandPrim. Without this reset
-	// our prims would already be in here by the time we're checking if we need to add tracks or not, and we wouldn't re-add
-	// the tracks
-	PrimsToAnimate.Reset();
+	RegenerateLevelSequence();
 
 	TSharedRef<FUsdSchemaTranslationContext> TranslationContext = FUsdStageActorImpl::CreateUsdSchemaTranslationContext(this, RootTwin->PrimPath);
 
@@ -3378,12 +3379,17 @@ void AUsdStageActor::Refresh() const
 
 void AUsdStageActor::ReloadAnimations()
 {
-	TRACE_CPUPROFILER_EVENT_SCOPE(AUsdStageActor::ReloadAnimations);
+	RegenerateLevelSequence();
+}
+
+void AUsdStageActor::RegenerateLevelSequence()
+{
+	TRACE_CPUPROFILER_EVENT_SCOPE(AUsdStageActor::RegenerateLevelSequence);
 
 	// If we're using some property editor that can trigger a stage reload (like the Nanite threshold spinbox),
-	// applying a value may trigger ReloadAnimations -> Can trigger asset editors to open/close/change focus ->
+	// applying a value may trigger RegenerateLevelSequence -> Can trigger asset editors to open/close/change focus ->
 	// -> Can trigger focus to drop from the property editors -> Can cause the values to be applied from the
-	// property editors when releasing focus -> Can trigger another call to ReloadAnimations.
+	// property editors when releasing focus -> Can trigger another call to RegenerateLevelSequence.
 	// CloseAllEditorsForAsset in particular is problematic for this because it will destroy the asset editor
 	// (which is TSharedFromThis) and the reentrant call will try use AsShared() internally and assert, as it
 	// hasn't finished being destroyed.
@@ -3428,6 +3434,18 @@ void AUsdStageActor::ReloadAnimations()
 	}
 }
 
+void AUsdStageActor::RepopulateLevelSequence()
+{
+	TSharedRef<FUsdSchemaTranslationContext> TranslationContext = FUsdStageActorImpl::CreateUsdSchemaTranslationContext(
+		this,
+		UE::FSdfPath::AbsoluteRootPath().GetString()
+	);
+	TranslationContext->bIsJustRepopulatingLevelSequence = true;
+
+	const bool bIsResync = false;
+	UpdatePrim(UE::FSdfPath::AbsoluteRootPath(), bIsResync, *TranslationContext);
+}
+
 TSharedPtr<FUsdInfoCache> AUsdStageActor::GetInfoCache()
 {
 	return InfoCache;
@@ -3532,7 +3550,8 @@ void AUsdStageActor::PostTransacted(const FTransactionObjectEvent& TransactionEv
 			const bool bUnloadIfNeeded = false;
 			CloseUsdStage(bUnloadIfNeeded);
 			OpenUsdStage();
-			ReloadAnimations();
+			RegenerateLevelSequence();
+			RepopulateLevelSequence();
 		}
 		else if (ChangedProperties.Contains(GET_MEMBER_NAME_CHECKED(AUsdStageActor, StageState)))
 		{
diff --git a/Engine/Plugins/Importers/USDImporter/Source/USDStage/Public/USDStageActor.h b/Engine/Plugins/Importers/USDImporter/Source/USDStage/Public/USDStageActor.h
--- a/Engine/Plugins/Importers/USDImporter/Source/USDStage/Public/USDStageActor.h
+++ b/Engine/Plugins/Importers/USDImporter/Source/USDStage/Public/USDStageActor.h
@@ -303,9 +303,20 @@ public:
 	// Enters isolated mode by creating a new USD Stage using the provided layer as its root
 	USDSTAGE_API void IsolateLayer(const UE::FSdfLayer& Layer);
 
-	// Regenerates our LevelSequence from the USD Stage
+	UE_DEPRECATED(5.5, "This function has been renamed into 'RegenerateLevelSequence', which better describes what it does")
 	USDSTAGE_API void ReloadAnimations();
 
+	// Creates new ULevelSequence assets for the root layer and sublayers
+	USDSTAGE_API void RegenerateLevelSequence();
+
+	/**
+	 * Fills in the current LevelSequences with animations for the prims of the opened stage.
+	 *
+	 * Can be called even if our current LevelSequence has already been populated: It will simply refresh the
+	 * current LevelSequence, adding/removing tracks and bindings as needed.
+	 */
+	USDSTAGE_API void RepopulateLevelSequence();
+
 	USDSTAGE_API TSharedPtr<FUsdInfoCache> GetInfoCache();
 	USDSTAGE_API TSharedPtr<UE::FUsdGeomBBoxCache> GetBBoxCache();
 	USDSTAGE_API TMap<FString, TMap<FString, int32>> GetMaterialToPrimvarToUVIndex();
@@ -363,7 +374,12 @@ protected:
 	UUsdPrimTwin* GetRootPrimTwin() const;
 	UUsdPrimTwin* GetOrCreatePrimTwin(const UE::FSdfPath& UsdPrimPath);
 
-	UUsdPrimTwin* ExpandPrim(const UE::FUsdPrim& Prim, bool bResync, FUsdSchemaTranslationContext& TranslationContext);
+	UUsdPrimTwin* ExpandPrim(
+		const UE::FUsdPrim& Prim,
+		bool bResync,
+		FUsdSchemaTranslationContext& TranslationContext,
+		TOptional<bool> bParentHasAnimatedVisibility = {}
+	);
 	void UpdatePrim(const UE::FSdfPath& UsdPrimPath, bool bResync, FUsdSchemaTranslationContext& TranslationContext);
 
 	void OpenUsdStage();
diff --git a/Engine/Plugins/Importers/USDImporter/Source/USDStageEditor/Private/SUSDStage.cpp b/Engine/Plugins/Importers/USDImporter/Source/USDStageEditor/Private/SUSDStage.cpp
--- a/Engine/Plugins/Importers/USDImporter/Source/USDStageEditor/Private/SUSDStage.cpp
+++ b/Engine/Plugins/Importers/USDImporter/Source/USDStageEditor/Private/SUSDStage.cpp
@@ -899,8 +899,8 @@ void SUsdStage::FillFileMenu(FMenuBuilder& MenuBuilder)
 	MenuBuilder.BeginSection("Reload", LOCTEXT("Reload", "Reload"));
 	{
 		MenuBuilder.AddMenuEntry(
-			LOCTEXT("Reload", "Reload"),
-			LOCTEXT("Reload_ToolTip", "Reloads the stage from disk, keeping aspects of the session intact"),
+			LOCTEXT("Reload", "Reload stage"),
+			LOCTEXT("Reload_ToolTip", "Reloads the stage and animations from disk, keeping aspects of the session intact"),
 			FSlateIcon(),
 			FUIAction(
 				FExecuteAction::CreateSP(this, &SUsdStage::FileReload),
@@ -1011,6 +1011,31 @@ void SUsdStage::FillActionsMenu(FMenuBuilder& MenuBuilder)
 			NAME_None,
 			EUserInterfaceActionType::Button
 		);
+
+		MenuBuilder.AddMenuEntry(
+			LOCTEXT("Regenerate", "Regenerate USD level sequence"),
+			LOCTEXT("Regenerate_ToolTip", "Regenerates the transient LevelSequences based on the opened stage"),
+			FSlateIcon(),
+			FUIAction(
+				FExecuteAction::CreateSP(this, &SUsdStage::ActionsRegenerate),
+				FCanExecuteAction::CreateLambda(
+					[this]()
+					{
+						if (const AUsdStageActor* StageActor = ViewModel.UsdStageActor.Get())
+						{
+							if (UE::FUsdStage Stage = StageActor->GetUsdStage())
+							{
+								return true;
+							}
+						}
+
+						return false;
+					}
+				)
+			),
+			NAME_None,
+			EUserInterfaceActionType::Button
+		);
 	}
 	MenuBuilder.EndSection();
 }
@@ -2471,6 +2496,16 @@ void SUsdStage::ActionsImport(const FString& OutputContentFolder, UUsdStageImpor
 	ViewModel.ImportStage(*OutputContentFolder, Options);
 }
 
+void SUsdStage::ActionsRegenerate()
+{
+	FScopedTransaction Transaction(LOCTEXT("RegenerateTransaction", "Regenerate USD level sequence"));
+
+	if (AUsdStageActor* StageActor = ViewModel.UsdStageActor.Get())
+	{
+		StageActor->RepopulateLevelSequence();
+	}
+}
+
 void SUsdStage::ExportSelectedLayers(const FString& OutputLayerOrDirectory)
 {
 	if (UsdLayersTreeView)
diff --git a/Engine/Plugins/Importers/USDImporter/Source/USDStageEditor/Private/SUSDStage.h b/Engine/Plugins/Importers/USDImporter/Source/USDStageEditor/Private/SUSDStage.h
--- a/Engine/Plugins/Importers/USDImporter/Source/USDStageEditor/Private/SUSDStage.h
+++ b/Engine/Plugins/Importers/USDImporter/Source/USDStageEditor/Private/SUSDStage.h
@@ -71,6 +71,7 @@ public:
 	void FileClose();
 	void ActionsImportWithDialog();
 	void ActionsImport(const FString& OutputContentFolder, UUsdStageImportOptions* Options);
+	void ActionsRegenerate();
 	void ExportSelectedLayers(const FString& OutputLayerOrDirectory = {});
 
 protected:
diff --git a/Engine/Plugins/Runtime/USDCore/Source/USDUtilities/Private/USDConversionUtils.cpp b/Engine/Plugins/Runtime/USDCore/Source/USDUtilities/Private/USDConversionUtils.cpp
--- a/Engine/Plugins/Runtime/USDCore/Source/USDUtilities/Private/USDConversionUtils.cpp
+++ b/Engine/Plugins/Runtime/USDCore/Source/USDUtilities/Private/USDConversionUtils.cpp
@@ -1060,8 +1060,7 @@ bool UsdUtils::HasAnimatedVisibility(const pxr::UsdPrim& Prim)
 	{
 		if (pxr::UsdAttribute Attr = Imageable.GetVisibilityAttr())
 		{
-			std::vector<double> TimeSamples;
-			if (Attr.GetTimeSamples(&TimeSamples) && TimeSamples.size() > 0)
+			if (Attr.GetNumTimeSamples() > 0)
 			{
 				return true;
 			}
