USD: Fix actor visibility tracks being removed after LevelSequence updates;

diff --git a/Engine/Plugins/Importers/USDImporter/Source/USDStage/Private/USDLevelSequenceHelper.cpp b/Engine/Plugins/Importers/USDImporter/Source/USDStage/Private/USDLevelSequenceHelper.cpp
--- a/Engine/Plugins/Importers/USDImporter/Source/USDStage/Private/USDLevelSequenceHelper.cpp
+++ b/Engine/Plugins/Importers/USDImporter/Source/USDStage/Private/USDLevelSequenceHelper.cpp
@@ -274,7 +274,11 @@ namespace UsdLevelSequenceHelperImpl
 	TSharedPtr<ISequencer> GetOpenedSequencerForLevelSequence(ULevelSequence* LevelSequence)
 	{
 		const bool bFocusIfOpen = false;
-		IAssetEditorInstance* AssetEditor = GEditor ? GEditor->GetEditorSubsystem<UAssetEditorSubsystem>()->FindEditorForAsset(LevelSequence, bFocusIfOpen) : nullptr;
+		IAssetEditorInstance* AssetEditor = GEditor ? GEditor->GetEditorSubsystem<UAssetEditorSubsystem>()->FindEditorForAsset(
+												LevelSequence,
+												bFocusIfOpen
+											)
+													: nullptr;
 		ILevelSequenceEditorToolkit* LevelSequenceEditor = static_cast<ILevelSequenceEditorToolkit*>(AssetEditor);
 		return LevelSequenceEditor ? LevelSequenceEditor->GetSequencer() : nullptr;
 	}
@@ -3910,6 +3914,8 @@ FGuid FUsdLevelSequenceHelperImpl::GetOrCreateComponentBinding(
 
 void FUsdLevelSequenceHelperImpl::HandleMovieSceneChange(UMovieScene& MovieScene)
 {
+	using namespace UsdLevelSequenceHelperImpl;
+
 	// It's possible to get this called when the actor and it's level sequences are being all destroyed in one go.
 	// We need the FScopedBlockNotices in this function, but if our StageActor is already being destroyed, we can't reliably
 	// use its listener, and so then we can't do anything. We likely don't want to write back to the stage at this point anyway.
@@ -4058,7 +4064,16 @@ void FUsdLevelSequenceHelperImpl::HandleMovieSceneChange(UMovieScene& MovieScene
 				if (UE::FUsdPrim UsdPrim = UsdStage.GetPrimAtPath(UE::FSdfPath(*UsdPrimTwin->PrimPath)))
 				{
 					RemoveTimeSamplesForPropertyIfNeeded(UsdPrim, Guid, UnrealIdentifiers::TransformPropertyName);
-					RemoveTimeSamplesForPropertyIfNeeded(UsdPrim, Guid, UnrealIdentifiers::HiddenInGamePropertyName);
+
+					// Handle visibility explicitly here because we may have a visibility track on the actor or on the component
+					if (!FindTrackTypeOrDerived<UMovieScenePropertyTrack>(&MovieScene, Guid, UnrealIdentifiers::HiddenInGamePropertyName)
+						&& !FindTrackTypeOrDerived<UMovieScenePropertyTrack>(&MovieScene, Guid, UnrealIdentifiers::HiddenPropertyName))
+					{
+						for (UE::FUsdAttribute& Attr : UsdUtils::GetAttributesForProperty(UsdPrim, UnrealIdentifiers::HiddenInGamePropertyName))
+						{
+							RemoveTimeSamplesForAttr(Attr);
+						}
+					}
 
 					if (bIsCamera)
 					{
