USD: Prevent Time-based animation from affecting visibility when the Sequencer is opened.

diff --git a/Engine/Plugins/Importers/USDImporter/Source/USDSchemas/Private/USDGeomXformableTranslator.cpp b/Engine/Plugins/Importers/USDImporter/Source/USDSchemas/Private/USDGeomXformableTranslator.cpp
--- a/Engine/Plugins/Importers/USDImporter/Source/USDSchemas/Private/USDGeomXformableTranslator.cpp
+++ b/Engine/Plugins/Importers/USDImporter/Source/USDSchemas/Private/USDGeomXformableTranslator.cpp
@@ -732,7 +732,7 @@ void FUsdGeomXformableTranslator::UpdateComponents(USceneComponent* SceneCompone
 
 		// Only put the transform into the component if we haven't parsed LODs for our static mesh: The Mesh transforms will already be baked
 		// into the mesh at that case, as each LOD could technically have a separate transform
-		if (!Context->bAllowInterpretingLODs || !bHasMultipleLODs)
+		if (!Context->bSequencerIsAnimating && (!Context->bAllowInterpretingLODs || !bHasMultipleLODs))
 		{
 			// Don't update the component's transform if this is already factored in as root motion within the AnimSequence
 			bool bConvertTransform = true;
