USD: Fix direct usages of the "bHidden" and "bHiddenInGame" strings;

diff --git a/Engine/Plugins/Importers/USDImporter/Source/USDExporter/Private/LevelSequenceExporterUSD.cpp b/Engine/Plugins/Importers/USDImporter/Source/USDExporter/Private/LevelSequenceExporterUSD.cpp
--- a/Engine/Plugins/Importers/USDImporter/Source/USDExporter/Private/LevelSequenceExporterUSD.cpp
+++ b/Engine/Plugins/Importers/USDImporter/Source/USDExporter/Private/LevelSequenceExporterUSD.cpp
@@ -1187,7 +1187,7 @@ namespace UE::LevelSequenceExporterUSD::Private
 						// Just handle spawnable tracks as if they're visibility tracks, and hide the prim when not "spawned"
 						// Remember that our spawn register just hides the spawnables when they're not spawned anyway, so this
 						// is essentially the same
-						const FString PropertyPath = TEXT("bHidden");
+						const static FString PropertyPath = UnrealIdentifiers::HiddenPropertyName.ToString();
 						UnrealToUsd::CreateComponentPropertyBaker(Prim, *BoundComponent, PropertyPath, Baker);
 					}
 					// Check for the control rig tracks too, because if the user did "Bake to Control Rig" the controlrig code will silently
@@ -1200,7 +1200,8 @@ namespace UE::LevelSequenceExporterUSD::Private
 					// change, as we can't export "hierarchy changes" otherwise
 					else if (Track->IsA<UMovieScene3DAttachTrack>())
 					{
-						UnrealToUsd::CreateComponentPropertyBaker(Prim, *BoundComponent, TEXT("Transform"), Baker);
+						const static FString PropertyPath = UnrealIdentifiers::TransformPropertyName.ToString();
+						UnrealToUsd::CreateComponentPropertyBaker(Prim, *BoundComponent, PropertyPath, Baker);
 					}
 					else if (const UMovieSceneAudioTrack* AudioTrack = Cast<UMovieSceneAudioTrack>(Track))
 					{
@@ -1273,7 +1274,7 @@ namespace UE::LevelSequenceExporterUSD::Private
 				{
 					UnrealToUsd::FComponentBaker Baker;
 
-					const FString PropertyPath = TEXT("Transform");
+					const static FString PropertyPath = UnrealIdentifiers::TransformPropertyName.ToString();
 					UnrealToUsd::CreateComponentPropertyBaker(Prim, *BoundComponent, PropertyPath, Baker);
 
 					AddBaker(Baker);
diff --git a/Engine/Plugins/Importers/USDImporter/Source/USDStage/Private/USDStageActor.cpp b/Engine/Plugins/Importers/USDImporter/Source/USDStage/Private/USDStageActor.cpp
--- a/Engine/Plugins/Importers/USDImporter/Source/USDStage/Private/USDStageActor.cpp
+++ b/Engine/Plugins/Importers/USDImporter/Source/USDStage/Private/USDStageActor.cpp
@@ -3377,10 +3377,12 @@ void AUsdStageActor::Refresh() const
 	OnTimeChanged.Broadcast();
 }
 
+PRAGMA_DISABLE_DEPRECATION_WARNINGS
 void AUsdStageActor::ReloadAnimations()
 {
 	RegenerateLevelSequence();
 }
+PRAGMA_ENABLE_DEPRECATION_WARNINGS
 
 void AUsdStageActor::RegenerateLevelSequence()
 {
@@ -4261,7 +4263,7 @@ void AUsdStageActor::OnObjectPropertyChanged(UObject* ObjectBeingModified, FProp
 		// To accomplish that while blocking notices we must always propagate component visibility changes manually.
 		// This part is effectively the same as calling pxr::UsdGeomImageable::MakeVisible/Invisible.
 		// TODO: Allow writing out visibility without needing a prim twin
-		if (UsdPrimTwin && PropertyChangedEvent.GetPropertyName() == TEXT("bHiddenInGame"))
+		if (UsdPrimTwin && PropertyChangedEvent.GetPropertyName() == UnrealIdentifiers::HiddenInGamePropertyName)
 		{
 			PrimSceneComponent->Modify();
 
diff --git a/Engine/Plugins/Runtime/USDCore/Source/USDUtilities/Private/USDPrimConversion.cpp b/Engine/Plugins/Runtime/USDCore/Source/USDUtilities/Private/USDPrimConversion.cpp
--- a/Engine/Plugins/Runtime/USDCore/Source/USDUtilities/Private/USDPrimConversion.cpp
+++ b/Engine/Plugins/Runtime/USDCore/Source/USDUtilities/Private/USDPrimConversion.cpp
@@ -3979,7 +3979,11 @@ bool UnrealToUsd::CreateComponentPropertyBaker(
 
 	// SceneComponent
 	{
-		if (PropertyPath == TEXT("Transform"))
+		const static FString TransformPropertyString = UnrealIdentifiers::TransformPropertyName.ToString();
+		const static FString HiddenPropertyString = UnrealIdentifiers::HiddenPropertyName.ToString();
+		const static FString HiddenInGamePropertyString = UnrealIdentifiers::HiddenInGamePropertyName.ToString();
+
+		if (PropertyPath == TransformPropertyString)
 		{
 			pxr::UsdGeomXformable Xformable(UsdPrim);
 			if (!Xformable)
@@ -4077,7 +4081,7 @@ bool UnrealToUsd::CreateComponentPropertyBaker(
 		// A component is only visible when it's not hidden and its actor is not hidden
 		// A bHidden is just handled like a bHiddenInGame for the actor's root component
 		// Whenever we handle a bHiddenInGame, we always combine it with the actor's bHidden
-		else if (PropertyPath == TEXT("bHidden") || PropertyPath == TEXT("bHiddenInGame"))
+		else if (PropertyPath == HiddenPropertyString || PropertyPath == HiddenInGamePropertyString)
 		{
 			pxr::UsdGeomImageable Imageable(UsdPrim);
 			if (!Imageable)
