USD: Fix File -> Reload updating components many more times than needed.

diff --git a/Engine/Plugins/Importers/USDImporter/Source/USDStage/Private/USDStageActor.cpp b/Engine/Plugins/Importers/USDImporter/Source/USDStage/Private/USDStageActor.cpp
--- a/Engine/Plugins/Importers/USDImporter/Source/USDStage/Private/USDStageActor.cpp
+++ b/Engine/Plugins/Importers/USDImporter/Source/USDStage/Private/USDStageActor.cpp
@@ -1521,6 +1521,8 @@ void AUsdStageActor::OnUsdObjectsChanged(const UsdUtils::FObjectChangesByPath& I
 			bool bIsResync = Pair.Value;
 
 			bool bRemoveThisPath = false;
+
+			// Note: UE::FSdfPath::GetPrefixes() does not return the pseudoroot path "/"!
 			for (const UE::FSdfPath& Prefix : ThisPrim.GetPrefixes())
 			{
 				if (ResyncedPaths.Contains(Prefix))
@@ -1534,12 +1536,22 @@ void AUsdStageActor::OnUsdObjectsChanged(const UsdUtils::FObjectChangesByPath& I
 				continue;
 			}
 
+			CleanedPairs.Add(Pair);
+
 			if (bIsResync)
 			{
 				ResyncedPaths.Add(ThisPrim);
-			}
 
-			CleanedPairs.Add(Pair);
+				// Resyncing the root prim automatically means we have to resync the entire stage: There is no
+				// point checking for anything else. Since we sort the keys this should always come first, if
+				// present, so we can just break here.
+				//
+				// Note that this is how we reload stages, so it should happen relatively often
+				if (ThisPrim == UE::FSdfPath::AbsoluteRootPath())
+				{
+					break;
+				}
+			}
 		}
 		Swap(CleanedPairs, InOutMap);
 	};
