USD: Fix visibility track changes not writing back out to USD.

diff --git a/Engine/Plugins/Importers/USDImporter/Source/USDStage/Private/USDLevelSequenceHelper.cpp b/Engine/Plugins/Importers/USDImporter/Source/USDStage/Private/USDLevelSequenceHelper.cpp
--- a/Engine/Plugins/Importers/USDImporter/Source/USDStage/Private/USDLevelSequenceHelper.cpp
+++ b/Engine/Plugins/Importers/USDImporter/Source/USDStage/Private/USDLevelSequenceHelper.cpp
@@ -4543,6 +4543,10 @@ void FUsdLevelSequenceHelperImpl::HandleTrackChange(const UMovieSceneTrack& Trac
 				{
 					UnrealToUsd::ConvertBoolTrack(*BoolTrack, SequenceTransform, Writer.BoolWriter, UsdPrim);
 				}
+				else if (const UMovieSceneVisibilityTrack* VisTrack = Cast<const UMovieSceneVisibilityTrack>(&Track))
+				{
+					UnrealToUsd::ConvertBoolTrack(*VisTrack, SequenceTransform, Writer.BoolWriter, UsdPrim);
+				}
 				else if (const UMovieSceneColorTrack* ColorTrack = Cast<const UMovieSceneColorTrack>(&Track))
 				{
 					UnrealToUsd::ConvertColorTrack(*ColorTrack, SequenceTransform, Writer.ColorWriter, UsdPrim);
diff --git a/Engine/Plugins/Runtime/USDCore/Source/USDUtilities/Private/USDPrimConversion.cpp b/Engine/Plugins/Runtime/USDCore/Source/USDUtilities/Private/USDPrimConversion.cpp
--- a/Engine/Plugins/Runtime/USDCore/Source/USDUtilities/Private/USDPrimConversion.cpp
+++ b/Engine/Plugins/Runtime/USDCore/Source/USDUtilities/Private/USDPrimConversion.cpp
@@ -2454,7 +2454,7 @@ bool UnrealToUsd::ConvertAudioSection(
 }
 
 bool UnrealToUsd::ConvertBoolTrack(
-	const UMovieSceneBoolTrack& MovieSceneTrack,
+	const UMovieScenePropertyTrack& MovieSceneTrack,
 	const FMovieSceneSequenceTransform& SequenceTransform,
 	const TFunction<void(bool, double)>& WriterFunc,
 	UE::FUsdPrim& Prim
diff --git a/Engine/Plugins/Runtime/USDCore/Source/USDUtilities/Public/USDPrimConversion.h b/Engine/Plugins/Runtime/USDCore/Source/USDUtilities/Public/USDPrimConversion.h
--- a/Engine/Plugins/Runtime/USDCore/Source/USDUtilities/Public/USDPrimConversion.h
+++ b/Engine/Plugins/Runtime/USDCore/Source/USDUtilities/Public/USDPrimConversion.h
@@ -279,7 +279,7 @@ namespace UnrealToUsd
 		UE::FUsdPrim& Prim
 	);
 	USDUTILITIES_API bool ConvertBoolTrack(
-		const UMovieSceneBoolTrack& MovieSceneTrack,
+		const UMovieScenePropertyTrack& MovieSceneTrack,	// PropertyTrack here so it handles BoolTracks and VisibilityTracks
 		const FMovieSceneSequenceTransform& SequenceTransform,
 		const TFunction<void(bool, double)>& WriterFunc,
 		UE::FUsdPrim& Prim
