USD: Show warnings when authoring visibility tracks due to visibility inheritance issues;

diff --git a/Engine/Plugins/Importers/USDImporter/Source/USDStage/Private/USDLevelSequenceHelper.cpp b/Engine/Plugins/Importers/USDImporter/Source/USDStage/Private/USDLevelSequenceHelper.cpp
--- a/Engine/Plugins/Importers/USDImporter/Source/USDStage/Private/USDLevelSequenceHelper.cpp
+++ b/Engine/Plugins/Importers/USDImporter/Source/USDStage/Private/USDLevelSequenceHelper.cpp
@@ -101,6 +101,7 @@
 
 #if USE_USD_SDK
 #include "USDIncludesStart.h"
+#include "pxr/usd/usdGeom/tokens.h"
 #include "pxr/usd/usdMedia/tokens.h"
 #include "USDIncludesEnd.h"
 #endif	  // USE_USD_SDK
@@ -587,6 +588,125 @@ cleanup:
 			}
 		}
 	}
+
+	void ShowVisibilityWarningIfNeeded(const UMovieScenePropertyTrack* PropertyTrack, const UE::FUsdPrim& UsdPrim)
+	{
+		if (!PropertyTrack || !UsdPrim)
+		{
+			return;
+		}
+
+		FName PropertyPath = PropertyTrack->GetPropertyName();
+		if (PropertyPath != UnrealIdentifiers::HiddenPropertyName && PropertyPath != UnrealIdentifiers::HiddenInGamePropertyName)
+		{
+			return;
+		}
+
+		// Only show the warning after we have at least one key in the track, otherwise pressing 'File->Regenerate sequence' will
+		// really just wipe the empty tracks and bindings in the first place...
+		bool bHasKeys = false;
+		const TArray<UMovieSceneSection*>& Sections = PropertyTrack->GetAllSections();
+		for (UMovieSceneSection* Section : Sections)
+		{
+			if (Section->IsActive())
+			{
+				FMovieSceneChannelProxy& Proxy = Section->GetChannelProxy();
+				for (int32 ChannelIndex = 0; ChannelIndex < Proxy.NumChannels(); ++ChannelIndex)
+				{
+					if (FMovieSceneBoolChannel* Channel = Proxy.GetChannel<FMovieSceneBoolChannel>(ChannelIndex))
+					{
+						if (Channel->GetNumKeys() > 0)
+						{
+							bHasKeys = true;
+							break;
+						}
+					}
+				}
+			}
+
+			if (bHasKeys)
+			{
+				break;
+			}
+		}
+		if (!bHasKeys)
+		{
+			return;
+		}
+
+		const static FString VisibilityAttrName = UsdToUnreal::ConvertToken(pxr::UsdGeomTokens->visibility);
+		UE::FUsdAttribute VisibilityAttr = UsdPrim.GetAttribute(*VisibilityAttrName);
+		if (!VisibilityAttr || VisibilityAttr.GetNumTimeSamples() > 0)
+		{
+			// Only show the warning when first creating a visibility track for something that doesn't
+			// previously have any. Presumably if the visibility animation comes from USD the user is already aware
+			// of the differences between UE/USD given all the additional tracks we'll generate on the UE side
+			return;
+		}
+
+		const FText Text = LOCTEXT("VisibilityWarningTitle", "USD: Inherited visibility");
+
+		const FText SubText = LOCTEXT(
+			"VisibilityWarningTitle",
+			"Visibility in USD is inherited (if a parent prim is hidden, its children are also implicitly hidden), while it is not inherited in Unreal. This means that authoring visibility animation from Unreal may have unexpected consequences on the USD stage.\n\nYou may want to use 'File -> Regenerate sequence' to resynchronize the LevelSequence with the current state of the stage, whenever is convenient."
+		);
+
+		const UUsdProjectSettings* Settings = GetDefault<UUsdProjectSettings>();
+		if (Settings && Settings->bShowInheritedVisibilityWarning)
+		{
+			static TWeakPtr<SNotificationItem> Notification;
+
+			FNotificationInfo Toast(Text);
+			Toast.SubText = SubText;
+			Toast.Image = FCoreStyle::Get().GetBrush(TEXT("MessageLog.Warning"));
+			Toast.CheckBoxText = LOCTEXT("DontAskAgain", "Don't prompt again");
+			Toast.bUseLargeFont = false;
+			Toast.bFireAndForget = false;
+			Toast.FadeOutDuration = 0.0f;
+			Toast.ExpireDuration = 0.0f;
+			Toast.bUseThrobber = false;
+			Toast.bUseSuccessFailIcons = false;
+			Toast.ButtonDetails.Emplace(
+				LOCTEXT("OverridenOpinionMessageOk", "Ok"),
+				FText::GetEmpty(),
+				FSimpleDelegate::CreateLambda(
+					[]()
+					{
+						if (TSharedPtr<SNotificationItem> PinnedNotification = Notification.Pin())
+						{
+							PinnedNotification->SetCompletionState(SNotificationItem::CS_Success);
+							PinnedNotification->ExpireAndFadeout();
+						}
+					}
+				)
+			);
+			// This is flipped because the default checkbox message is "Don't prompt again"
+			Toast.CheckBoxState = Settings->bShowInheritedVisibilityWarning ? ECheckBoxState::Unchecked : ECheckBoxState::Checked;
+			Toast.CheckBoxStateChanged = FOnCheckStateChanged::CreateStatic(
+				[](ECheckBoxState NewState)
+				{
+					if (UUsdProjectSettings* Settings = GetMutableDefault<UUsdProjectSettings>())
+					{
+						// This is flipped because the default checkbox message is "Don't prompt again"
+						Settings->bShowInheritedVisibilityWarning = NewState == ECheckBoxState::Unchecked;
+						Settings->SaveConfig();
+					}
+				}
+			);
+
+			// Only show one at a time
+			if (!Notification.IsValid())
+			{
+				UE_LOG(LogUsd, Warning, TEXT("%s"), *SubText.ToString().Replace(TEXT("\n\n"), TEXT(" ")));
+				Notification = FSlateNotificationManager::Get().AddNotification(Toast);
+			}
+
+			if (TSharedPtr<SNotificationItem> PinnedNotification = Notification.Pin())
+			{
+				PinnedNotification->SetCompletionState(SNotificationItem::CS_Pending);
+			}
+		}
+	}
 #endif	  // WITH_EDITOR
 }	 // namespace UsdLevelSequenceHelperImpl
 
@@ -3524,9 +3644,9 @@ void FUsdLevelSequenceHelperImpl::UpdateUsdLayerOffsetFromSection(const UMovieSc
 	// This will obviously be quantized to frame intervals for now
 	double SubSectionStartTimeCode = TickResolution.AsSeconds(ModifiedStartFrame) * TimeCodesPerSecond;
 
-	const double FixedPlayRate = Section->Parameters.TimeScale.GetType() == EMovieSceneTimeWarpType::FixedPlayRate\
-		? Section->Parameters.TimeScale.AsFixedPlayRate()
-		: 1.f;
+	const double FixedPlayRate = Section->Parameters.TimeScale.GetType() == EMovieSceneTimeWarpType::FixedPlayRate
+									 ? Section->Parameters.TimeScale.AsFixedPlayRate()
+									 : 1.f;
 
 	UE::FSdfLayerOffset NewLayerOffset;
 	NewLayerOffset.Scale = FMath::IsNearlyZero(FixedPlayRate) ? 0.f : 1.f / FixedPlayRate;
@@ -4521,6 +4641,8 @@ void FUsdLevelSequenceHelperImpl::HandleTrackChange(const UMovieSceneTrack& Trac
 			// and pick up on changes to the section offset or play rate and bake out the UAnimSequence again
 			if (const UMovieScenePropertyTrack* PropertyTrack = Cast<const UMovieScenePropertyTrack>(&Track))
 			{
+				UsdLevelSequenceHelperImpl::ShowVisibilityWarningIfNeeded(PropertyTrack, UsdPrim);
+
 				TSet<FName> PropertyPathsToRefresh;
 				UnrealToUsd::FPropertyTrackWriter
 					Writer = UnrealToUsd::CreatePropertyTrackWriter(*BoundSceneComponent, *PropertyTrack, UsdPrim, PropertyPathsToRefresh);
diff --git a/Engine/Plugins/Runtime/USDCore/Source/USDClasses/Public/USDProjectSettings.h b/Engine/Plugins/Runtime/USDCore/Source/USDClasses/Public/USDProjectSettings.h
--- a/Engine/Plugins/Runtime/USDCore/Source/USDClasses/Public/USDProjectSettings.h
+++ b/Engine/Plugins/Runtime/USDCore/Source/USDClasses/Public/USDProjectSettings.h
@@ -101,6 +101,10 @@ public:
 	UPROPERTY(config, EditAnywhere, Category = "USD|Dialogs")
 	bool bShowTransformTrackOnCameraComponentWarning = true;
 
+	// Whether to show the warning dialog when authoring new visibility tracks from Unreal
+	UPROPERTY(config, EditAnywhere, Category = "USD|Dialogs")
+	bool bShowInheritedVisibilityWarning = true;
+
 	// Whether to display the pop up dialog asking what to do about dirty USD layers when saving the UE level
 	UPROPERTY(config, EditAnywhere, Category = "USD|Dialogs")
 	EUsdSaveDialogBehavior ShowSaveLayersDialogWhenSaving = EUsdSaveDialogBehavior::ShowPrompt;
