USD: Add option to the level sequence exporter to prevent the spawning of additional spawnables per subsection instance (5.5 version)

diff --git a/Engine/Plugins/Importers/USDImporter/Source/USDExporter/Private/LevelSequenceExporterUSD.cpp b/Engine/Plugins/Importers/USDImporter/Source/USDExporter/Private/LevelSequenceExporterUSD.cpp
--- a/Engine/Plugins/Importers/USDImporter/Source/USDExporter/Private/LevelSequenceExporterUSD.cpp
+++ b/Engine/Plugins/Importers/USDImporter/Source/USDExporter/Private/LevelSequenceExporterUSD.cpp
@@ -118,6 +118,7 @@ namespace UE::LevelSequenceExporterUSD::Private
 	public:
 
 		bool bDestroyingJustHides = true;
+		bool bExportSeparatePrimsPerSpawnableInstance = true;
 
 		virtual UObject* SpawnObject(
 			const FGuid& Guid,
@@ -171,6 +172,11 @@ namespace UE::LevelSequenceExporterUSD::Private
 					*ExistingIndex
 				);
 			}
+			else if (!bExportSeparatePrimsPerSpawnableInstance && ExistingInstancesForGuid.Num() > 0)
+			{
+				Object = ExistingInstancesForGuid[0];
+				SpawnableIndices.Add(InstanceKey, 0);
+			}
 
 			// We don't have an instance of the spawnable spawned for this exact movie sequence ID, but try to see if we can
 			// reuse any of the existing spawns for it.
@@ -907,9 +913,6 @@ namespace UE::LevelSequenceExporterUSD::Private
 								BoundObject = Context.SpawnRegister->FindSpawnedObject(Guid, SequenceInstance, BindingIndex++).Get();
 								if (!BoundObject)
 								{
-									// This should never happen as we preemptively spawn everything:
-									// At this point all our spawns should be spawned, but invisible
-									UE_LOG(LogUsd, Warning, TEXT("Failed to find spawned object for spawnable with Guid '%s'"), *Guid.ToString());
 									continue;
 								}
 
@@ -946,17 +949,9 @@ namespace UE::LevelSequenceExporterUSD::Private
 				}
 
 				// Go through FMovieSceneObjectCache and FindBoundObjects because that will also evaluate DynamicBindings.
-				// Note that we need to make sure that PreSpawnSpawnables has been called above this (at all, but also at
-				// least once *after* Context.SpawnRegister->CleanUp(), if that has been called). The idea here is that FindBoundObjects
-				// will manage to find the binding even it their "parent context" is a spawnable (e.g. if it's a possessable component of
-				// a spawnable) and also uses DynamicBindings, which is great! For it to be able to find our spawns however, the
-				// spawnables must be *currently* spawned.
-				// It doesn't help at all that our custom spawn register only hides stuff instead of destroying them, because even if the
-				// UObject itself still exists and is just hidden, having been "despawned" means the spawnable has been removed from the
-				// "Register" member of FMovieSceneSpawnRegister, and so the base part of the spawn register "doesn't know about it".
-				// For reference, check how FMovieSceneObjectCache::UpdateBindings (called by FindBoundObjects) will end up calling
-				// "Player.GetSpawnRegister().FindSpawnedObject", and observe how that in turn just checks the "Register" member...
-				// Ideally we could tweak a bit how the "Register" member is used, but that is part of the base FMovieSceneSpawnRegister.
+				// The idea here is that FindBoundObjects will manage to find the binding even it their "parent context" is
+				// a spawnable (e.g. if it's a possessable component of a spawnable) and also uses DynamicBindings, which is great!
+				// For it to be able to find our spawns however, the spawnables must be *currently* spawned.
 				TArrayView<TWeakObjectPtr<UObject>> ObjectWeakPtrs = ObjectCache.FindBoundObjects(Guid, *Context.Sequencer);
 				if (ObjectWeakPtrs.Num() > 0)
 				{
@@ -1815,6 +1810,7 @@ bool ULevelSequenceExporterUsd::ExportBinary(
 	TempSequencer->SetPlaybackStatus(EMovieScenePlayerStatus::Playing);
 
 	SpawnRegister->SetSequencer(TempSequencer);
+	SpawnRegister->bExportSeparatePrimsPerSpawnableInstance = Options->bExportSeparatePrimsPerSpawnableInstance;
 
 	LevelSequenceExporterImpl::FLevelSequenceExportContext Context{*LevelSequence, TempSequencer.ToSharedRef(), SpawnRegister.ToSharedRef()};
 	Context.ExportOptions = Options;
diff --git a/Engine/Plugins/Importers/USDImporter/Source/USDExporter/Public/LevelSequenceExporterUSDOptions.h b/Engine/Plugins/Importers/USDImporter/Source/USDExporter/Public/LevelSequenceExporterUSDOptions.h
--- a/Engine/Plugins/Importers/USDImporter/Source/USDExporter/Public/LevelSequenceExporterUSDOptions.h
+++ b/Engine/Plugins/Importers/USDImporter/Source/USDExporter/Public/LevelSequenceExporterUSDOptions.h
@@ -77,6 +77,14 @@ public:
 	)
 	bool bUseExportedLevelAsSublayer = true;
 
+	/**
+	 * If you have multiple instances of a same LevelSequence subsection with spawnables, setting this to true means each separate
+	 * subsection instance will get its own individual copy of the prim exported for the spawnable. Setting this to false means that
+	 * only a single prim will be exported for that spawnable, and shared for all subsection instances.
+	 */
+	UPROPERTY(EditAnywhere, config, BlueprintReadWrite, Category = "Level Sequence")
+	bool bExportSeparatePrimsPerSpawnableInstance = true;
+
 	/**
 	 * Whether to export levels and LevelSequences even if the existing files already describe the same versions of compatible assets.
 	 * This is only checked when bReplaceIdentical is set on the asset export task. Otherwise we'll never overwrite files.
