USD: Fix some performance bottlenecks on USD Stage Editor.

Only generate the needed child nodes instead of generating a node for every prim on the stage;
Don't call RestoreExpansionStates more than once per frame;
Optimize RestoreExpansionStates.

diff --git a/Engine/Plugins/Importers/USDImporter/Source/USDStageEditor/Private/SUSDLayersTreeView.cpp b/Engine/Plugins/Importers/USDImporter/Source/USDStageEditor/Private/SUSDLayersTreeView.cpp
--- a/Engine/Plugins/Importers/USDImporter/Source/USDStageEditor/Private/SUSDLayersTreeView.cpp
+++ b/Engine/Plugins/Importers/USDImporter/Source/USDStageEditor/Private/SUSDLayersTreeView.cpp
@@ -428,6 +428,8 @@ void SUsdLayersTreeView::Construct(const FArguments& InArgs)
 
 void SUsdLayersTreeView::Refresh(const UE::FUsdStageWeak& NewStage, const UE::FUsdStageWeak& InIsolatedStage, bool bResync)
 {
+	TRACE_CPUPROFILER_EVENT_SCOPE(SUsdLayersTreeView::Refresh);
+
 	if (bResync)
 	{
 		bool bShouldResetExpansionStates = false;
diff --git a/Engine/Plugins/Importers/USDImporter/Source/USDStageEditor/Private/SUSDStage.cpp b/Engine/Plugins/Importers/USDImporter/Source/USDStageEditor/Private/SUSDStage.cpp
--- a/Engine/Plugins/Importers/USDImporter/Source/USDStageEditor/Private/SUSDStage.cpp
+++ b/Engine/Plugins/Importers/USDImporter/Source/USDStageEditor/Private/SUSDStage.cpp
@@ -2618,6 +2618,8 @@ void SUsdStage::RequestFullRefresh()
 
 void SUsdStage::OnSlateTick(float Time)
 {
+	TRACE_CPUPROFILER_EVENT_SCOPE(SUsdStage::OnSlateTick);
+
 	if (bEditorIsShuttingDown)
 	{
 		return;
diff --git a/Engine/Plugins/Importers/USDImporter/Source/USDStageEditor/Private/SUSDStageTreeView.cpp b/Engine/Plugins/Importers/USDImporter/Source/USDStageEditor/Private/SUSDStageTreeView.cpp
--- a/Engine/Plugins/Importers/USDImporter/Source/USDStageEditor/Private/SUSDStageTreeView.cpp
+++ b/Engine/Plugins/Importers/USDImporter/Source/USDStageEditor/Private/SUSDStageTreeView.cpp
@@ -18,6 +18,7 @@
 #include "UsdWrappers/UsdPrim.h"
 #include "UsdWrappers/UsdStage.h"
 
+#include "Editor.h"
 #include "Framework/Commands/GenericCommands.h"
 #include "Framework/MultiBox/MultiBoxBuilder.h"
 #include "ScopedTransaction.h"
@@ -361,12 +362,41 @@ void SUsdStageTreeView::Construct(const FArguments& InArgs)
 				return;
 			}
 
-			TreeItemExpansionStates.Add(Prim.GetPrimPath().GetString(), bIsExpanded);
+			UsdPrimViewModel->SetIsExpanded(bIsExpanded);
+
+			// We have a special handling for the root because we'll want to manually expand it by
+			// default at first but also remember if the user collapsed it or not. For all
+			// other prims we truly just leave the nodes at default collapsed unless we have
+			// recorded that the node should be expanded
+			if (Prim.IsPseudoRoot())
+			{
+				RootWasExpanded = bIsExpanded;
+			}
+			else
+			{
+				const FString& PrimPath = Prim.GetPrimPath().GetString();
+				if (bIsExpanded)
+				{
+					ExpandedPrimPaths.Add(PrimPath);
+				}
+				else
+				{
+					ExpandedPrimPaths.Remove(PrimPath);
+				}
+			}
 		}
 	);
 
 	OnPrimSelectionChanged = InArgs._OnPrimSelectionChanged;
 
+	PostUndoRedoHandle = FEditorDelegates::PostUndoRedo.AddLambda(
+		[this]()
+		{
+			// This is in charge of restoring our expansion states after we undo/redo a prim rename
+			RequestExpansionStateRestore();
+		}
+	);
+
 	UICommandList = MakeShared<FUICommandList>();
 	UICommandList->MapAction(
 		FGenericCommands::Get().Cut,
@@ -405,6 +435,11 @@ void SUsdStageTreeView::Construct(const FArguments& InArgs)
 	);
 }
 
+SUsdStageTreeView::~SUsdStageTreeView()
+{
+	FEditorDelegates::PostUndoRedo.Remove(PostUndoRedoHandle);
+}
+
 TSharedRef<ITableRow> SUsdStageTreeView::OnGenerateRow(FUsdPrimViewModelRef InDisplayNode, const TSharedRef<STableViewBase>& OwnerTable)
 {
 	return SNew(SUsdTreeRow<FUsdPrimViewModelRef>, InDisplayNode, OwnerTable, SharedData);
@@ -420,6 +455,8 @@ void SUsdStageTreeView::OnGetChildren(FUsdPrimViewModelRef InParent, TArray<FUsd
 
 void SUsdStageTreeView::Refresh(const UE::FUsdStageWeak& NewStage)
 {
+	TRACE_CPUPROFILER_EVENT_SCOPE(SUsdStageTreeView::Refresh);
+
 	UE::FUsdStageWeak OldStage = RootItems.Num() > 0 ? RootItems[0]->UsdStage : UE::FUsdStageWeak();
 
 	RootItems.Empty();
@@ -428,7 +465,8 @@ void SUsdStageTreeView::Refresh(const UE::FUsdStageWeak& NewStage)
 
 	if (OldStage != NewStage)
 	{
-		TreeItemExpansionStates.Reset();
+		RootWasExpanded.Reset();
+		ExpandedPrimPaths.Reset();
 	}
 
 	if (NewStage)
@@ -438,12 +476,14 @@ void SUsdStageTreeView::Refresh(const UE::FUsdStageWeak& NewStage)
 			RootItems.Add(MakeShared<FUsdPrimViewModel>(nullptr, UsdStage, RootPrim));
 		}
 
-		RestoreExpansionStates();
+		RequestExpansionStateRestore();
 	}
 }
 
 void SUsdStageTreeView::RefreshPrim(const FString& PrimPath, bool bResync)
 {
+	TRACE_CPUPROFILER_EVENT_SCOPE(SUsdStageTreeView::RefreshPrim);
+
 	FScopedUnrealAllocs UnrealAllocs;	 // RefreshPrim can be called by a delegate for which we don't know the active allocator
 
 	FUsdPrimViewModelPtr FoundItem = GetItemFromPrimPath(PrimPath);
@@ -561,6 +601,8 @@ void SUsdStageTreeView::SelectItemsInternal(const TArray<FUsdPrimViewModelRef>&
 
 void SUsdStageTreeView::SetSelectedPrimPaths(const TArray<FString>& PrimPaths)
 {
+	TRACE_CPUPROFILER_EVENT_SCOPE(SetSelectedPrimPaths);
+
 	TArray<FUsdPrimViewModelRef> ItemsToSelect;
 	ItemsToSelect.Reserve(PrimPaths.Num());
 
@@ -993,8 +1035,7 @@ void SUsdStageTreeView::OnPastePrim()
 	for (const UE::FUsdPrim& ParentPrim : ParentPrims)
 	{
 		// Preemptively mark the parent prims as expanded so that we can always see what we pasted
-		const bool bIsExpanded = true;
-		TreeItemExpansionStates.Add(ParentPrim.GetPrimPath().GetString(), bIsExpanded);
+		ExpandedPrimPaths.Add(ParentPrim.GetPrimPath().GetString());
 
 		UsdUtils::PastePrims(ParentPrim);
 	}
@@ -1412,40 +1453,68 @@ bool SUsdStageTreeView::DoesPrimHaveSpecOnLocalLayerStack() const
 	return false;
 }
 
-void SUsdStageTreeView::RequestListRefresh()
-{
-	SUsdTreeView<FUsdPrimViewModelRef>::RequestListRefresh();
-	RestoreExpansionStates();
-}
-
-void SUsdStageTreeView::RestoreExpansionStates()
+void SUsdStageTreeView::Tick(const FGeometry& AllottedGeometry, const double InCurrentTime, const float InDeltaTime)
 {
-	TFunction<void(const FUsdPrimViewModelRef&)> SetExpansionRecursive = [&](const FUsdPrimViewModelRef& Item)
+	// Restore expansion states.
+	//
+	// We do this on tick so that we only do at most one of these per frame, and also so that we do it as delayed
+	// as possible, as during some busy transitions like undo/redo we may end up creating new FUsdPrimViewModels,
+	// and we only want to try restoring these expansion states after all FUsdPrimViewModels have been created
+	if (bNeedExpansionStateRefresh)
 	{
-		if (const UE::FUsdPrim& Prim = Item->UsdPrim)
+		TRACE_CPUPROFILER_EVENT_SCOPE(SUsdStageTreeView::RestoreExpansionStates);
+
+		// We should have only one root item, and it should be the expanded by default unless it was manually collapsed
+		if (RootItems.Num() > 0)
 		{
-			if (bool* bFoundExpansionState = TreeItemExpansionStates.Find(Prim.GetPrimPath().GetString()))
-			{
-				SetItemExpansion(Item, *bFoundExpansionState);
-			}
-			// Default to showing the root level expanded
-			else if (Prim.GetStage().GetPseudoRoot() == Prim)
+			const UE::FUsdPrim& RootPrim = RootItems[0]->UsdPrim;
+			if (RootPrim.IsPseudoRoot())
 			{
 				const bool bShouldExpand = true;
-				SetItemExpansion(Item, bShouldExpand);
+
+				const bool bDefaultValue = true;
+				if (RootWasExpanded.Get(bDefaultValue))
+				{
+					SetItemExpansion(RootItems[0], bShouldExpand);
+				}
+				else
+				{
+					SetItemExpansion(RootItems[0], !bShouldExpand);
+				}
 			}
 		}
 
-		for (const FUsdPrimViewModelRef& Child : Item->Children)
+		TFunction<void(const FUsdPrimViewModelRef&)> SetExpansionRecursive = [&](const FUsdPrimViewModelRef& Item)
+		{
+			if (const UE::FUsdPrim& Prim = Item->UsdPrim)
+			{
+				if (ExpandedPrimPaths.Contains(Prim.GetPrimPath().GetString()))
+				{
+					const bool bShouldExpand = true;
+					SetItemExpansion(Item, bShouldExpand);
+				}
+			}
+
+			for (const FUsdPrimViewModelRef& Child : Item->Children)
+			{
+				SetExpansionRecursive(Child);
+			}
+		};
+
+		for (const FUsdPrimViewModelRef& RootItem : RootItems)
 		{
-			SetExpansionRecursive(Child);
+			SetExpansionRecursive(RootItem);
 		}
-	};
 
-	for (const FUsdPrimViewModelRef& RootItem : RootItems)
-	{
-		SetExpansionRecursive(RootItem);
+		bNeedExpansionStateRefresh = false;
 	}
+
+	SUsdTreeView<FUsdPrimViewModelRef>::Tick(AllottedGeometry, InCurrentTime, InDeltaTime);
+}
+
+void SUsdStageTreeView::RequestExpansionStateRestore()
+{
+	bNeedExpansionStateRefresh = true;
 }
 
 void SUsdStageTreeView::OnToggleAllPayloads(EPayloadsTrigger PayloadsTrigger)
@@ -1795,30 +1864,35 @@ void SUsdStageTreeView::OnPrimNameCommitted(const FUsdPrimViewModelRef& ViewMode
 
 		FScopedTransaction Transaction(LOCTEXT("RenamePrimTransaction", "Rename a prim"));
 
-		// e.g. "/Root/OldPrim/"
+		// e.g. "/Root/OldPrim"
 		FString OldPath = ViewModel->UsdPrim.GetPrimPath().GetString();
 
 		// e.g. "NewPrim"
 		FString NewNameStr = InPrimName.ToString();
 
 		// Preemptively preserve the prim's expansion state because RenamePrim will trigger notices from within itself
-		// that will trigger refreshes of the tree view
+		// that will trigger refreshes of the tree view.
+		//
+		// Note: We don't remove the old paths in here, as that lets us undo the rename and
+		// preserve our expansion states
 		{
 			// e.g. "/Root/NewPrim"
 			FString NewPath = FString::Printf(TEXT("%s/%s"), *FPaths::GetPath(OldPath), *NewNameStr);
-			TMap<FString, bool> PairsToAdd;
-			for (TMap<FString, bool>::TIterator It(TreeItemExpansionStates); It; ++It)
+			TSet<FString> EntriesToAdd;
+			for (TSet<FString>::TIterator It(ExpandedPrimPaths); It; ++It)
 			{
 				// e.g. "/Root/OldPrim/SomeChild"
-				FString SomePrimPath = It->Key;
-				if (SomePrimPath.RemoveFromStart(OldPath))	  // e.g. "/SomeChild"
+				TStringView SomePrimPath = *It;
+				if (SomePrimPath.StartsWith(OldPath))
 				{
+					// e.g. "/SomeChild"
+					SomePrimPath.RemovePrefix(OldPath.Len());
+
 					// e.g. "/Root/NewPrim/SomeChild"
-					SomePrimPath = NewPath + SomePrimPath;
-					PairsToAdd.Add(SomePrimPath, It->Value);
+					EntriesToAdd.Add(NewPath + FString{SomePrimPath});
 				}
 			}
-			TreeItemExpansionStates.Append(PairsToAdd);
+			ExpandedPrimPaths.Append(EntriesToAdd);
 		}
 
 		UsdUtils::RenamePrim(ViewModel->UsdPrim, *NewNameStr);
diff --git a/Engine/Plugins/Importers/USDImporter/Source/USDStageEditor/Private/SUSDStageTreeView.h b/Engine/Plugins/Importers/USDImporter/Source/USDStageEditor/Private/SUSDStageTreeView.h
--- a/Engine/Plugins/Importers/USDImporter/Source/USDStageEditor/Private/SUSDStageTreeView.h
+++ b/Engine/Plugins/Importers/USDImporter/Source/USDStageEditor/Private/SUSDStageTreeView.h
@@ -30,6 +30,7 @@ public:
 	SLATE_END_ARGS()
 
 	void Construct(const FArguments& InArgs);
+	virtual ~SUsdStageTreeView() override;
 
 	void Refresh(const UE::FUsdStageWeak& NewStage);
 	void RefreshPrim(const FString& PrimPath, bool bResync);
@@ -92,9 +93,8 @@ private:
 	bool DoesPrimExistOnEditTarget() const;
 	bool DoesPrimHaveSpecOnLocalLayerStack() const;
 
-	/** Uses TreeItemExpansionStates to travel the tree and call SetItemExpansion */
-	void RestoreExpansionStates();
-	virtual void RequestListRefresh() override;
+	void RequestExpansionStateRestore();
+	virtual void Tick( const FGeometry& AllottedGeometry, const double InCurrentTime, const float InDeltaTime ) override;
 
 	void SelectItemsInternal(const TArray<FUsdPrimViewModelRef>& ItemsToSelect);
 
@@ -122,10 +122,14 @@ private:
 	TWeakPtr<FUsdPrimViewModel> PendingRenameItem;
 
 	// So that we can store these across refreshes
-	TMap<FString, bool> TreeItemExpansionStates;
+	TSet<FString> ExpandedPrimPaths;
+	TOptional<bool> RootWasExpanded;
+	bool bNeedExpansionStateRefresh = false;
 
 	FOnPrimSelectionChanged OnPrimSelectionChanged;
 
+	FDelegateHandle PostUndoRedoHandle;
+
 	TSharedPtr<FUICommandList> UICommandList;
 };
 
diff --git a/Engine/Plugins/Importers/USDImporter/Source/USDStageEditorViewModels/Private/USDPrimViewModel.cpp b/Engine/Plugins/Importers/USDImporter/Source/USDStageEditorViewModels/Private/USDPrimViewModel.cpp
--- a/Engine/Plugins/Importers/USDImporter/Source/USDStageEditorViewModels/Private/USDPrimViewModel.cpp
+++ b/Engine/Plugins/Importers/USDImporter/Source/USDStageEditorViewModels/Private/USDPrimViewModel.cpp
@@ -30,8 +30,13 @@ FUsdPrimViewModel::FUsdPrimViewModel(FUsdPrimViewModel* InParentItem, const UE::
 	, ParentItem(InParentItem)
 	, RowData(MakeShared<FUsdPrimModel>())
 {
+	TRACE_CPUPROFILER_EVENT_SCOPE(FUsdPrimViewModel::FUsdPrimViewModel);
 	RefreshData(false);
-	FillChildren();
+
+	if (ShouldGenerateChildren())
+	{
+		FillChildren();
+	}
 }
 
 TArray<FUsdPrimViewModelRef>& FUsdPrimViewModel::UpdateChildren()
@@ -90,7 +95,7 @@ TArray<FUsdPrimViewModelRef>& FUsdPrimViewModel::UpdateChildren()
 		}
 	}
 
-	if (bNeedsRefresh)
+	if (bNeedsRefresh && ShouldGenerateChildren())
 	{
 		FillChildren();
 	}
@@ -99,6 +104,39 @@ TArray<FUsdPrimViewModelRef>& FUsdPrimViewModel::UpdateChildren()
 	return Children;
 }
 
+void FUsdPrimViewModel::SetIsExpanded(bool bNewIsExpanded)
+{
+	if (bNewIsExpanded == bIsExpanded)
+	{
+		return;
+	}
+	bIsExpanded = bNewIsExpanded;
+
+	// We should always have our own immediate children up-to-date, as that is needed to get an expander arrow.
+	// If we're collapsed though, we don't have to have grandchildren
+	if (bIsExpanded)
+	{
+		for (FUsdPrimViewModelRef Child : Children)
+		{
+			Child->FillChildren();
+		}
+	}
+	else
+	{
+		for (FUsdPrimViewModelRef Child : Children)
+		{
+			Child->Children.Reset();
+		}
+	}
+}
+
+bool FUsdPrimViewModel::ShouldGenerateChildren() const
+{
+	// We need to generate children if our parent is expanded, because having child nodes is what makes
+	// the treeview give us an expander arrow in the first place
+	return (!ParentItem || ParentItem->bIsExpanded);
+}
+
 void FUsdPrimViewModel::FillChildren()
 {
 #if USE_USD_SDK
diff --git a/Engine/Plugins/Importers/USDImporter/Source/USDStageEditorViewModels/Public/USDPrimViewModel.h b/Engine/Plugins/Importers/USDImporter/Source/USDStageEditorViewModels/Public/USDPrimViewModel.h
--- a/Engine/Plugins/Importers/USDImporter/Source/USDStageEditorViewModels/Public/USDPrimViewModel.h
+++ b/Engine/Plugins/Importers/USDImporter/Source/USDStageEditorViewModels/Public/USDPrimViewModel.h
@@ -57,6 +57,8 @@ public:
 
 	TArray<FUsdPrimViewModelRef>& UpdateChildren();
 
+	void SetIsExpanded(bool bNewIsExpanded);
+	bool ShouldGenerateChildren() const;
 	void FillChildren();
 
 	void RefreshData(bool bRefreshChildren);
@@ -95,4 +97,7 @@ public:
 	TSharedRef<FUsdPrimModel> RowData;	  // Data model
 
 	bool bIsRenamingExistingPrim = false;
+
+private:
+	bool bIsExpanded = false;	 // Set by our parent tree, controls whether we generate grandchildren or not
 };
