USD: Fix some mistakes/duplicate entries on LOCTEXT macros.

diff --git a/Engine/Plugins/Importers/USDImporter/Source/USDStage/Private/USDLevelSequenceHelper.cpp b/Engine/Plugins/Importers/USDImporter/Source/USDStage/Private/USDLevelSequenceHelper.cpp
--- a/Engine/Plugins/Importers/USDImporter/Source/USDStage/Private/USDLevelSequenceHelper.cpp
+++ b/Engine/Plugins/Importers/USDImporter/Source/USDStage/Private/USDLevelSequenceHelper.cpp
@@ -647,7 +647,7 @@ cleanup:
 		const FText Text = LOCTEXT("VisibilityWarningTitle", "USD: Inherited visibility");
 
 		const FText SubText = LOCTEXT(
-			"VisibilityWarningTitle",
+			"VisibilityWarningTooltip",
 			"Visibility in USD is inherited (if a parent prim is hidden, its children are also implicitly hidden), while it is not inherited in Unreal. This means that authoring visibility animation from Unreal may have unexpected consequences on the USD stage.\n\nYou may want to use 'File -> Regenerate sequence' to resynchronize the LevelSequence with the current state of the stage, whenever is convenient."
 		);
 
diff --git a/Engine/Plugins/Importers/USDImporter/Source/USDStage/Private/USDLocatorEditors.cpp b/Engine/Plugins/Importers/USDImporter/Source/USDStage/Private/USDLocatorEditors.cpp
--- a/Engine/Plugins/Importers/USDImporter/Source/USDStage/Private/USDLocatorEditors.cpp
+++ b/Engine/Plugins/Importers/USDImporter/Source/USDStage/Private/USDLocatorEditors.cpp
@@ -166,7 +166,7 @@ namespace UE::UniversalObjectLocator
 			}
 		}
 
-		return LOCTEXT("UsdPrimLocatorEditorTooltip", "A path to a prim on a USD stage");
+		return LOCTEXT("UsdPrimLocatorEditorTooltipInvalid", "A path to a prim on a USD stage");
 	}
 
 	FSlateIcon FUsdPrimLocatorEditor::GetDisplayIcon(const FUniversalObjectLocatorFragment* InFragment) const
diff --git a/Engine/Plugins/Importers/USDImporter/Source/USDStageEditor/Private/SUSDStage.cpp b/Engine/Plugins/Importers/USDImporter/Source/USDStageEditor/Private/SUSDStage.cpp
--- a/Engine/Plugins/Importers/USDImporter/Source/USDStageEditor/Private/SUSDStage.cpp
+++ b/Engine/Plugins/Importers/USDImporter/Source/USDStageEditor/Private/SUSDStage.cpp
@@ -899,8 +899,8 @@ void SUsdStage::FillFileMenu(FMenuBuilder& MenuBuilder)
 	MenuBuilder.BeginSection("Reload", LOCTEXT("Reload", "Reload"));
 	{
 		MenuBuilder.AddMenuEntry(
-			LOCTEXT("Reload", "Reload stage"),
-			LOCTEXT("Reload_ToolTip", "Reloads the stage and animations from disk, keeping aspects of the session intact"),
+			LOCTEXT("ReloadEntry", "Reload stage"),
+			LOCTEXT("ReloadEntry_ToolTip", "Reloads the stage and animations from disk, keeping aspects of the session intact"),
 			FSlateIcon(),
 			FUIAction(
 				FExecuteAction::CreateSP(this, &SUsdStage::FileReload),
