Fixed build issue.

diff --git a/Engine/Plugins/Importers/USDImporter/Source/USDStage/Private/USDLevelSequenceHelper.cpp b/Engine/Plugins/Importers/USDImporter/Source/USDStage/Private/USDLevelSequenceHelper.cpp
--- a/Engine/Plugins/Importers/USDImporter/Source/USDStage/Private/USDLevelSequenceHelper.cpp
+++ b/Engine/Plugins/Importers/USDImporter/Source/USDStage/Private/USDLevelSequenceHelper.cpp
@@ -4641,7 +4641,9 @@ void FUsdLevelSequenceHelperImpl::HandleTrackChange(const UMovieSceneTrack& Trac
 			// and pick up on changes to the section offset or play rate and bake out the UAnimSequence again
 			if (const UMovieScenePropertyTrack* PropertyTrack = Cast<const UMovieScenePropertyTrack>(&Track))
 			{
+#if WITH_EDITOR
 				UsdLevelSequenceHelperImpl::ShowVisibilityWarningIfNeeded(PropertyTrack, UsdPrim);
+#endif
 
 				TSet<FName> PropertyPathsToRefresh;
 				UnrealToUsd::FPropertyTrackWriter
