USD: Fix struct initialization error for FUsdPrimLocatorFragment.

diff --git a/Engine/Plugins/Importers/USDImporter/Source/USDStage/Public/USDLocatorFragments.h b/Engine/Plugins/Importers/USDImporter/Source/USDStage/Public/USDLocatorFragments.h
--- a/Engine/Plugins/Importers/USDImporter/Source/USDStage/Public/USDLocatorFragments.h
+++ b/Engine/Plugins/Importers/USDImporter/Source/USDStage/Public/USDLocatorFragments.h
@@ -18,7 +18,7 @@ struct FUsdPrimLocatorFragment
 	GENERATED_BODY()
 
 	UPROPERTY()
-	bool bPreferComponent;
+	bool bPreferComponent = false;
 
 	UPROPERTY()
 	FString PrimPath;
