USD: General FUsdInfoCache build optimizations.

Refactor info cache to fill in collapse roots on-demand;
Remove info cache lock contention by using separate locks per prim and putting all the info on the per-prim structs whenever possible;
Remove some unnecessary steps when registering instanceable aux prims;
Remove unnecessary handling of assets/component collapsing during info cache build;
Optimize a couple of steps in UsdUtils::GetPrimMaterialAssignments as that is in the hot path of the info cache build;
Extract CanBeCollapsed checks from XformableTranslator::CollapsesChildren into a dedicated pass;
Update automated test with a couple checks for main-aux prim links in instanceable/reference cases.

diff --git a/Engine/Plugins/Importers/USDImporter/Source/USDSchemas/Private/USDGeomMeshTranslator.cpp b/Engine/Plugins/Importers/USDImporter/Source/USDSchemas/Private/USDGeomMeshTranslator.cpp
--- a/Engine/Plugins/Importers/USDImporter/Source/USDSchemas/Private/USDGeomMeshTranslator.cpp
+++ b/Engine/Plugins/Importers/USDImporter/Source/USDSchemas/Private/USDGeomMeshTranslator.cpp
@@ -37,6 +37,7 @@
 #include "Misc/App.h"
 #include "Misc/SecureHash.h"
 #include "Modules/ModuleManager.h"
+#include "NaniteDefinitions.h"
 #include "PhysicsEngine/BodySetup.h"
 #include "StaticMeshAttributes.h"
 #include "StaticMeshOperations.h"
@@ -205,8 +206,7 @@ namespace UsdGeomMeshTranslatorImpl
 		{
 			TOptional<uint64> SubtreeSectionCount = Context.InfoCache->GetSubtreeMaterialSlotCount(PrimPath);
 
-			const int32 MaxNumSections = 64;	// There is no define for this, but it's checked for on NaniteBuilder.cpp, FBuilderModule::Build
-			if (!SubtreeSectionCount.IsSet() || SubtreeSectionCount.GetValue() > MaxNumSections)
+			if (!SubtreeSectionCount.IsSet() || SubtreeSectionCount.GetValue() > NANITE_MAX_CLUSTER_MATERIALS)
 			{
 				UE_LOG(
 					LogUsd,
@@ -215,7 +215,7 @@ namespace UsdGeomMeshTranslatorImpl
 						 "'%d'"),
 					*PrimPath.GetString(),
 					SubtreeSectionCount.GetValue(),
-					MaxNumSections
+					NANITE_MAX_CLUSTER_MATERIALS
 				);
 				return false;
 			}
@@ -742,8 +742,8 @@ namespace UsdGeomMeshTranslatorImpl
 		{
 			StaticMesh.GetRenderData()->InitializeRayTracingRepresentationFromRenderingLODs();
 		}
-#endif // RHI_RAYTRACING
-#endif // WITH_EDITOR
+#endif	  // RHI_RAYTRACING
+#endif	  // WITH_EDITOR
 
 		return true;
 	}
diff --git a/Engine/Plugins/Importers/USDImporter/Source/USDSchemas/Private/USDGeomXformableTranslator.cpp b/Engine/Plugins/Importers/USDImporter/Source/USDSchemas/Private/USDGeomXformableTranslator.cpp
--- a/Engine/Plugins/Importers/USDImporter/Source/USDSchemas/Private/USDGeomXformableTranslator.cpp
+++ b/Engine/Plugins/Importers/USDImporter/Source/USDSchemas/Private/USDGeomXformableTranslator.cpp
@@ -1067,6 +1067,7 @@ bool FUsdGeomXformableTranslator::CollapsesChildren(ECollapsingType CollapsingTy
 		// Note that this is false if PrimKind is None
 		const bool bPrimKindShouldCollapse = EnumHasAnyFlags(Context->KindsToCollapse, PrimKind);
 
+		// This indicates whether we *want* to collapse
 		bCollapsesChildren = Context->KindsToCollapse != EUsdDefaultKind::None
 							 && (bPrimKindShouldCollapse || (PrimKind == EUsdDefaultKind::None && GCollapsePrimsWithoutKind));
 
@@ -1078,29 +1079,11 @@ bool FUsdGeomXformableTranslator::CollapsesChildren(ECollapsingType CollapsingTy
 
 		if (bCollapsesChildren)
 		{
-			IUsdSchemasModule& UsdSchemasModule = FModuleManager::Get().LoadModuleChecked<IUsdSchemasModule>(TEXT("USDSchemas"));
-
-			// TODO: This can be optimized in order to make FUsdInfoCache::RebuildCacheForSubtree faster: If we have a child prim that we know doesn't
-			// collapse, any of our parents should be able to know they can't collapse *us* either. This is somewhat niche though: Realistically to
-			// waste time here a prim and its children need to have a kind that allows collapsing, and also not be able to collapse. Also, if any of
-			// these prims *does* manage to collapse, FUsdInfoCache will already not actually query the subtree children if they can collapse or not
-			// anymore, and just consider them collapsed by the parent
-			TArray<TUsdStore<pxr::UsdPrim>> ChildXformPrims = UsdUtils::GetAllPrimsOfType(Prim, pxr::TfType::Find<pxr::UsdGeomXformable>());
-			for (const TUsdStore<pxr::UsdPrim>& ChildXformPrim : ChildXformPrims)
+			// This indicates whether the subtree *can* be collapsed
+			TOptional<bool> bSubtreeCanBeCollapsed = Context->InfoCache->CanXformableSubtreeBeCollapsed(PrimPath, *Context);
+			if (bSubtreeCanBeCollapsed.IsSet())
 			{
-				if (ChildXformPrim.Get().IsA<pxr::UsdSkelRoot>())
-				{
-					return false;
-				}
-
-				if (TSharedPtr<FUsdSchemaTranslator> SchemaTranslator = UsdSchemasModule.GetTranslatorRegistry()
-																			.CreateTranslatorForSchema(Context, UE::FUsdTyped(ChildXformPrim.Get())))
-				{
-					if (!SchemaTranslator->CanBeCollapsed(CollapsingType))
-					{
-						return false;
-					}
-				}
+				return bSubtreeCanBeCollapsed.GetValue();
 			}
 		}
 	}
diff --git a/Engine/Plugins/Importers/USDImporter/Source/USDSchemas/Private/USDInfoCache.cpp b/Engine/Plugins/Importers/USDImporter/Source/USDSchemas/Private/USDInfoCache.cpp
--- a/Engine/Plugins/Importers/USDImporter/Source/USDSchemas/Private/USDInfoCache.cpp
+++ b/Engine/Plugins/Importers/USDImporter/Source/USDSchemas/Private/USDInfoCache.cpp
@@ -10,6 +10,7 @@
 
 #include "UsdWrappers/SdfPath.h"
 #include "UsdWrappers/UsdPrim.h"
+#include "UsdWrappers/UsdStage.h"
 
 #include "Async/ParallelFor.h"
 #include "Misc/Paths.h"
@@ -56,6 +57,15 @@ static FAutoConsoleVariableRef CVarGeometryCacheMaxDepth(
 	TEXT("Maximum distance between an animated mesh prim to its collapsed geometry cache root")
 );
 
+static int32 GNumPerPrimLocks = 32;
+static FAutoConsoleVariableRef CVarNumPerPrimLocks(
+	TEXT("USD.NumPerPrimLocks"),
+	GNumPerPrimLocks,
+	TEXT(
+		"Maximum number of locks that are distributed between all the prim info structs that the USDInfoCache keeps internally. More locks can imply better performance for the info cache build, but the total number of locks available on the system is finite"
+	)
+);
+
 namespace UE::UsdInfoCache::Private
 {
 	// Flags to hint at the state of a prim for the purpose of geometry cache
@@ -72,24 +82,60 @@ namespace UE::UsdInfoCache::Private
 
 	struct FUsdPrimInfo
 	{
-		UE::FSdfPath AssetCollapsedRoot;
-		UE::FSdfPath ComponentCollapsedRoot;
+		int32 PrimLockIndex = INDEX_NONE;
+
+		// Points to our collapsing root
+		// - Optional is not set: Prim wasn't visited on cache build. It's collapsed, but we don't know the root (yet);
+		// - Optional contains empty path: Prim was visited on build, we know it doesn't collapse and it isn't collapsed;
+		// - Optional contains prim's own path: This prim is a collapse root itself: It CollapsesChildren();
+		// - Optional contain other prim's path: This prim is collapsed, and that other prim is our collapse root (it CollapsesChildren());
+		TOptional<UE::FSdfPath> AssetCollapsedRoot;
+
+		// Whether this prim can be collapsed or not, according to its schema translator
+		// - Optional is not set: Prim wasn't visited yet, we don't know
+		// - Optional has value: Whether the prim can be collapsed or not
+		TOptional<bool> bXformSubtreeCanBeCollapsed;
+
 		uint64 ExpectedVertexCountForSubtree = 0;
 		TArray<UsdUtils::FUsdPrimMaterialSlot> SubtreeMaterialSlots;
+
 		int32 GeometryCacheDepth = -1;
 		EGeometryCachePrimState GeometryCacheState = EGeometryCachePrimState::None;
+
+		// Paths to material prims to the mesh prims they are bound to in the scene, given the current settings for
+		// render context, material purpose, variant selections, etc.
+		TSet<UE::FSdfPath> MaterialUsers;
+
+		// Maps from prims, to all the prims that require also reading this prim to be translated into an asset.
+		// Mainly used to update these assets whenever the depencency prim is updated.
+		TSet<UE::FSdfPath> MainPrims;
+		TSet<UE::FSdfPath> AuxPrims;
+
+		FUsdPrimInfo() = default;
+		FUsdPrimInfo(int InPrimLockIndex)
+			: PrimLockIndex(InPrimLockIndex)
+		{
+		}
 	};
-}
+}	 // namespace UE::UsdInfoCache::Private
 
 FArchive& operator<<(FArchive& Ar, UE::UsdInfoCache::Private::FUsdPrimInfo& Info)
 {
+	Ar << Info.PrimLockIndex;
+
 	Ar << Info.AssetCollapsedRoot;
-	Ar << Info.ComponentCollapsedRoot;
+
 	Ar << Info.ExpectedVertexCountForSubtree;
 	Ar << Info.SubtreeMaterialSlots;
+
 	Ar << Info.GeometryCacheDepth;
 	Ar << Info.GeometryCacheState;
 
+	Ar << Info.MaterialUsers;
+
+	Ar << Info.MainPrims;
+	Ar << Info.AuxPrims;
+
 	return Ar;
 }
 
@@ -99,6 +145,9 @@ struct FUsdInfoCache::FUsdInfoCacheImpl
 		: AllowedExtensionsForGeometryCacheSource(UnrealUSDWrapper::GetNativeFileFormats())
 	{
 		AllowedExtensionsForGeometryCacheSource.Add(TEXT("abc"));
+
+		FScopedUnrealAllocs Allocs;
+		PrimLocks = new FRWLock[GNumPerPrimLocks]();
 	}
 
 	FUsdInfoCacheImpl(const FUsdInfoCacheImpl& Other)
@@ -107,18 +156,18 @@ struct FUsdInfoCache::FUsdInfoCacheImpl
 		// Probably not necessary
 		FReadScopeLock ScopedInfoMapLock(Other.InfoMapLock);
 		FReadScopeLock ScopedPrimPathToAssetsLock(Other.PrimPathToAssetsLock);
-		FReadScopeLock ScopedMaterialUsersLock(Other.MaterialUsersLock);
-		FReadScopeLock ScopedAuxPrimsLock(Other.AuxiliaryPrimsLock);
 
 		InfoMap = Other.InfoMap;
 		PrimPathToAssets = Other.PrimPathToAssets;
 		AssetToPrimPaths = Other.AssetToPrimPaths;
-		MaterialUsers = Other.MaterialUsers;
-		AuxToMainPrims = Other.AuxToMainPrims;
-		MainToAuxPrims = Other.MainToAuxPrims;
 		AllowedExtensionsForGeometryCacheSource = Other.AllowedExtensionsForGeometryCacheSource;
 	}
 
+	~FUsdInfoCacheImpl()
+	{
+		delete[] PrimLocks;
+	}
+
 	// Information we must have about all prims on the stage
 	TMap<UE::FSdfPath, UE::UsdInfoCache::Private::FUsdPrimInfo> InfoMap;
 	mutable FRWLock InfoMapLock;
@@ -128,20 +177,19 @@ struct FUsdInfoCache::FUsdInfoCacheImpl
 	TMap<TWeakObjectPtr<UObject>, TArray<UE::FSdfPath>> AssetToPrimPaths;
 	mutable FRWLock PrimPathToAssetsLock;
 
-	// Paths to material prims to the mesh prims they are bound to in the scene, given the current settings for
-	// render context, material purpose, variant selections, etc.
-	TMap<UE::FSdfPath, TSet<UE::FSdfPath>> MaterialUsers;
-	mutable FRWLock MaterialUsersLock;
-
-	// Maps from prims, to all the prims that require also reading this prim to be translated into an asset.
-	// Mainly used to update these assets whenever the depencency prim is updated.
-	TMap<UE::FSdfPath, TSet<UE::FSdfPath>> AuxToMainPrims;
-	TMap<UE::FSdfPath, TSet<UE::FSdfPath>> MainToAuxPrims;
-	mutable FRWLock AuxiliaryPrimsLock;
+	// Temporarily used during the info cache build, as we need to do another pass on point instancers afterwards
+	TArray<FString> TempPointInstancerPaths;
+	mutable FRWLock TempPointInstancerPathsLock;
 
 	// Geometry cache can come from a reference or payload of these file types
 	TArray<FString> AllowedExtensionsForGeometryCacheSource;
 
+	// Valid only during the main info cache build
+	UE::FUsdStageWeak TempStage;
+
+	// Individual locks distributed across the FUsdPrimInfo in the InfoMap
+	FRWLock* PrimLocks = nullptr;
+
 public:
 	void RegisterAuxiliaryPrims(const UE::FSdfPath& MainPrimPath, const TSet<UE::FSdfPath>& AuxPrimPaths)
 	{
@@ -152,15 +200,23 @@ public:
 			return;
 		}
 
-		FWriteScopeLock ScopeLock(AuxiliaryPrimsLock);
+		FReadScopeLock ScopeLock{InfoMapLock};
 
-		MainToAuxPrims.FindOrAdd(MainPrimPath).Append(AuxPrimPaths);
+		if (UE::UsdInfoCache::Private::FUsdPrimInfo* MainPrim = InfoMap.Find(MainPrimPath))
+		{
+			FWriteScopeLock PrimLock{PrimLocks[MainPrim->PrimLockIndex]};
+			MainPrim->AuxPrims.Append(AuxPrimPaths);
+		}
 
 		for (const UE::FSdfPath& AuxPrimPath : AuxPrimPaths)
 		{
-			UE_LOG(LogUsd, Verbose, TEXT("Registering main prim '%s' and aux prim '%s'"), *MainPrimPath.GetString(), *AuxPrimPath.GetString());
+			if (UE::UsdInfoCache::Private::FUsdPrimInfo* AuxPrim = InfoMap.Find(AuxPrimPath))
+			{
+				FWriteScopeLock PrimLock{PrimLocks[AuxPrim->PrimLockIndex]};
+				AuxPrim->MainPrims.Add(MainPrimPath);
+			}
 
-			AuxToMainPrims.FindOrAdd(AuxPrimPath).Add(MainPrimPath);
+			UE_LOG(LogUsd, Verbose, TEXT("Registering main prim '%s' and aux prim '%s'"), *MainPrimPath.GetString(), *AuxPrimPath.GetString());
 		}
 	}
 };
@@ -192,15 +248,6 @@ bool FUsdInfoCache::Serialize(FArchive& Ar)
 			Ar << ImplPtr->PrimPathToAssets;
 			Ar << ImplPtr->AssetToPrimPaths;
 		}
-		{
-			FWriteScopeLock ScopeLock(ImplPtr->MaterialUsersLock);
-			Ar << ImplPtr->MaterialUsers;
-		}
-		{
-			FWriteScopeLock ScopeLock(ImplPtr->AuxiliaryPrimsLock);
-			Ar << ImplPtr->AuxToMainPrims;
-			Ar << ImplPtr->MainToAuxPrims;
-		}
 	}
 
 	return true;
@@ -240,11 +287,10 @@ bool FUsdInfoCache::IsPathCollapsed(const UE::FSdfPath& Path, ECollapsingType Co
 
 		if (const UE::UsdInfoCache::Private::FUsdPrimInfo* FoundInfo = ImplPtr->InfoMap.Find(Path))
 		{
-			const UE::FSdfPath* CollapsedRoot = CollapsingType == ECollapsingType::Assets ? &FoundInfo->AssetCollapsedRoot
-																						  : &FoundInfo->ComponentCollapsedRoot;
+			FReadScopeLock PrimLock{ImplPtr->PrimLocks[FoundInfo->PrimLockIndex]};
 
-			// A non-empty path to another prim means this prim is collapsed into that one
-			return !CollapsedRoot->IsEmpty() && *CollapsedRoot != Path;
+			// See comment on the AssetCollapsedRoot member
+			return !FoundInfo->AssetCollapsedRoot.IsSet() || (!FoundInfo->AssetCollapsedRoot->IsEmpty() && FoundInfo->AssetCollapsedRoot != Path);
 		}
 
 		// This should never happen: We should have cached the entire tree
@@ -262,12 +308,11 @@ bool FUsdInfoCache::DoesPathCollapseChildren(const UE::FSdfPath& Path, ECollapsi
 
 		if (const UE::UsdInfoCache::Private::FUsdPrimInfo* FoundInfo = ImplPtr->InfoMap.Find(Path))
 		{
-			const UE::FSdfPath* CollapsedRoot = CollapsingType == ECollapsingType::Assets ? &FoundInfo->AssetCollapsedRoot
-																						  : &FoundInfo->ComponentCollapsedRoot;
+			FReadScopeLock PrimLock{ImplPtr->PrimLocks[FoundInfo->PrimLockIndex]};
 
 			// We store our own Path in there when we collapse children.
 			// Otherwise we hold the path of our collapse root, or empty (in case nothing is collapsed up to here)
-			return *CollapsedRoot == Path;
+			return FoundInfo->AssetCollapsedRoot == Path;
 		}
 
 		// This should never happen: We should have cached the entire tree
@@ -277,31 +322,193 @@ bool FUsdInfoCache::DoesPathCollapseChildren(const UE::FSdfPath& Path, ECollapsi
 	return false;
 }
 
+namespace UE::USDInfoCache::Private
+{
+#if USE_USD_SDK
+	bool RecursiveQueryCanBeCollapsed(
+		const pxr::UsdPrim& UsdPrim,
+		FUsdSchemaTranslationContext& Context,
+		FUsdInfoCache::FUsdInfoCacheImpl& Impl,
+		FUsdSchemaTranslatorRegistry& Registry
+	)
+	{
+		TRACE_CPUPROFILER_EVENT_SCOPE(UE::USDInfoCache::Private::RecursiveQueryCanBeCollapsed);
+
+		UE::FSdfPath UsdPrimPath = UE::FSdfPath{UsdPrim.GetPrimPath()};
+
+		// If we already have a value for our prim then we can just return it right now. We only fill these bCanBeCollapsed values
+		// through here, so if we know e.g. that UsdPrim can be collapsed, we know its entire subtree can too.
+		UE::UsdInfoCache::Private::FUsdPrimInfo* MainPrimInfo = Impl.InfoMap.Find(UsdPrimPath);
+		if (MainPrimInfo)
+		{
+			FReadScopeLock PrimLock{Impl.PrimLocks[MainPrimInfo->PrimLockIndex]};
+			if (MainPrimInfo->bXformSubtreeCanBeCollapsed.IsSet())
+			{
+				return MainPrimInfo->bXformSubtreeCanBeCollapsed.GetValue();
+			}
+		}
+
+		// If we're here, we don't know whether UsdPrim CanBeCollapsed or not.
+		// Since these calls are usually cheap, let's just query it for ourselves right now
+		bool bCanBeCollapsed = true;
+		if (TSharedPtr<FUsdSchemaTranslator> SchemaTranslator = Registry.CreateTranslatorForSchema(Context.AsShared(), UE::FUsdTyped(UsdPrim)))
+		{
+			bCanBeCollapsed = SchemaTranslator->CanBeCollapsed(ECollapsingType::Assets);
+		}
+
+		// If we can be collapsed ourselves we're not still done, because this is about the subtree. If any of our
+		// children can't be collapsed, we actually can't either
+		if (bCanBeCollapsed)
+		{
+			TArray<pxr::UsdPrim> Children;
+			for (pxr::UsdPrim Child : UsdPrim.GetFilteredChildren(pxr::UsdTraverseInstanceProxies(pxr::UsdPrimAllPrimsPredicate)))
+			{
+				// We don't care about non-xformable prims (materials, stuff we don't have schema translators for, etc.)
+				if (Child.IsA<pxr::UsdGeomXformable>())
+				{
+					Children.Emplace(Child);
+				}
+			}
+
+			TArray<bool> ChildrenCanBeCollapsed;
+			ChildrenCanBeCollapsed.SetNumZeroed(Children.Num());
+
+			const int32 MinBatchSize = 1;
+			ParallelFor(
+				TEXT("RecursiveQueryCanBeCollapsed"),
+				Children.Num(),
+				MinBatchSize,
+				[&](int32 Index)
+				{
+					ChildrenCanBeCollapsed[Index] = RecursiveQueryCanBeCollapsed(Children[Index], Context, Impl, Registry);
+				}
+			);
+
+			for (bool bChildCanBeCollapsed : ChildrenCanBeCollapsed)
+			{
+				if (!bChildCanBeCollapsed)
+				{
+					bCanBeCollapsed = false;
+					break;
+				}
+			}
+		}
+
+		// Record what we found about our main prim
+		{
+			FWriteScopeLock PrimLock{Impl.PrimLocks[MainPrimInfo->PrimLockIndex]};
+			MainPrimInfo->bXformSubtreeCanBeCollapsed = bCanBeCollapsed;
+		}
+
+		// Before we return though, what we can do here is that if we know that we can't be collapsed ourselves,
+		// then none of our ancestors can either! So let's quickly paint upwards to make future queries faster
+		if (!bCanBeCollapsed)
+		{
+			UE::FSdfPath TraversalPath = UsdPrimPath.GetParentPath();
+			while (!TraversalPath.IsAbsoluteRootPath())
+			{
+				if (UE::UsdInfoCache::Private::FUsdPrimInfo* AncestorInfo = Impl.InfoMap.Find(TraversalPath))
+				{
+					FWriteScopeLock PrimLock{Impl.PrimLocks[AncestorInfo->PrimLockIndex]};
+
+					// We found something that was already filled out: Let's stop traversing here
+					if (AncestorInfo->bXformSubtreeCanBeCollapsed.IsSet())
+					{
+						// If we can't collapse ourselves then like we mentioned above none of our ancestors should
+						// be able to collapse either
+						ensure(AncestorInfo->bXformSubtreeCanBeCollapsed.GetValue() == false);
+						break;
+					}
+					else
+					{
+						AncestorInfo->bXformSubtreeCanBeCollapsed = false;
+					}
+				}
+
+				TraversalPath = TraversalPath.GetParentPath();
+			}
+		}
+
+		return bCanBeCollapsed;
+	}
+#endif	  // USE_USD_SDK
+}	 // namespace UE::USDInfoCache::Private
+
 UE::FSdfPath FUsdInfoCache::UnwindToNonCollapsedPath(const UE::FSdfPath& Path, ECollapsingType CollapsingType) const
 {
 	if (FUsdInfoCacheImpl* ImplPtr = Impl.Get())
 	{
 		FReadScopeLock ScopeLock(ImplPtr->InfoMapLock);
 
-		if (const UE::UsdInfoCache::Private::FUsdPrimInfo* FoundInfo = ImplPtr->InfoMap.Find(Path))
+		TOptional<UE::FSdfPath> CollapseRoot;
+		UE::UsdInfoCache::Private::FUsdPrimInfo* MainFoundInfo = ImplPtr->InfoMap.Find(Path);
+		if (MainFoundInfo)
+		{
+			FReadScopeLock PrimLock{ImplPtr->PrimLocks[MainFoundInfo->PrimLockIndex]};
+			CollapseRoot = MainFoundInfo->AssetCollapsedRoot;
+		}
+
+		// We never visited this prim before. We know it's collapsed, let's find our collapse root
+		if (!CollapseRoot.IsSet())
 		{
-			const UE::FSdfPath* CollapsedRoot = CollapsingType == ECollapsingType::Assets ? &FoundInfo->AssetCollapsedRoot
-																						  : &FoundInfo->ComponentCollapsedRoot;
+			TArray<UE::UsdInfoCache::Private::FUsdPrimInfo*> InfosToUpdate = {MainFoundInfo};
 
+			UE::FSdfPath TraversalPath = Path.GetParentPath();
+			while (!TraversalPath.IsAbsoluteRootPath())
+			{
+				if (UE::UsdInfoCache::Private::FUsdPrimInfo* AncestorInfo = ImplPtr->InfoMap.Find(TraversalPath))
+				{
+					FReadScopeLock PrimLock{ImplPtr->PrimLocks[AncestorInfo->PrimLockIndex]};
+
+					// We found an ancestor that has this filled in: We're collapsed, so whatever is its
+					// collapse root is also our collapse root
+					if (AncestorInfo->AssetCollapsedRoot.IsSet())
+					{
+						// If our original Path doesn't have anything filled in, then we *must* be a child of a
+						// collapsed prim (i.e. something that has a non-empty path in its AssetCollapsedRoot)
+						ensure(!AncestorInfo->AssetCollapsedRoot->IsEmpty());
+
+						CollapseRoot = AncestorInfo->AssetCollapsedRoot;
+						break;
+					}
+					// Still nothing. Let's keep track of this info so that we can update it later
+					else
+					{
+						InfosToUpdate.Add(AncestorInfo);
+					}
+				}
+
+				TraversalPath = TraversalPath.GetParentPath();
+			}
+
+			// Fill in all visited infos with what we found on our ancestor
+			if (CollapseRoot.IsSet())
+			{
+				for (UE::UsdInfoCache::Private::FUsdPrimInfo* Info : InfosToUpdate)
+				{
+					FWriteScopeLock PrimLock{ImplPtr->PrimLocks[Info->PrimLockIndex]};
+					Info->AssetCollapsedRoot = CollapseRoot;
+				}
+			}
+		}
+
+		// We have visited this prim during the info cache build (or another Unwind, or just now within this function)
+		if (CollapseRoot.IsSet())
+		{
 			// An empty path here means that we are not collapsed at all
-			if (CollapsedRoot->IsEmpty())
+			if (CollapseRoot->IsEmpty())
 			{
 				return Path;
 			}
 			// Otherwise we have our own path in there (in case we collapse children) or the path to the prim that collapsed us
 			else
 			{
-				return *CollapsedRoot;
+				return CollapseRoot.GetValue();
 			}
 		}
 
 		// This should never happen: We should have cached the entire tree
-		ensureMsgf(false, TEXT("Prim path '%s' has not been cached!"), *Path.GetString());
+		ensureAlwaysMsgf(false, TEXT("Prim path '%s' has not been cached!"), *Path.GetString());
 	}
 
 	return Path;
@@ -311,11 +518,13 @@ TSet<UE::FSdfPath> FUsdInfoCache::GetMainPrims(const UE::FSdfPath& AuxPrimPath)
 {
 	if (FUsdInfoCacheImpl* ImplPtr = Impl.Get())
 	{
-		FReadScopeLock ScopeLock(ImplPtr->AuxiliaryPrimsLock);
+		FReadScopeLock ScopeLock(ImplPtr->InfoMapLock);
 
-		if (const TSet<UE::FSdfPath>* FoundMainPrims = ImplPtr->AuxToMainPrims.Find(AuxPrimPath))
+		if (const UE::UsdInfoCache::Private::FUsdPrimInfo* FoundInfo = ImplPtr->InfoMap.Find(AuxPrimPath))
 		{
-			TSet<UE::FSdfPath> Result = *FoundMainPrims;
+			FReadScopeLock PrimLock{ImplPtr->PrimLocks[FoundInfo->PrimLockIndex]};
+
+			TSet<UE::FSdfPath> Result = FoundInfo->MainPrims;
 			Result.Add(AuxPrimPath);
 			return Result;
 		}
@@ -328,11 +537,13 @@ TSet<UE::FSdfPath> FUsdInfoCache::GetAuxiliaryPrims(const UE::FSdfPath& MainPrim
 {
 	if (FUsdInfoCacheImpl* ImplPtr = Impl.Get())
 	{
-		FReadScopeLock ScopeLock(ImplPtr->AuxiliaryPrimsLock);
+		FReadScopeLock ScopeLock(ImplPtr->InfoMapLock);
 
-		if (const TSet<UE::FSdfPath>* FoundAuxPrims = ImplPtr->MainToAuxPrims.Find(MainPrimPath))
+		if (const UE::UsdInfoCache::Private::FUsdPrimInfo* FoundInfo = ImplPtr->InfoMap.Find(MainPrimPath))
 		{
-			TSet<UE::FSdfPath> Result = *FoundAuxPrims;
+			FReadScopeLock PrimLock{ImplPtr->PrimLocks[FoundInfo->PrimLockIndex]};
+
+			TSet<UE::FSdfPath> Result = FoundInfo->AuxPrims;
 			Result.Add(MainPrimPath);
 			return Result;
 		}
@@ -345,11 +556,13 @@ TSet<UE::FSdfPath> FUsdInfoCache::GetMaterialUsers(const UE::FSdfPath& Path) con
 {
 	if (FUsdInfoCacheImpl* ImplPtr = Impl.Get())
 	{
-		FReadScopeLock ScopeLock(ImplPtr->MaterialUsersLock);
+		FReadScopeLock ScopeLock(ImplPtr->InfoMapLock);
 
-		if (const TSet<UE::FSdfPath>* FoundUsers = ImplPtr->MaterialUsers.Find(Path))
+		if (const UE::UsdInfoCache::Private::FUsdPrimInfo* FoundInfo = ImplPtr->InfoMap.Find(Path))
 		{
-			return *FoundUsers;
+			FReadScopeLock PrimLock{ImplPtr->PrimLocks[FoundInfo->PrimLockIndex]};
+
+			return FoundInfo->MaterialUsers;
 		}
 	}
 
@@ -360,27 +573,33 @@ bool FUsdInfoCache::IsMaterialUsed(const UE::FSdfPath& Path) const
 {
 	if (FUsdInfoCacheImpl* ImplPtr = Impl.Get())
 	{
-		FReadScopeLock ScopeLock(ImplPtr->MaterialUsersLock);
+		FReadScopeLock ScopeLock(ImplPtr->InfoMapLock);
+
+		if (const UE::UsdInfoCache::Private::FUsdPrimInfo* FoundInfo = ImplPtr->InfoMap.Find(Path))
+		{
+			FReadScopeLock PrimLock{ImplPtr->PrimLocks[FoundInfo->PrimLockIndex]};
 
-		return ImplPtr->MaterialUsers.Contains(Path);
+			return FoundInfo->MaterialUsers.Num() > 0;
+		}
 	}
 
 	return false;
 }
 
-namespace UE::USDInfoCacheImpl::Private
+namespace UE::USDInfoCache::Private
 {
 #if USE_USD_SDK
 	void GetPrimVertexCountAndSlots(
 		const pxr::UsdPrim& UsdPrim,
 		const FUsdSchemaTranslationContext& Context,
 		const FUsdInfoCache::FUsdInfoCacheImpl& Impl,
-		const TMap<UE::FSdfPath, TArray<UsdUtils::FUsdPrimMaterialSlot>>& InSubtreeToMaterialSlots,
 		uint64& OutVertexCount,
 		TArray<UsdUtils::FUsdPrimMaterialSlot>& OutMaterialSlots
 	)
 	{
-		TRACE_CPUPROFILER_EVENT_SCOPE(UE::USDInfoCacheImpl::Private::GetPrimVertexCountAndSlots);
+		TRACE_CPUPROFILER_EVENT_SCOPE(GetPrimVertexCountAndSlots);
+
+		FScopedUsdAllocs Allocs;
 
 		if (UsdPrim.IsA<pxr::UsdGeomGprim>() || UsdPrim.IsA<pxr::UsdGeomSubset>())
 		{
@@ -447,13 +666,10 @@ namespace UE::USDInfoCacheImpl::Private
 						// prototype subtrees already
 						if (const UE::UsdInfoCache::Private::FUsdPrimInfo* FoundInfo = Impl.InfoMap.Find(UE::FSdfPath{PrototypePath}))
 						{
-							PrototypeVertexCounts[PrototypeIndex] = FoundInfo->ExpectedVertexCountForSubtree;
-						}
+							FReadScopeLock PrimLock{Impl.PrimLocks[FoundInfo->PrimLockIndex]};
 
-						if (const TArray<UsdUtils::FUsdPrimMaterialSlot>* FoundPrototypeSlots = InSubtreeToMaterialSlots.Find(UE::FSdfPath{
-								PrototypePath}))
-						{
-							OutMaterialSlots.Append(*FoundPrototypeSlots);
+							PrototypeVertexCounts[PrototypeIndex] = FoundInfo->ExpectedVertexCountForSubtree;
+							OutMaterialSlots.Append(FoundInfo->SubtreeMaterialSlots);
 						}
 					}
 				}
@@ -473,20 +689,35 @@ namespace UE::USDInfoCacheImpl::Private
 		}
 	}
 
+	void RepopulateInfoMap(const pxr::UsdPrim& UsdPrim, FUsdInfoCache::FUsdInfoCacheImpl& Impl)
+	{
+		using namespace UE::UsdInfoCache::Private;
+
+		TRACE_CPUPROFILER_EVENT_SCOPE(RepopulateInfoMap);
+
+		FWriteScopeLock ScopeLock(Impl.InfoMapLock);
+
+		uint64 PrimCounter = 0;
+
+		pxr::UsdPrimRange PrimRange{UsdPrim, pxr::UsdTraverseInstanceProxies(pxr::UsdPrimAllPrimsPredicate)};
+		for (pxr::UsdPrim Child : PrimRange)
+		{
+			Impl.InfoMap.Emplace(Child.GetPrimPath(), FUsdPrimInfo{static_cast<int32>(PrimCounter++ % GNumPerPrimLocks)});
+		}
+	}
+
 	void RecursivePropagateVertexAndMaterialSlotCounts(
 		const pxr::UsdPrim& UsdPrim,
 		FUsdSchemaTranslationContext& Context,
 		const pxr::TfToken& MaterialPurposeToken,
 		FUsdInfoCache::FUsdInfoCacheImpl& Impl,
 		FUsdSchemaTranslatorRegistry& Registry,
-		TMap<UE::FSdfPath, TArray<UsdUtils::FUsdPrimMaterialSlot>>& InOutSubtreeToMaterialSlots,
-		TArray<FString>& InOutPointInstancerPaths,
 		uint64& OutSubtreeVertexCount,
 		TArray<UsdUtils::FUsdPrimMaterialSlot>& OutSubtreeSlots,
 		bool bPossibleInheritedBindings
 	)
 	{
-		TRACE_CPUPROFILER_EVENT_SCOPE(UE::USDInfoCacheImpl::Private::RecursivePropagateVertexAndMaterialSlotCounts);
+		TRACE_CPUPROFILER_EVENT_SCOPE(RecursivePropagateVertexAndMaterialSlotCounts);
 
 		if (!UsdPrim)
 		{
@@ -499,6 +730,30 @@ namespace UE::USDInfoCacheImpl::Private
 		UE::FSdfPath PrimPath{UsdPrimPath};
 		pxr::UsdStageRefPtr Stage = UsdPrim.GetStage();
 
+		TFunction<void(const UE::FSdfPath&, TSet<UE::FSdfPath>&)> TryAddMaterialUser = [&Stage](const UE::FSdfPath& Path, TSet<UE::FSdfPath>& Users)
+		{
+			pxr::UsdPrim UserPrim = Stage->GetPrimAtPath(Path);
+
+			if (UserPrim.IsA<pxr::UsdGeomImageable>())
+			{
+				// Do this filtering here because Collection.ComputeIncludedPaths() can be very aggressive and return
+				// literally *all prims* below an included prim path. That's fine and it really does mean that any Mesh prim
+				// in there could use the collection-based material binding, but nevertheless we don't want to register that
+				// e.g. Shader prims or SkelAnimation prims are "material users"
+				Users.Add(Path);
+			}
+			else if (UserPrim.IsA<pxr::UsdGeomSubset>())
+			{
+				// If a UsdGeomSubset is a material user, make its Mesh parent prim into a user too.
+				// Our notice handling is somewhat stricter now, and we have no good way of upgrading a simple material info change
+				// into a resync change of the StaticMeshComponent when we change a material that is bound directly to a
+				// UsdGeomSubset, since the GeomMesh translator doesn't collapse. We'll unwind this path later when fetching material
+				// users, so collapsed static meshes are handled OK, skeletal meshes are handled OK, we just need this one exception
+				// for handling uncollapsed static meshes, because by default Mesh prims don't "collapse" their child UsdGeomSubsets
+				Users.Add(Path.GetParentPath());
+			}
+		};
+
 		// Material bindings are inherited down to child prims, so if we detect a binding on a parent Xform,
 		// we should register the child Mesh prims as users of the material too (regardless of collapsing).
 		// Note that we only consider this for direct bindings: Collection-based bindings will already provide the exhaustive
@@ -508,7 +763,8 @@ namespace UE::USDInfoCacheImpl::Private
 		// Register material users
 		if (!UsdPrim.IsPseudoRoot())
 		{
-			pxr::UsdShadeMaterial ShadeMaterial;
+			TRACE_CPUPROFILER_EVENT_SCOPE(CheckingMaterialUsers);
+
 			TMap<UE::FSdfPath, TSet<UE::FSdfPath>> NewMaterialUsers;
 
 			pxr::UsdShadeMaterialBindingAPI BindingAPI{UsdPrim};
@@ -561,11 +817,12 @@ namespace UE::USDInfoCacheImpl::Private
 
 							if (pxr::UsdCollectionAPI Collection = pxr::UsdCollectionAPI::Get(Stage, *CollectionPath))
 							{
-								std::set<pxr::SdfPath> IncludedPaths = Collection.ComputeIncludedPaths(Collection.ComputeMembershipQuery(), Stage);
+								TSet<UE::FSdfPath>& MaterialUsers = NewMaterialUsers.FindOrAdd(UE::FSdfPath{*MaterialPath});
 
+								std::set<pxr::SdfPath> IncludedPaths = Collection.ComputeIncludedPaths(Collection.ComputeMembershipQuery(), Stage);
 								for (const pxr::SdfPath& IncludedPath : IncludedPaths)
 								{
-									NewMaterialUsers.FindOrAdd(UE::FSdfPath{*MaterialPath}).Add(UE::FSdfPath{IncludedPath});
+									TryAddMaterialUser(UE::FSdfPath{IncludedPath}, MaterialUsers);
 								}
 							}
 							else
@@ -583,11 +840,12 @@ namespace UE::USDInfoCacheImpl::Private
 				}
 
 				// Check for material bindings directly for this prim
-				ShadeMaterial = BindingAPI.ComputeBoundMaterial(MaterialPurposeToken);
-				if (ShadeMaterial)
+				if (pxr::UsdShadeMaterial ShadeMaterial = BindingAPI.ComputeBoundMaterial(MaterialPurposeToken))
 				{
 					bPrimHasInheritableMaterialBindings = true;
-					NewMaterialUsers.FindOrAdd(UE::FSdfPath{ShadeMaterial.GetPrim().GetPath()}).Add(UE::FSdfPath{UsdPrimPath});
+
+					TSet<UE::FSdfPath>& MaterialUsers = NewMaterialUsers.FindOrAdd(UE::FSdfPath{ShadeMaterial.GetPrim().GetPath()});
+					TryAddMaterialUser(UE::FSdfPath{UsdPrimPath}, MaterialUsers);
 				}
 			}
 			// Temporary fallback for prims that don't have the MaterialBindingAPI but do have the relationship.
@@ -601,67 +859,36 @@ namespace UE::USDInfoCacheImpl::Private
 				{
 					const pxr::SdfPath& TargetMaterialPrimPath = Targets[0];
 					pxr::UsdPrim MaterialPrim = Stage->GetPrimAtPath(TargetMaterialPrimPath);
-					ShadeMaterial = pxr::UsdShadeMaterial{MaterialPrim};
-					if (ShadeMaterial)
+					if (pxr::UsdShadeMaterial ShadeMaterial = pxr::UsdShadeMaterial{MaterialPrim})
 					{
 						bPrimHasInheritableMaterialBindings = true;
-						NewMaterialUsers.FindOrAdd(UE::FSdfPath{TargetMaterialPrimPath}).Add(UE::FSdfPath{UsdPrimPath});
+
+						TSet<UE::FSdfPath>& MaterialUsers = NewMaterialUsers.FindOrAdd(UE::FSdfPath{TargetMaterialPrimPath});
+						TryAddMaterialUser(UE::FSdfPath{UsdPrimPath}, MaterialUsers);
 					}
 				}
 			}
 
-			FWriteScopeLock ScopeLock(Impl.MaterialUsersLock);
 			for (const TPair<UE::FSdfPath, TSet<UE::FSdfPath>>& NewMaterialToUsers : NewMaterialUsers)
 			{
-				TSet<UE::FSdfPath>& UserPrimPaths = Impl.MaterialUsers.FindOrAdd(NewMaterialToUsers.Key);
-				UserPrimPaths.Reserve(UserPrimPaths.Num() + NewMaterialToUsers.Value.Num());
-				for (const UE::FSdfPath& NewUserPath : NewMaterialToUsers.Value)
+				FReadScopeLock ScopeLock(Impl.InfoMapLock);
+				if (UsdInfoCache::Private::FUsdPrimInfo* Info = Impl.InfoMap.Find(NewMaterialToUsers.Key))
 				{
-					pxr::UsdPrim UserPrim = Stage->GetPrimAtPath(NewUserPath);
-
-					// Do this filtering here because Collection.ComputeIncludedPaths() can be very aggressive and return
-					// literally *all prims* below an included prim path. That's fine and it really does mean that any Mesh prim
-					// in there could use the collection-based material binding, but nevertheless we don't want to register that
-					// e.g. Shader prims or SkelAnimation prims are "material users"
-					if (UserPrim.IsA<pxr::UsdGeomImageable>())
-					{
-						UE_LOG(
-							LogUsd,
-							Verbose,
-							TEXT("Registering material user '%s' of material '%s'"),
-							*NewUserPath.GetString(),
-							*NewMaterialToUsers.Key.GetString()
-						);
-						UserPrimPaths.Add(NewUserPath);
-					}
-					// If a UsdGeomSubset is a material user, make its Mesh parent prim into a user too.
-					// Our notice handling is somewhat stricter now, and we have no good way of upgrading a simple material info change
-					// into a resync change of the StaticMeshComponent when we change a material that is bound directly to a
-					// UsdGeomSubset, since the GeomMesh translator doesn't collapse. We'll unwind this path later when fetching material
-					// users, so collapsed static meshes are handled OK, skeletal meshes are handled OK, we just need this one exception
-					// for handling uncollapsed static meshes, because by default Mesh prims don't "collapse" their child UsdGeomSubsets
-					else if (UserPrim.IsA<pxr::UsdGeomSubset>())
-					{
-						UE_LOG(
-							LogUsd,
-							Verbose,
-							TEXT("Registering parent Mesh prim '%s' of UsdGeomSubset '%s' as a material user of material prim '%s'"),
-							*NewUserPath.GetString(),
-							*PrimPath.GetString(),
-							*NewMaterialToUsers.Key.GetString()
-						);
-						UserPrimPaths.Add(NewUserPath.GetParentPath());
-					}
+					FWriteScopeLock PrimLock{Impl.PrimLocks[Info->PrimLockIndex]};
+					Info->MaterialUsers.Append(NewMaterialToUsers.Value);
 				}
 			}
 		}
 
-		pxr::UsdPrimSiblingRange PrimChildren = UsdPrim.GetFilteredChildren(pxr::UsdTraverseInstanceProxies(pxr::UsdPrimAllPrimsPredicate));
-
 		TArray<pxr::UsdPrim> Prims;
-		for (pxr::UsdPrim Child : PrimChildren)
 		{
-			Prims.Emplace(Child);
+			TRACE_CPUPROFILER_EVENT_SCOPE(CollectingChildren);
+			pxr::UsdPrimSiblingRange PrimChildren = UsdPrim.GetFilteredChildren(pxr::UsdTraverseInstanceProxies(pxr::UsdPrimAllPrimsPredicate));
+
+			for (pxr::UsdPrim Child : PrimChildren)
+			{
+				Prims.Emplace(Child);
+			}
 		}
 
 		const uint32 NumChildren = Prims.Num();
@@ -673,7 +900,6 @@ namespace UE::USDInfoCacheImpl::Private
 		ChildSubtreeMaterialSlots.SetNum(NumChildren);
 
 		const int32 MinBatchSize = 1;
-
 		ParallelFor(
 			TEXT("RecursivePropagateVertexAndMaterialSlotCounts"),
 			Prims.Num(),
@@ -686,8 +912,6 @@ namespace UE::USDInfoCacheImpl::Private
 					MaterialPurposeToken,
 					Impl,
 					Registry,
-					InOutSubtreeToMaterialSlots,
-					InOutPointInstancerPaths,
 					ChildSubtreeVertexCounts[Index],
 					ChildSubtreeMaterialSlots[Index],
 					bPrimHasInheritableMaterialBindings || bPossibleInheritedBindings
@@ -701,60 +925,68 @@ namespace UE::USDInfoCacheImpl::Private
 		// We will still step into invisible prims to collect all info we can, but we won't count their material slots
 		// or vertex counts: The main usage of those counts is to handle collapsed meshes, and during collapse we just
 		// early out whenever we encounter an invisible prim
-		bool bPrimIsInvisible = false;
-		if (pxr::UsdGeomImageable UsdGeomImageable = pxr::UsdGeomImageable(UsdPrim))
+		bool bIsPointInstancer = false;
 		{
-			if (pxr::UsdAttribute VisibilityAttr = UsdGeomImageable.GetVisibilityAttr())
+			TRACE_CPUPROFILER_EVENT_SCOPE(GettingVertexCountAndSlots);
+
+			bool bPrimIsInvisible = false;
+			if (pxr::UsdGeomImageable UsdGeomImageable = pxr::UsdGeomImageable(UsdPrim))
 			{
-				pxr::TfToken VisibilityToken;
-				if (VisibilityAttr.Get(&VisibilityToken) && VisibilityToken == pxr::UsdGeomTokens->invisible)
+				if (pxr::UsdAttribute VisibilityAttr = UsdGeomImageable.GetVisibilityAttr())
 				{
-					bPrimIsInvisible = true;
+					pxr::TfToken VisibilityToken;
+					if (VisibilityAttr.Get(&VisibilityToken) && VisibilityToken == pxr::UsdGeomTokens->invisible)
+					{
+						bPrimIsInvisible = true;
+					}
 				}
 			}
-		}
-
-		// If the mesh prim has an unselected geometry purpose, it is also essentially invisible
-		if (!EnumHasAllFlags(Context.PurposesToLoad, IUsdPrim::GetPurpose(UsdPrim)))
-		{
-			bPrimIsInvisible = true;
-		}
 
-		bool bIsPointInstancer = false;
-		if (pxr::UsdGeomPointInstancer PointInstancer{UsdPrim})
-		{
-			bIsPointInstancer = true;
-		}
-		else if (!bPrimIsInvisible)
-		{
-			GetPrimVertexCountAndSlots(UsdPrim, Context, Impl, InOutSubtreeToMaterialSlots, OutSubtreeVertexCount, OutSubtreeSlots);
+			// If the mesh prim has an unselected geometry purpose, it is also essentially invisible
+			if (!EnumHasAllFlags(Context.PurposesToLoad, IUsdPrim::GetPurpose(UsdPrim)))
+			{
+				bPrimIsInvisible = true;
+			}
 
-			for (uint32 ChildIndex = 0; ChildIndex < NumChildren; ++ChildIndex)
+			if (pxr::UsdGeomPointInstancer PointInstancer{UsdPrim})
+			{
+				bIsPointInstancer = true;
+			}
+			else if (!bPrimIsInvisible)
 			{
-				OutSubtreeVertexCount += ChildSubtreeVertexCounts[ChildIndex];
-				OutSubtreeSlots.Append(ChildSubtreeMaterialSlots[ChildIndex]);
+				GetPrimVertexCountAndSlots(UsdPrim, Context, Impl, OutSubtreeVertexCount, OutSubtreeSlots);
+
+				for (uint32 ChildIndex = 0; ChildIndex < NumChildren; ++ChildIndex)
+				{
+					OutSubtreeVertexCount += ChildSubtreeVertexCounts[ChildIndex];
+					OutSubtreeSlots.Append(ChildSubtreeMaterialSlots[ChildIndex]);
+				}
 			}
 		}
 
 		{
-			FWriteScopeLock ScopeLock(Impl.InfoMapLock);
-
-			UE::UsdInfoCache::Private::FUsdPrimInfo& Info = Impl.InfoMap.FindOrAdd(UE::FSdfPath{UsdPrimPath});
+			TRACE_CPUPROFILER_EVENT_SCOPE(StoringCounts);
 
-			// For point instancers we can't guarantee we parsed the prototypes yet because they
-			// could technically be anywhere, so store them here for a later pass
-			if (bIsPointInstancer)
-			{
-				InOutPointInstancerPaths.Emplace(UE::FSdfPath{UsdPrimPath}.GetString());
-			}
-			// While we will compute the totals for any and all children normally, don't just append the regular
-			// traversal vertex count to the point instancer prim itself just yet, as that doesn't really represent
-			// what will happen. We'll later do another pass to handle point instancers where we'll properly instance
-			// stuff, and then we'll updadte all ancestors
-			else
+			FReadScopeLock ScopeLock(Impl.InfoMapLock);
+			if (UE::UsdInfoCache::Private::FUsdPrimInfo* Info = Impl.InfoMap.Find(UE::FSdfPath{UsdPrimPath}))
 			{
-				Info.ExpectedVertexCountForSubtree = OutSubtreeVertexCount;
-				InOutSubtreeToMaterialSlots.Emplace(UsdPrimPath, OutSubtreeSlots);
+				// For point instancers we can't guarantee we parsed the prototypes yet because they
+				// could technically be anywhere, so store them here for a later pass
+				if (bIsPointInstancer)
+				{
+					FWriteScopeLock PointInstancerLock(Impl.TempPointInstancerPathsLock);
+					Impl.TempPointInstancerPaths.Emplace(UE::FSdfPath{UsdPrimPath}.GetString());
+				}
+				// While we will compute the totals for any and all children normally, don't just append the regular
+				// traversal vertex count to the point instancer prim itself just yet, as that doesn't really represent
+				// what will happen. We'll later do another pass to handle point instancers where we'll properly instance
+				// stuff, and then we'll updadte all ancestors
+				else
+				{
+					FWriteScopeLock PrimLock{Impl.PrimLocks[Info->PrimLockIndex]};
+					Info->ExpectedVertexCountForSubtree = OutSubtreeVertexCount;
+					Info->SubtreeMaterialSlots.Append(OutSubtreeSlots);
+				}
 			}
 		}
 	}
@@ -768,14 +1000,17 @@ namespace UE::USDInfoCacheImpl::Private
 	 * of prim subtrees all over to build the final counts of point instancers that use them as prototypes, and
 	 * then update their parents.
 	 */
-	void UpdateInfoForPointInstancers(
-		const pxr::UsdStageRefPtr& Stage,
-		const FUsdSchemaTranslationContext& Context,
-		FUsdInfoCache::FUsdInfoCacheImpl& Impl,
-		TArray<FString>& PointInstancerPaths,
-		TMap<UE::FSdfPath, TArray<UsdUtils::FUsdPrimMaterialSlot>>& InOutSubtreeMaterialSlots
-	)
+	void UpdateInfoForPointInstancers(const FUsdSchemaTranslationContext& Context, FUsdInfoCache::FUsdInfoCacheImpl& Impl)
 	{
+		TRACE_CPUPROFILER_EVENT_SCOPE(UpdateInfoForPointInstancers);
+
+		pxr::UsdStageWeakPtr UsdStageWeak = Impl.TempStage;
+		pxr::UsdStageRefPtr Stage = UsdStageWeak;
+		if (!Stage)
+		{
+			return;
+		}
+
 		// We must sort point instancers in a particular order in case they depend on each other.
 		// At least we know that an ordering like this should be possible, because A with B as a prototype and B with A
 		// as a prototype leads to an invalid USD stage.
@@ -815,9 +1050,13 @@ namespace UE::USDInfoCacheImpl::Private
 
 			return LHS < RHS;
 		};
-		PointInstancerPaths.Sort(SortFunction);
+		{
+			FWriteScopeLock PointInstancerLock{Impl.TempPointInstancerPathsLock};
+			Impl.TempPointInstancerPaths.Sort(SortFunction);
+		}
 
-		for (const FString& PointInstancerPath : PointInstancerPaths)
+		FReadScopeLock PointInstancerLock{Impl.TempPointInstancerPathsLock};
+		for (const FString& PointInstancerPath : Impl.TempPointInstancerPaths)
 		{
 			UE::FSdfPath UsdPointInstancerPath{*PointInstancerPath};
 
@@ -826,49 +1065,48 @@ namespace UE::USDInfoCacheImpl::Private
 				uint64 PointInstancerVertexCount = 0;
 				TArray<UsdUtils::FUsdPrimMaterialSlot> PointInstancerMaterialSlots;
 
-				GetPrimVertexCountAndSlots(
-					PointInstancer,
-					Context,
-					Impl,
-					InOutSubtreeMaterialSlots,
-					PointInstancerVertexCount,
-					PointInstancerMaterialSlots
-				);
-
-				FWriteScopeLock Lock{Impl.InfoMapLock};
-				UE::UsdInfoCache::Private::FUsdPrimInfo& Info = Impl.InfoMap.FindOrAdd(UsdPointInstancerPath);
-
-				Info.ExpectedVertexCountForSubtree = PointInstancerVertexCount;
-				InOutSubtreeMaterialSlots.Emplace(UsdPointInstancerPath, PointInstancerMaterialSlots);
+				GetPrimVertexCountAndSlots(PointInstancer, Context, Impl, PointInstancerVertexCount, PointInstancerMaterialSlots);
 
-				// Now that we have info on the point instancer itself, update the counts of all ancestors.
-				// Note: The vertex/material slot count for the entire point instancer subtree are just the counts
-				// for the point instancer itself, as we stop regular traversal when we hit them
-				UE::FSdfPath ParentPath = UsdPointInstancerPath.GetParentPath();
-				pxr::UsdPrim Prim = Stage->GetPrimAtPath(ParentPath);
-				while (Prim)
+				FReadScopeLock Lock{Impl.InfoMapLock};
+				if (UE::UsdInfoCache::Private::FUsdPrimInfo* Info = Impl.InfoMap.Find(UsdPointInstancerPath))
 				{
-					// If our ancestor is a point instancer itself, just abort as we'll only get the actual counts
-					// when we handle that ancestor directly. We don't want to update the ancestor point instancer's
-					// ancestors with incorrect values
-					if (Prim.IsA<pxr::UsdGeomPointInstancer>())
 					{
-						break;
+						FWriteScopeLock PrimLock{Impl.PrimLocks[Info->PrimLockIndex]};
+						Info->ExpectedVertexCountForSubtree = PointInstancerVertexCount;
+						Info->SubtreeMaterialSlots.Append(PointInstancerMaterialSlots);
 					}
 
-					UE::UsdInfoCache::Private::FUsdPrimInfo& ParentInfo = Impl.InfoMap.FindOrAdd(ParentPath);
-					ParentInfo.ExpectedVertexCountForSubtree += PointInstancerVertexCount;
+					// Now that we have info on the point instancer itself, update the counts of all ancestors.
+					// Note: The vertex/material slot count for the entire point instancer subtree are just the counts
+					// for the point instancer itself, as we stop regular traversal when we hit them
+					UE::FSdfPath ParentPath = UsdPointInstancerPath.GetParentPath();
+					pxr::UsdPrim Prim = Stage->GetPrimAtPath(ParentPath);
+					while (Prim)
+					{
+						// If our ancestor is a point instancer itself, just abort as we'll only get the actual counts
+						// when we handle that ancestor directly. We don't want to update the ancestor point instancer's
+						// ancestors with incorrect values
+						if (Prim.IsA<pxr::UsdGeomPointInstancer>())
+						{
+							break;
+						}
 
-					InOutSubtreeMaterialSlots[ParentPath].Append(PointInstancerMaterialSlots);
+						if (UE::UsdInfoCache::Private::FUsdPrimInfo* ParentInfo = Impl.InfoMap.Find(ParentPath))
+						{
+							FWriteScopeLock ParentPrimLock{Impl.PrimLocks[ParentInfo->PrimLockIndex]};
+							ParentInfo->ExpectedVertexCountForSubtree += PointInstancerVertexCount;
+							ParentInfo->SubtreeMaterialSlots.Append(PointInstancerMaterialSlots);
+						}
 
-					// Break only here so we update the pseudoroot too
-					if (Prim.IsPseudoRoot())
-					{
-						break;
-					}
+						// Break only here so we update the pseudoroot too
+						if (Prim.IsPseudoRoot())
+						{
+							break;
+						}
 
-					ParentPath = ParentPath.GetParentPath();
-					Prim = Stage->GetPrimAtPath(ParentPath);
+						ParentPath = ParentPath.GetParentPath();
+						Prim = Stage->GetPrimAtPath(ParentPath);
+					}
 				}
 			}
 		}
@@ -882,41 +1120,42 @@ namespace UE::USDInfoCacheImpl::Private
 	 * the main recursive pass just adds them to arrays, and we're allowed to handle bMergeIdenticalSlots
 	 * only here.
 	 */
-	void CollectMaterialSlotCounts(
-		FUsdInfoCache::FUsdInfoCacheImpl& Impl,
-		const TMap<UE::FSdfPath, TArray<UsdUtils::FUsdPrimMaterialSlot>>& SubtreeMaterialSlots,
-		bool bContextMergeIdenticalSlots
-	)
+	void CollectMaterialSlotCounts(FUsdInfoCache::FUsdInfoCacheImpl& Impl, bool bContextMergeIdenticalSlots)
 	{
-		FWriteScopeLock Lock{Impl.InfoMapLock};
+		TRACE_CPUPROFILER_EVENT_SCOPE(CollectMaterialSlotCounts);
 
-		for (const TPair<UE::FSdfPath, TArray<UsdUtils::FUsdPrimMaterialSlot>>& Pair : SubtreeMaterialSlots)
+		if (!bContextMergeIdenticalSlots)
+		{
+			return;
+		}
+
+		FReadScopeLock Lock{Impl.InfoMapLock};
+		for (TPair<UE::FSdfPath, UE::UsdInfoCache::Private::FUsdPrimInfo>& Pair : Impl.InfoMap)
 		{
 			const UE::FSdfPath& PrimPath = Pair.Key;
+			UE::UsdInfoCache::Private::FUsdPrimInfo& Info = Pair.Value;
 
 			bool bCanMergeSlotsForThisPrim = false;
-			if (const UE::UsdInfoCache::Private::FUsdPrimInfo* FoundInfo = Impl.InfoMap.Find(PrimPath))
+
+			// Check whether we merge slots for this prim or not
 			{
-				const UE::FSdfPath* CollapsedRoot = &FoundInfo->AssetCollapsedRoot;
+				FReadScopeLock ReadPrimLock{Impl.PrimLocks[Info.PrimLockIndex]};
 
 				// We only merge slots in the context of collapsing
-				const bool bPrimIsCollapsedOrCollapseRoot = !CollapsedRoot->IsEmpty() || PrimPath.IsAbsoluteRootPath();
+				const TOptional<UE::FSdfPath>& CollapsedRoot = Info.AssetCollapsedRoot;
+				const bool bPrimIsCollapsedOrCollapseRoot = !CollapsedRoot.IsSet() || (!CollapsedRoot->IsEmpty() || PrimPath.IsAbsoluteRootPath());
 
-				const bool bPrimIsPotentialGeometryCacheRoot = FoundInfo->GeometryCacheState
+				const bool bPrimIsPotentialGeometryCacheRoot = Info.GeometryCacheState
 															   == UE::UsdInfoCache::Private::EGeometryCachePrimState::ValidRoot;
 				bCanMergeSlotsForThisPrim = bPrimIsCollapsedOrCollapseRoot && !bPrimIsPotentialGeometryCacheRoot;
 			}
 
-			UE::UsdInfoCache::Private::FUsdPrimInfo& Info = Impl.InfoMap.FindOrAdd(Pair.Key);
-
+			// Actually update the slot count
 			// For now we only ever merge material slots when collapsing
-			if (bCanMergeSlotsForThisPrim && bContextMergeIdenticalSlots)
+			if (bCanMergeSlotsForThisPrim)
 			{
-				Info.SubtreeMaterialSlots = TSet<UsdUtils::FUsdPrimMaterialSlot>{Pair.Value}.Array();
-			}
-			else
-			{
-				Info.SubtreeMaterialSlots = Pair.Value;
+				FWriteScopeLock WritePrimLock{Impl.PrimLocks[Info.PrimLockIndex]};
+				Info.SubtreeMaterialSlots = TSet<UsdUtils::FUsdPrimMaterialSlot>{Info.SubtreeMaterialSlots}.Array();
 			}
 		}
 	}
@@ -928,6 +1167,8 @@ namespace UE::USDInfoCacheImpl::Private
 		const TSharedPtr<FUsdSchemaTranslator>& Translator
 	)
 	{
+		TRACE_CPUPROFILER_EVENT_SCOPE(CanMeshSubtreeBeCollapsed);
+
 		if (!UsdPrim)
 		{
 			return false;
@@ -941,12 +1182,15 @@ namespace UE::USDInfoCacheImpl::Private
 
 		pxr::SdfPath UsdPrimPath = UsdPrim.GetPrimPath();
 
-		FWriteScopeLock ScopeLock(Impl.InfoMapLock);
-		UE::UsdInfoCache::Private::FUsdPrimInfo& Info = Impl.InfoMap.FindOrAdd(UE::FSdfPath{UsdPrimPath});
-
-		if (Info.ExpectedVertexCountForSubtree > GMaxNumVerticesCollapsedMesh)
+		FReadScopeLock ScopeLock(Impl.InfoMapLock);
+		if (UE::UsdInfoCache::Private::FUsdPrimInfo* Info = Impl.InfoMap.Find(UE::FSdfPath{UsdPrimPath}))
 		{
-			return false;
+			FReadScopeLock PrimLock{Impl.PrimLocks[Info->PrimLockIndex]};
+
+			if (Info->ExpectedVertexCountForSubtree > GMaxNumVerticesCollapsedMesh)
+			{
+				return false;
+			}
 		}
 
 		return true;
@@ -956,103 +1200,84 @@ namespace UE::USDInfoCacheImpl::Private
 		const pxr::UsdPrim& UsdPrim,
 		FUsdSchemaTranslationContext& Context,
 		FUsdInfoCache::FUsdInfoCacheImpl& Impl,
-		FUsdSchemaTranslatorRegistry& Registry,
-		const pxr::SdfPath& AssetCollapsedRoot = pxr::SdfPath::EmptyPath(),
-		const pxr::SdfPath& ComponentCollapsedRoot = pxr::SdfPath::EmptyPath()
+		FUsdSchemaTranslatorRegistry& Registry
 	)
 	{
-		TRACE_CPUPROFILER_EVENT_SCOPE(UE::USDInfoCacheImpl::Private::RecursiveQueryCollapsesChildren);
+		TRACE_CPUPROFILER_EVENT_SCOPE(UE::USDInfoCache::Private::RecursiveQueryCollapsesChildren);
 		FScopedUsdAllocs Allocs;
 
 		pxr::SdfPath UsdPrimPath = UsdPrim.GetPrimPath();
 		UE::FSdfPath PrimPath{UsdPrimPath};
 
-		// Prevents allocation by referencing instead of copying
-		const pxr::SdfPath* AssetCollapsedRootOverride = &AssetCollapsedRoot;
-		const pxr::SdfPath* ComponentCollapsedRootOverride = &ComponentCollapsedRoot;
-
-		bool bIsAssetCollapsed = !AssetCollapsedRoot.IsEmpty();
-		bool bIsComponentCollapsed = !ComponentCollapsedRoot.IsEmpty();
+		bool bCollapsesChildren = false;
 
 		TSharedPtr<FUsdSchemaTranslator> SchemaTranslator = Registry.CreateTranslatorForSchema(Context.AsShared(), UE::FUsdTyped(UsdPrim));
 		if (SchemaTranslator)
 		{
-			if (!bIsAssetCollapsed || !bIsComponentCollapsed)
+			bool bIsPotentialGeometryCacheRoot = false;
 			{
-				const bool bCanMeshSubtreeBeCollapsed = CanMeshSubtreeBeCollapsed(UsdPrim, Context, Impl, SchemaTranslator);
-
-				const UE::UsdInfoCache::Private::FUsdPrimInfo& Info = Impl.InfoMap.FindChecked(UE::FSdfPath{UsdPrimPath});
-				const bool bIsPotentialGeometryCacheRoot = Info.GeometryCacheState == UE::UsdInfoCache::Private::EGeometryCachePrimState::ValidRoot;
-
-				if (!bIsAssetCollapsed)
+				FReadScopeLock ScopeLock{Impl.InfoMapLock};
+				if (const UE::UsdInfoCache::Private::FUsdPrimInfo* Info = Impl.InfoMap.Find(UE::FSdfPath{UsdPrimPath}))
 				{
-					// The potential geometry cache root is checked first since the FUsdGeometryCacheTranslator::CollapsesChildren has no logic of its
-					// own
-					if (bIsPotentialGeometryCacheRoot || (SchemaTranslator->CollapsesChildren(ECollapsingType::Assets) && bCanMeshSubtreeBeCollapsed))
-					{
-						AssetCollapsedRootOverride = &UsdPrimPath;
-					}
+					bIsPotentialGeometryCacheRoot = Info->GeometryCacheState == UE::UsdInfoCache::Private::EGeometryCachePrimState::ValidRoot;
 				}
+			}
 
-				if (!bIsComponentCollapsed)
-				{
-					if (bIsPotentialGeometryCacheRoot
-						|| (SchemaTranslator->CollapsesChildren(ECollapsingType::Components) && bCanMeshSubtreeBeCollapsed))
-					{
-						ComponentCollapsedRootOverride = &UsdPrimPath;
-					}
-				}
+			// The potential geometry cache root is checked first since the FUsdGeometryCacheTranslator::CollapsesChildren
+			// has no logic of its own
+			if (bIsPotentialGeometryCacheRoot
+				|| (SchemaTranslator->CollapsesChildren(ECollapsingType::Assets)
+					&& CanMeshSubtreeBeCollapsed(UsdPrim, Context, Impl, SchemaTranslator)))
+			{
+				bCollapsesChildren = true;
 			}
 		}
 
-		pxr::UsdPrimSiblingRange PrimChildren = UsdPrim.GetFilteredChildren(pxr::UsdTraverseInstanceProxies(pxr::UsdPrimAllPrimsPredicate));
-
-		TArray<pxr::UsdPrim> Prims;
-		for (pxr::UsdPrim Child : PrimChildren)
+		// We only need to visit our children if we don't collapse. We'll leave the AssetCollapsedRoot fields unset on
+		// the InfoMap, and whenever we query info about a particular prim will fill that in on-demand by just traveling
+		// upwards until we run into our collapse root
+		if (!bCollapsesChildren)
 		{
-			Prims.Emplace(Child);
-		}
-
-		const int32 MinBatchSize = 1;
-
-		ParallelFor(
-			TEXT("RecursiveQueryCollapsesChildren"),
-			Prims.Num(),
-			MinBatchSize,
-			[&](int32 Index)
+			TArray<pxr::UsdPrim> Prims;
+			for (pxr::UsdPrim Child : UsdPrim.GetFilteredChildren(pxr::UsdTraverseInstanceProxies(pxr::UsdPrimAllPrimsPredicate)))
 			{
-				RecursiveQueryCollapsesChildren(Prims[Index], Context, Impl, Registry, *AssetCollapsedRootOverride, *ComponentCollapsedRootOverride);
+				Prims.Emplace(Child);
 			}
-		);
 
-		{
-			FWriteScopeLock ScopeLock(Impl.InfoMapLock);
-
-			UE::UsdInfoCache::Private::FUsdPrimInfo& Info = Impl.InfoMap.FindOrAdd(UE::FSdfPath{UsdPrimPath});
+			const int32 MinBatchSize = 1;
+			ParallelFor(
+				TEXT("RecursiveQueryCollapsesChildren"),
+				Prims.Num(),
+				MinBatchSize,
+				[&](int32 Index)
+				{
+					RecursiveQueryCollapsesChildren(Prims[Index], Context, Impl, Registry);
+				}
+			);
+		}
 
-			// These paths will be still empty in case nothing has collapsed yet, hold UsdPrimPath in case UsdPrim
-			// collapses that type, or hold the path to the collapsed root passed in via our caller, in case we're
-			// collapsed
-			Info.AssetCollapsedRoot = *AssetCollapsedRootOverride;
-			Info.ComponentCollapsedRoot = *ComponentCollapsedRootOverride;
+		// Record our collapse root
+		{
+			FReadScopeLock ScopeLock(Impl.InfoMapLock);
+			if (UE::UsdInfoCache::Private::FUsdPrimInfo* Info = Impl.InfoMap.Find(UE::FSdfPath{UsdPrimPath}))
+			{
+				FWriteScopeLock PrimLock{Impl.PrimLocks[Info->PrimLockIndex]};
+				Info->AssetCollapsedRoot.Emplace(bCollapsesChildren ? UsdPrimPath : pxr::SdfPath::EmptyPath());
+			}
 		}
 
 		// This really should be a separate pass, but it does no harm here and we have so many passes already...
 		// This needs to happen after we set things into the FUsdPrimInfo for this prim right above this, as it may query
 		// whether this prim (or any of its children) collapse
+		//
+		// We only do this for uncollapsed prims or collapse roots. This because whenever the collapse root
+		// registers its auxiliary prims here, it will already account for all of the relevant child prims in the entire subtree,
+		// according to the translator type. The links between prims inside of a collapsed subtree aren't really useful,
+		// because if anything inside the collapsed subtree updates, we'll always just need to update from the collapsed
+		// root anyway
 		if (SchemaTranslator)
 		{
-			const bool bIsCollapsedBySomeParent = AssetCollapsedRootOverride && !AssetCollapsedRootOverride->IsEmpty()
-												  && *AssetCollapsedRootOverride != UsdPrimPath;
-
-			// We don't care about prims that were collapsed by another. This because whenever the collapse root
-			// registers its auxiliary prims it will already account for all of the collapsed prims that are relevant,
-			// according to the translator type. If we registered aux prims for all prims here, we'd get useless aux
-			// prim links between e.g. all parent and child prims *within* the collapsed subtree
-			if (!bIsCollapsedBySomeParent)
-			{
-				Impl.RegisterAuxiliaryPrims(PrimPath, SchemaTranslator->CollectAuxiliaryPrims());
-			}
+			Impl.RegisterAuxiliaryPrims(PrimPath, SchemaTranslator->CollectAuxiliaryPrims());
 		}
 	}
 
@@ -1095,63 +1320,113 @@ namespace UE::USDInfoCacheImpl::Private
 		return Result;
 	}
 
-	void RegisterInstanceableAuxPrims(const pxr::UsdPrim& UsdPrim, FUsdSchemaTranslationContext& Context, FUsdInfoCache::FUsdInfoCacheImpl& Impl)
+	void RegisterInstanceableAuxPrims(FUsdSchemaTranslationContext& Context, FUsdInfoCache::FUsdInfoCacheImpl& Impl)
 	{
-		TRACE_CPUPROFILER_EVENT_SCOPE(UE::USDInfoCacheImpl::Private::RegisterInstanceableAuxPrims);
+		TRACE_CPUPROFILER_EVENT_SCOPE(UE::USDInfoCache::Private::RegisterInstanceableAuxPrims);
 		FScopedUsdAllocs Allocs;
 
-		pxr::UsdStageRefPtr Stage = UsdPrim.GetStage();
-		for (const pxr::UsdPrim& Prototype : Stage->GetPrototypes())
+		pxr::UsdStageWeakPtr UsdStageWeak = Impl.TempStage;
+		pxr::UsdStageRefPtr Stage = UsdStageWeak;
+		if (!Stage)
 		{
-			TRACE_CPUPROFILER_EVENT_SCOPE(RegisterInstanceableAuxPrims::Prototype);
+			return;
+		}
 
-			if (!Prototype)
-			{
-				continue;
-			}
+		FReadScopeLock ScopeLock{Impl.InfoMapLock};
 
-			// Step into every instance of this prototype on the stage
-			for (const pxr::UsdPrim& Instance : Prototype.GetInstances())
+		std::vector<pxr::UsdPrim> Prototypes = Stage->GetPrototypes();
+		const int32 MinBatchSize = 1;
+		ParallelFor(
+			TEXT("RegisterInstanceableAuxPrimsPrototypes"),
+			Prototypes.size(),
+			MinBatchSize,
+			[&Prototypes, &Context, &Impl, &Stage](int32 Index)
 			{
-				TRACE_CPUPROFILER_EVENT_SCOPE(RegisterInstanceableAuxPrims::PrototypeInstance);
+				TRACE_CPUPROFILER_EVENT_SCOPE(RegisterInstanceableAuxPrims::Prototype);
 
-				UE::FSdfPath InstancePath{Instance.GetPrimPath()};
+				FScopedUsdAllocs Allocs;
 
-				// Adding a dependency on the prototype directly is interesting, even though we currently don't display those.
-				// and the prototype paths are mostly transient
-				Impl.RegisterAuxiliaryPrims(InstancePath, {UE::FSdfPath{Prototype.GetPrimPath()}});
+				const pxr::UsdPrim& Prototype = Prototypes[Index];
+				if (!Prototype)
+				{
+					return;
+				}
+
+				std::vector<pxr::UsdPrim> Instances = Prototype.GetInstances();
+				if (Instances.size() < 1)
+				{
+					return;
+				}
+
+				TArray<UE::FSdfPath> InstancePaths;
+				InstancePaths.SetNum(Instances.size());
 
 				// Really what we want is to find the source prim that generated this prototype though. Instances always work
 				// through some kind of composition arc, so here we collect all references/payloads/inherits/specializes/etc.
-				TSet<UE::FSdfPath> SourcePaths = GetLocalNonRootCompositionArcSourcePaths(Instance);
-				Impl.RegisterAuxiliaryPrims(InstancePath, SourcePaths);
-
-				// Here we'll traverse the entire subtree of the instance
-				pxr::UsdPrimRange PrimRange(Instance, pxr::UsdTraverseInstanceProxies());
-				for (pxr::UsdPrimRange::iterator InstanceChildIt = ++PrimRange.begin(); InstanceChildIt != PrimRange.end(); ++InstanceChildIt)
+				// There's a single source prim shared across all instances, so just fetch it from the first one
+				TSet<UE::FSdfPath> SourcePaths = GetLocalNonRootCompositionArcSourcePaths(Instances[0]);
+				if (SourcePaths.Num() == 0)
 				{
-					TRACE_CPUPROFILER_EVENT_SCOPE(RegisterInstanceableAuxPrims::InstanceChild);
+					return;
+				}
 
-					pxr::SdfPath SdfChildPrimPath = InstanceChildIt->GetPrimPath();
-					UE::FSdfPath ChildPrimPath{SdfChildPrimPath};
+				// Step into every instance of this prototype on the stage
+				ParallelFor(
+					TEXT("RegisterInstanceableAuxPrimsInstances"),
+					Instances.size(),
+					MinBatchSize,
+					[&Instances, &InstancePaths, &Context, &Impl, &Stage, &SourcePaths](int32 Index)
+					{
+						TRACE_CPUPROFILER_EVENT_SCOPE(RegisterInstanceableAuxPrims::PrototypeInstance);
 
-					// Register a dependency from child prim to the analogue prim within the prototype itself
-					Impl.RegisterAuxiliaryPrims(ChildPrimPath, {UE::FSdfPath{InstanceChildIt->GetPrimInPrototype().GetPrimPath()}});
+						FScopedUsdAllocs Allocs;
 
-					// Register a dependency from child prim to analogue prims on the sources used for the instance.
-					// We have to do some path surgery to discover what the analogue paths on the source prims are though
-					pxr::SdfPath RelativeChildPath = SdfChildPrimPath.MakeRelativePath(InstancePath);
-					for (const UE::FSdfPath& SourcePath : SourcePaths)
-					{
-						pxr::SdfPath ChildOnSourcePath = pxr::SdfPath{SourcePath}.AppendPath(RelativeChildPath);
-						if (pxr::UsdPrim ChildOnSource = Stage->GetPrimAtPath(ChildOnSourcePath))
+						const pxr::UsdPrim& Instance = Instances[Index];
+
+						UE::FSdfPath InstancePath{Instance.GetPrimPath()};
+						InstancePaths[Index] = InstancePath;
+
+						if (UE::UsdInfoCache::Private::FUsdPrimInfo* MainPrim = Impl.InfoMap.Find(InstancePath))
 						{
-							Impl.RegisterAuxiliaryPrims(ChildPrimPath, {UE::FSdfPath{ChildOnSourcePath}});
+							FWriteScopeLock PrimLock{Impl.PrimLocks[MainPrim->PrimLockIndex]};
+							MainPrim->AuxPrims.Append(SourcePaths);
 						}
+
+						// Here we'll traverse the entire subtree of the instance
+						pxr::UsdPrimRange PrimRange(Instance, pxr::UsdTraverseInstanceProxies());
+						for (pxr::UsdPrimRange::iterator InstanceChildIt = ++PrimRange.begin(); InstanceChildIt != PrimRange.end(); ++InstanceChildIt)
+						{
+							TRACE_CPUPROFILER_EVENT_SCOPE(RegisterInstanceableAuxPrims::InstanceChild);
+
+							pxr::SdfPath SdfChildPrimPath = InstanceChildIt->GetPrimPath();
+							UE::FSdfPath ChildPrimPath{SdfChildPrimPath};
+
+							// Register a dependency from child prim to analogue prims on the sources used for the instance.
+							// We have to do some path surgery to discover what the analogue paths on the source prims are though
+							pxr::SdfPath RelativeChildPath = SdfChildPrimPath.MakeRelativePath(InstancePath);
+							for (const UE::FSdfPath& SourcePath : SourcePaths)
+							{
+								pxr::SdfPath ChildOnSourcePath = pxr::SdfPath{SourcePath}.AppendPath(RelativeChildPath);
+								if (pxr::UsdPrim ChildOnSource = Stage->GetPrimAtPath(ChildOnSourcePath))
+								{
+									Impl.RegisterAuxiliaryPrims(ChildPrimPath, {UE::FSdfPath{ChildOnSourcePath}});
+								}
+							}
+						}
+					}
+				);
+
+				// Append all the instance paths in one go for these source paths
+				for (const UE::FSdfPath& AuxPrimPath : SourcePaths)
+				{
+					if (UE::UsdInfoCache::Private::FUsdPrimInfo* AuxPrim = Impl.InfoMap.Find(AuxPrimPath))
+					{
+						FWriteScopeLock PrimLock{Impl.PrimLocks[AuxPrim->PrimLockIndex]};
+						AuxPrim->MainPrims.Append(InstancePaths);
 					}
 				}
 			}
-		}
+		);
 	}
 
 	void FindValidGeometryCacheRoot(
@@ -1171,24 +1446,26 @@ namespace UE::USDInfoCacheImpl::Private
 		UE::FSdfPath PrimPath{UsdPrimPath};
 		{
 			FReadScopeLock ScopeLock(Impl.InfoMapLock);
-
-			UE::UsdInfoCache::Private::FUsdPrimInfo& Info = Impl.InfoMap.FindChecked(UE::FSdfPath(UsdPrim.GetPrimPath()));
-
-			// A prim is considered a valid root if its subtree has no uncollapsible branch and a valid depth.
-			// A valid depth is positive, meaning it has an animated mesh, and doesn't exceed the limit.
-			bool bIsValidDepth = Info.GeometryCacheDepth > -1 && Info.GeometryCacheDepth <= GGeometryCacheMaxDepth;
-			if (!EnumHasAnyFlags(Info.GeometryCacheState, EGeometryCachePrimState::Uncollapsible) && bIsValidDepth)
-			{
-				OutState = EGeometryCachePrimState::ValidRoot;
-				Info.GeometryCacheState = EGeometryCachePrimState::ValidRoot;
-				return;
-			}
-			// The prim is not a valid root so it's flagged as uncollapsible since the root will be among its children
-			// and the eventual geometry cache cannot be collapsed.
-			else
+			if (UE::UsdInfoCache::Private::FUsdPrimInfo* Info = Impl.InfoMap.Find(UE::FSdfPath(UsdPrim.GetPrimPath())))
 			{
-				OutState = EGeometryCachePrimState::Uncollapsible;
-				Info.GeometryCacheState = EGeometryCachePrimState::Uncollapsible;
+				FWriteScopeLock PrimLock{Impl.PrimLocks[Info->PrimLockIndex]};
+
+				// A prim is considered a valid root if its subtree has no uncollapsible branch and a valid depth.
+				// A valid depth is positive, meaning it has an animated mesh, and doesn't exceed the limit.
+				bool bIsValidDepth = Info->GeometryCacheDepth > -1 && Info->GeometryCacheDepth <= GGeometryCacheMaxDepth;
+				if (!EnumHasAnyFlags(Info->GeometryCacheState, EGeometryCachePrimState::Uncollapsible) && bIsValidDepth)
+				{
+					OutState = EGeometryCachePrimState::ValidRoot;
+					Info->GeometryCacheState = EGeometryCachePrimState::ValidRoot;
+					return;
+				}
+				// The prim is not a valid root so it's flagged as uncollapsible since the root will be among its children
+				// and the eventual geometry cache cannot be collapsed.
+				else
+				{
+					OutState = EGeometryCachePrimState::Uncollapsible;
+					Info->GeometryCacheState = EGeometryCachePrimState::Uncollapsible;
+				}
 			}
 		}
 
@@ -1198,11 +1475,15 @@ namespace UE::USDInfoCacheImpl::Private
 		TArray<pxr::UsdPrim> Prims;
 		for (pxr::UsdPrim Child : PrimChildren)
 		{
-			FReadScopeLock ScopeLock(Impl.InfoMapLock);
-			const UE::UsdInfoCache::Private::FUsdPrimInfo& Info = Impl.InfoMap.FindChecked(UE::FSdfPath(Child.GetPrimPath()));
+			bool bIsCollapsible = false;
+			{
+				FReadScopeLock ScopeLock(Impl.InfoMapLock);
+				const UE::UsdInfoCache::Private::FUsdPrimInfo& Info = Impl.InfoMap.FindChecked(UE::FSdfPath(Child.GetPrimPath()));
+				bIsCollapsible = EnumHasAnyFlags(Info.GeometryCacheState, EGeometryCachePrimState::Collapsible);
+			}
 
 			// A subtree is considered only if it has anything collapsible in the first place
-			if (EnumHasAnyFlags(Info.GeometryCacheState, EGeometryCachePrimState::Collapsible))
+			if (bIsCollapsible)
 			{
 				FindValidGeometryCacheRoot(Child, Context, Impl, OutState);
 			}
@@ -1270,15 +1551,17 @@ namespace UE::USDInfoCacheImpl::Private
 		bool bIsAnimatedMesh = UsdUtils::IsAnimatedMesh(UsdPrim);
 		if (!bIsImporting)
 		{
-			FWriteScopeLock ScopeLock(Impl.InfoMapLock);
-
-			UE::UsdInfoCache::Private::FUsdPrimInfo& Info = Impl.InfoMap.FindChecked(UE::FSdfPath(UsdPrimPath));
+			FReadScopeLock ScopeLock(Impl.InfoMapLock);
+			if (UE::UsdInfoCache::Private::FUsdPrimInfo* Info = Impl.InfoMap.Find(UE::FSdfPath(UsdPrimPath)))
+			{
+				FWriteScopeLock PrimLock{Impl.PrimLocks[Info->PrimLockIndex]};
 
-			// When loading on the stage, the GeometryCache root can only be the animated mesh prim itself
-			// and there's no collapsing involved since each animated mesh will become a GeometryCache.
-			// The depth is irrelevant here.
-			Info.GeometryCacheDepth = -1;
-			Info.GeometryCacheState = bIsAnimatedMesh ? EGeometryCachePrimState::ValidRoot : EGeometryCachePrimState::Uncollapsible;
+				// When loading on the stage, the GeometryCache root can only be the animated mesh prim itself
+				// and there's no collapsing involved since each animated mesh will become a GeometryCache.
+				// The depth is irrelevant here.
+				Info->GeometryCacheDepth = -1;
+				Info->GeometryCacheState = bIsAnimatedMesh ? EGeometryCachePrimState::ValidRoot : EGeometryCachePrimState::Uncollapsible;
+			}
 
 			return;
 		}
@@ -1399,11 +1682,14 @@ namespace UE::USDInfoCacheImpl::Private
 		}
 
 		{
-			FWriteScopeLock ScopeLock(Impl.InfoMapLock);
+			FReadScopeLock ScopeLock(Impl.InfoMapLock);
+			if (UE::UsdInfoCache::Private::FUsdPrimInfo* Info = Impl.InfoMap.Find(UE::FSdfPath(UsdPrimPath)))
+			{
+				FWriteScopeLock PrimLock{Impl.PrimLocks[Info->PrimLockIndex]};
 
-			UE::UsdInfoCache::Private::FUsdPrimInfo& Info = Impl.InfoMap.FindChecked(UE::FSdfPath(UsdPrimPath));
-			Info.GeometryCacheDepth = Depth;
-			Info.GeometryCacheState = PrimState;
+				Info->GeometryCacheDepth = Depth;
+				Info->GeometryCacheState = PrimState;
+			}
 		}
 
 		// We've encountered a potential root and the subtree has a geometry cache branch, so find its root
@@ -1495,7 +1781,7 @@ namespace UE::USDInfoCacheImpl::Private
 		}
 	}
 #endif	  // USE_SD_SDK
-}	 // namespace UE::USDInfoCacheImpl::Private
+}	 // namespace UE::USDInfoCache::Private
 
 bool FUsdInfoCache::IsPotentialGeometryCacheRoot(const UE::FSdfPath& Path) const
 {
@@ -1505,6 +1791,8 @@ bool FUsdInfoCache::IsPotentialGeometryCacheRoot(const UE::FSdfPath& Path) const
 
 		if (const UE::UsdInfoCache::Private::FUsdPrimInfo* FoundInfo = ImplPtr->InfoMap.Find(Path))
 		{
+			FReadScopeLock PrimLock{ImplPtr->PrimLocks[FoundInfo->PrimLockIndex]};
+
 			return FoundInfo->GeometryCacheState == UE::UsdInfoCache::Private::EGeometryCachePrimState::ValidRoot;
 		}
 
@@ -1523,6 +1811,8 @@ TOptional<uint64> FUsdInfoCache::GetSubtreeVertexCount(const UE::FSdfPath& Path)
 
 		if (const UE::UsdInfoCache::Private::FUsdPrimInfo* FoundInfo = ImplPtr->InfoMap.Find(Path))
 		{
+			FReadScopeLock PrimLock{ImplPtr->PrimLocks[FoundInfo->PrimLockIndex]};
+
 			return FoundInfo->ExpectedVertexCountForSubtree;
 		}
 
@@ -1541,6 +1831,8 @@ TOptional<uint64> FUsdInfoCache::GetSubtreeMaterialSlotCount(const UE::FSdfPath&
 
 		if (const UE::UsdInfoCache::Private::FUsdPrimInfo* FoundInfo = ImplPtr->InfoMap.Find(Path))
 		{
+			FReadScopeLock PrimLock{ImplPtr->PrimLocks[FoundInfo->PrimLockIndex]};
+
 			return FoundInfo->SubtreeMaterialSlots.Num();
 		}
 
@@ -1559,6 +1851,8 @@ TOptional<TArray<UsdUtils::FUsdPrimMaterialSlot>> FUsdInfoCache::GetSubtreeMater
 
 		if (const UE::UsdInfoCache::Private::FUsdPrimInfo* FoundInfo = ImplPtr->InfoMap.Find(Path))
 		{
+			FReadScopeLock PrimLock{ImplPtr->PrimLocks[FoundInfo->PrimLockIndex]};
+
 			return FoundInfo->SubtreeMaterialSlots;
 		}
 
@@ -1725,6 +2019,8 @@ void FUsdInfoCache::RebuildCacheForSubtree(const UE::FUsdPrim& Prim, FUsdSchemaT
 #if USE_USD_SDK
 	TRACE_CPUPROFILER_EVENT_SCOPE(FUsdInfoCache::RebuildCacheForSubtree);
 
+	using namespace UE::USDInfoCache::Private;
+
 	FUsdInfoCacheImpl* ImplPtr = Impl.Get();
 	if (!ImplPtr)
 	{
@@ -1736,94 +2032,86 @@ void FUsdInfoCache::RebuildCacheForSubtree(const UE::FUsdPrim& Prim, FUsdSchemaT
 
 	TGuardValue<bool> Guard{Context.bIsBuildingInfoCache, true};
 	{
-		FScopedUsdAllocs Allocs;
-
 		pxr::UsdPrim UsdPrim{Prim};
 		if (!UsdPrim)
 		{
 			return;
 		}
 
-		// Don't call Clear() here as we don't want to get rid of PrimPathToAssets because we're rebuilding the cache,
-		// as that info is also linked to the asset cache
-		{
-			FWriteScopeLock ScopeLock(ImplPtr->InfoMapLock);
-			ImplPtr->InfoMap.Empty();
-		}
-		{
-			FWriteScopeLock ScopeLock(ImplPtr->MaterialUsersLock);
-			ImplPtr->MaterialUsers.Empty();
-		}
-		{
-			FWriteScopeLock ScopeLock(ImplPtr->AuxiliaryPrimsLock);
-			ImplPtr->AuxToMainPrims.Empty();
-			ImplPtr->MainToAuxPrims.Empty();
-		}
+		ImplPtr->TempStage = UE::FUsdStageWeak{UsdPrim.GetStage()};
 
 		IUsdSchemasModule& UsdSchemasModule = FModuleManager::Get().LoadModuleChecked<IUsdSchemasModule>(TEXT("USDSchemas"));
 		FUsdSchemaTranslatorRegistry& Registry = UsdSchemasModule.GetTranslatorRegistry();
 
-		TMap<UE::FSdfPath, TArray<UsdUtils::FUsdPrimMaterialSlot>> TempSubtreeSlots;
-		TArray<FString> PointInstancerPaths;
-
 		pxr::TfToken MaterialPurposeToken = pxr::UsdShadeTokens->allPurpose;
 		if (!Context.MaterialPurpose.IsNone())
 		{
 			MaterialPurposeToken = UnrealToUsd::ConvertToken(*Context.MaterialPurpose.ToString()).Get();
 		}
 
+		// We don't call FUsdInfoCache::Clear() here as we don't want to get rid of PrimPathToAssets because
+		// we're rebuilding the cache, as that info is also linked to the asset cache
+		{
+			FWriteScopeLock ScopeLock(ImplPtr->InfoMapLock);
+			ImplPtr->InfoMap.Reset();
+		}
+		{
+			FWriteScopeLock PointInstancerLock(ImplPtr->TempPointInstancerPathsLock);
+			ImplPtr->TempPointInstancerPaths.Reset();
+		}
+
+		// This should be the first step as all future functions will expect to find one entry per prim in the cache
+		RepopulateInfoMap(UsdPrim, *ImplPtr);
+
 		// Propagate vertex and material slot counts before we query CollapsesChildren because the Xformable
 		// translator needs to know when it would generate too large a static mesh
 		uint64 SubtreeVertexCount = 0;
 		TArray<UsdUtils::FUsdPrimMaterialSlot> SubtreeSlots;
 		const bool bPossibleInheritedBindings = false;
-		UE::USDInfoCacheImpl::Private::RecursivePropagateVertexAndMaterialSlotCounts(
+		RecursivePropagateVertexAndMaterialSlotCounts(
 			UsdPrim,
 			Context,
 			MaterialPurposeToken,
 			*ImplPtr,
 			Registry,
-			TempSubtreeSlots,
-			PointInstancerPaths,
 			SubtreeVertexCount,
 			SubtreeSlots,
 			bPossibleInheritedBindings
 		);
 
-		UE::USDInfoCacheImpl::Private::UpdateInfoForPointInstancers(UsdPrim.GetStage(), Context, *ImplPtr, PointInstancerPaths, TempSubtreeSlots);
+		UpdateInfoForPointInstancers(Context, *ImplPtr);
 
-		UE::USDInfoCacheImpl::Private::CheckForGeometryCache(UsdPrim, Context, *ImplPtr);
+		CheckForGeometryCache(UsdPrim, Context, *ImplPtr);
 
-		UE::USDInfoCacheImpl::Private::RecursiveQueryCollapsesChildren(UsdPrim, Context, *ImplPtr, Registry);
+		RecursiveQueryCollapsesChildren(UsdPrim, Context, *ImplPtr, Registry);
 
-		UE::USDInfoCacheImpl::Private::RegisterInstanceableAuxPrims(UsdPrim, Context, *ImplPtr);
+		RegisterInstanceableAuxPrims(Context, *ImplPtr);
 
-		UE::USDInfoCacheImpl::Private::CollectMaterialSlotCounts(*ImplPtr, TempSubtreeSlots, Context.bMergeIdenticalMaterialSlots);
+		CollectMaterialSlotCounts(*ImplPtr, Context.bMergeIdenticalMaterialSlots);
 	}
 #endif	  // USE_USD_SDK
 }
 
 void FUsdInfoCache::Clear()
 {
+	TRACE_CPUPROFILER_EVENT_SCOPE(AUsdStageActor::UnloadUsdStage);
+
 	if (FUsdInfoCacheImpl* ImplPtr = Impl.Get())
 	{
 		{
+			TRACE_CPUPROFILER_EVENT_SCOPE(InfoMapEmpty);
 			FWriteScopeLock ScopeLock(ImplPtr->InfoMapLock);
 			ImplPtr->InfoMap.Empty();
 		}
 		{
+			TRACE_CPUPROFILER_EVENT_SCOPE(PrimPathToAssetsEmpty);
 			FWriteScopeLock ScopeLock(ImplPtr->PrimPathToAssetsLock);
 			ImplPtr->PrimPathToAssets.Empty();
 			ImplPtr->AssetToPrimPaths.Empty();
 		}
 		{
-			FWriteScopeLock ScopeLock(ImplPtr->MaterialUsersLock);
-			ImplPtr->MaterialUsers.Empty();
-		}
-		{
-			FWriteScopeLock ScopeLock(ImplPtr->AuxiliaryPrimsLock);
-			ImplPtr->AuxToMainPrims.Empty();
-			ImplPtr->MainToAuxPrims.Empty();
+			FWriteScopeLock PointInstancerLock(ImplPtr->TempPointInstancerPathsLock);
+			ImplPtr->TempPointInstancerPaths.Empty();
 		}
 	}
 }
@@ -1838,3 +2126,73 @@ bool FUsdInfoCache::IsEmpty()
 
 	return true;
 }
+
+TOptional<bool> FUsdInfoCache::CanXformableSubtreeBeCollapsed(const UE::FSdfPath& RootPath, FUsdSchemaTranslationContext& Context) const
+{
+#if USE_USD_SDK
+	TRACE_CPUPROFILER_EVENT_SCOPE(FUsdInfoCache::CanSubtreeBeCollapsed);
+
+	// The only reason this function exists is that FUsdGeomXformableTranslator::CollapsesChildren() needs to check if all
+	// GeomXformable prims in its subtree return true for CanBeCollapsed().
+	//
+	// We don't want to compute this for the entire stage on the main info cache build, because it may not be needed.
+	// However, we definitely do not want each call to FUsdGeomXformableTranslator::CollapsesChildren() to traverse its entire
+	// subtree of prims calling CanBeCollapsed() on their own: That would be a massive waste since the output is going to
+	// be the same regardless of the caller.
+	//
+	// This is the awkward compromise where the first call to FUsdGeomXformableTranslator::CollapsesChildren() will traverse
+	// its entire subtree and fill this in, and subsequent calls can just use those results, or fill in additional subtrees, etc.
+
+	if (FUsdInfoCacheImpl* ImplPtr = Impl.Get())
+	{
+		FReadScopeLock ScopeLock(ImplPtr->InfoMapLock);
+
+		if (const UE::UsdInfoCache::Private::FUsdPrimInfo* FoundInfo = ImplPtr->InfoMap.Find(RootPath))
+		{
+			FReadScopeLock PrimLock{ImplPtr->PrimLocks[FoundInfo->PrimLockIndex]};
+			if (FoundInfo->bXformSubtreeCanBeCollapsed.IsSet())
+			{
+				return FoundInfo->bXformSubtreeCanBeCollapsed.GetValue();
+			}
+		}
+
+		TOptional<bool> bCanBeCollapsed;
+
+		// Fill in missing entries for CanBeCollapsed on-demand and compute the value for the prim at RootPath,
+		// if we can still access our stage
+		pxr::UsdStageWeakPtr UsdStageWeak = ImplPtr->TempStage;
+		pxr::UsdStageRefPtr Stage = UsdStageWeak;
+		if (Stage)
+		{
+			if (pxr::UsdPrim Prim = Stage->GetPrimAtPath(RootPath))
+			{
+				IUsdSchemasModule& UsdSchemasModule = FModuleManager::Get().LoadModuleChecked<IUsdSchemasModule>(TEXT("USDSchemas"));
+				FUsdSchemaTranslatorRegistry& Registry = UsdSchemasModule.GetTranslatorRegistry();
+
+				bCanBeCollapsed = UE::USDInfoCache::Private::RecursiveQueryCanBeCollapsed(Prim, Context, *ImplPtr, Registry);
+			}
+		}
+
+		// We can potentially still fail to find this here, in case our stage reference is broken (i.e. called outside of the
+		// main infocache build callstack).
+		//
+		// There shouldn't be any point in checking our FoundInfo again though: If we didn't return anything valid from
+		// our call to RecursiveQueryCanBeCollapsed, then we didn't put anything new on the InfoMap either
+		if (bCanBeCollapsed.IsSet())
+		{
+			return bCanBeCollapsed;
+		}
+
+		UE_LOG(
+			LogUsd,
+			Warning,
+			TEXT(
+				"Failed to find whether subtree '%s' can be collapsed or not. Note: This function is meant to be used only during the main FUsdInfoCache build!"
+			),
+			*RootPath.GetString()
+		);
+	}
+#endif	  // USE_USD_SDK
+
+	return {};
+}
diff --git a/Engine/Plugins/Importers/USDImporter/Source/USDSchemas/Public/USDInfoCache.h b/Engine/Plugins/Importers/USDImporter/Source/USDSchemas/Public/USDInfoCache.h
--- a/Engine/Plugins/Importers/USDImporter/Source/USDSchemas/Public/USDInfoCache.h
+++ b/Engine/Plugins/Importers/USDImporter/Source/USDSchemas/Public/USDInfoCache.h
@@ -137,6 +137,18 @@ public:
 	TArray<UE::FSdfPath> GetPrimsForAsset(const UObject* Asset) const;
 	TMap<UE::FSdfPath, TArray<TWeakObjectPtr<UObject>>> GetAllAssetPrimLinks() const;
 
+private:
+	friend class FUsdGeomXformableTranslator;
+	// Returns true if every prim on the subtree below RootPath (including the RootPath prim itself) returns true for
+	// CanBeCollapsed(), according to their own schema translators.
+	//
+	// WARNING: This is intended for internal use, and exclusively during the actual info cache build process as it will
+	// need to query the prim/stage directly. Calling it after the info cache build may yield back an empty optional,
+	// meaning it is unknown at this point whether the prim CanBeCollapsed or not.
+	//
+	// In general, you shouldn't call this, but just use "IsPathCollapsed" or "DoesPathCollapseChildren" instead.
+	TOptional<bool> CanXformableSubtreeBeCollapsed(const UE::FSdfPath& RootPath, FUsdSchemaTranslationContext& Context) const;
+
 private:
 	TUniquePtr<FUsdInfoCacheImpl> Impl;
 };
diff --git a/Engine/Plugins/Importers/USDImporter/Source/USDStage/Private/USDStageActor.cpp b/Engine/Plugins/Importers/USDImporter/Source/USDStage/Private/USDStageActor.cpp
--- a/Engine/Plugins/Importers/USDImporter/Source/USDStage/Private/USDStageActor.cpp
+++ b/Engine/Plugins/Importers/USDImporter/Source/USDStage/Private/USDStageActor.cpp
@@ -1866,9 +1866,17 @@ void AUsdStageActor::OnUsdObjectsChanged(const UsdUtils::FObjectChangesByPath& I
 
 USDSTAGE_API void AUsdStageActor::Reset()
 {
-	Modify();
+	TRACE_CPUPROFILER_EVENT_SCOPE(AUsdStageActor::Reset);
+
+	{
+		TRACE_CPUPROFILER_EVENT_SCOPE(Modify);
+		Modify();
+	}
 
-	Super::Reset();
+	{
+		TRACE_CPUPROFILER_EVENT_SCOPE(Reset);
+		Super::Reset();
+	}
 
 	bool bUnloadIfNeeded = true;
 	CloseUsdStage(bUnloadIfNeeded);
@@ -1948,8 +1956,12 @@ UUsdPrimTwin* AUsdStageActor::ExpandPrim(
 	UUsdPrimTwin* UsdPrimTwin = nullptr;
 #if USE_USD_SDK
 	// "Active" is the non-destructive deletion used in USD. Sometimes when we rename/remove a prim in a complex stage it may remain in
-	// an inactive state, but its otherwise effectively deleted
-	if (!Prim || !Prim.IsActive())
+	// an inactive state, but its otherwise effectively deleted.
+	//
+	// We check IsDefined() because we need to consider the possibility that we've been called directly for this prim
+	// (e.g. when handling an update notice). During regular traversal when opening the stage, the Prim.GetFilteredChildren() call
+	// within this same function will naturally strip all "pure over" prims
+	if (!Prim || !Prim.IsActive() || !Prim.IsDefined())
 	{
 		return nullptr;
 	}
@@ -3250,6 +3262,8 @@ void AUsdStageActor::LoadUsdStage(bool bOpenIfNeeded)
 
 void AUsdStageActor::UnloadUsdStage()
 {
+	TRACE_CPUPROFILER_EVENT_SCOPE(AUsdStageActor::UnloadUsdStage);
+
 	// No point doing any of this if we're unloading because we're exiting the engine altogether
 	if (IsEngineExitRequested())
 	{
diff --git a/Engine/Plugins/Importers/USDImporter/Source/USDStageEditor/Private/SUSDStage.cpp b/Engine/Plugins/Importers/USDImporter/Source/USDStageEditor/Private/SUSDStage.cpp
--- a/Engine/Plugins/Importers/USDImporter/Source/USDStageEditor/Private/SUSDStage.cpp
+++ b/Engine/Plugins/Importers/USDImporter/Source/USDStageEditor/Private/SUSDStage.cpp
@@ -2457,6 +2457,8 @@ void SUsdStage::FileReset()
 
 void SUsdStage::FileClose()
 {
+	TRACE_CPUPROFILER_EVENT_SCOPE(SUsdStage::FileClose);
+
 	TOptional<FScopedTransaction> Transaction;
 	if (!GDiscardUndoBufferOnStageOpenClose)
 	{
diff --git a/Engine/Plugins/Importers/USDImporter/Source/USDStageEditorViewModels/Private/USDStageViewModel.cpp b/Engine/Plugins/Importers/USDImporter/Source/USDStageEditorViewModels/Private/USDStageViewModel.cpp
--- a/Engine/Plugins/Importers/USDImporter/Source/USDStageEditorViewModels/Private/USDStageViewModel.cpp
+++ b/Engine/Plugins/Importers/USDImporter/Source/USDStageEditorViewModels/Private/USDStageViewModel.cpp
@@ -196,6 +196,7 @@ void FUsdStageViewModel::ResetStage()
 
 void FUsdStageViewModel::CloseStage()
 {
+	TRACE_CPUPROFILER_EVENT_SCOPE(FUsdStageViewModel::CloseStage);
 	if (AUsdStageActor* StageActor = UsdStageActor.Get())
 	{
 		StageActor->Reset();
diff --git a/Engine/Plugins/Runtime/USDCore/Source/USDUtilities/Private/USDConversionUtils.cpp b/Engine/Plugins/Runtime/USDCore/Source/USDUtilities/Private/USDConversionUtils.cpp
--- a/Engine/Plugins/Runtime/USDCore/Source/USDUtilities/Private/USDConversionUtils.cpp
+++ b/Engine/Plugins/Runtime/USDCore/Source/USDUtilities/Private/USDConversionUtils.cpp
@@ -1517,6 +1517,8 @@ TArray<TUsdStore<pxr::UsdPrim>> UsdUtils::GetAllPrimsOfType(
 	const TArray<TUsdStore<pxr::TfType>>& ExcludeSchemaTypes
 )
 {
+	TRACE_CPUPROFILER_EVENT_SCOPE(UsdUtils::GetAllPrimsOfType);
+
 	TArray<TUsdStore<pxr::UsdPrim>> Result;
 
 	pxr::UsdPrimRange PrimRange(StartPrim, pxr::UsdTraverseInstanceProxies());
diff --git a/Engine/Plugins/Runtime/USDCore/Source/USDUtilities/Private/USDGeomMeshConversion.cpp b/Engine/Plugins/Runtime/USDCore/Source/USDUtilities/Private/USDGeomMeshConversion.cpp
--- a/Engine/Plugins/Runtime/USDCore/Source/USDUtilities/Private/USDGeomMeshConversion.cpp
+++ b/Engine/Plugins/Runtime/USDCore/Source/USDUtilities/Private/USDGeomMeshConversion.cpp
@@ -215,6 +215,8 @@ namespace UE::UsdGeomMeshConversion::Private
 
 	pxr::VtArray<int> GetFaceVertexCounts(const pxr::UsdPrim& UsdPrim, pxr::UsdTimeCode TimeCode)
 	{
+		TRACE_CPUPROFILER_EVENT_SCOPE(GetFaceVertexCounts);
+
 		if (pxr::UsdGeomMesh Mesh = pxr::UsdGeomMesh{UsdPrim})
 		{
 			pxr::UsdAttribute Attr = Mesh.GetFaceVertexCountsAttr();
@@ -3360,6 +3362,8 @@ UsdUtils::FUsdPrimMaterialAssignmentInfo UsdUtils::GetPrimMaterialAssignments(
 	const pxr::TfToken& MaterialPurpose
 )
 {
+	TRACE_CPUPROFILER_EVENT_SCOPE(UsdUtils::GetPrimMaterialAssignments);
+
 	if (!UsdPrim)
 	{
 		return {};
@@ -3471,19 +3475,6 @@ UsdUtils::FUsdPrimMaterialAssignmentInfo UsdUtils::GetPrimMaterialAssignments(
 	std::vector<pxr::UsdGeomSubset> GeomSubsets = pxr::UsdShadeMaterialBindingAPI(UsdPrim).GetMaterialBindSubsets();
 	if (GeomSubsets.size() > 0)
 	{
-		std::string ReasonWhyNotPartition;
-		bool bHasValidPartitions = pxr::UsdGeomSubset::ValidateSubsets(GeomSubsets, NumFaces, pxr::UsdGeomTokens->partition, &ReasonWhyNotPartition);
-		if (!bHasValidPartitions)
-		{
-			UE_LOG(
-				LogUsd,
-				Warning,
-				TEXT("Found an invalid GeomSubsets partition in prim '%s': %s"),
-				*UsdToUnreal::ConvertPath(UsdPrim.GetPath()),
-				*UsdToUnreal::ConvertString(ReasonWhyNotPartition)
-			);
-		}
-
 		for (uint32 GeomSubsetIndex = 0; GeomSubsetIndex < GeomSubsets.size(); ++GeomSubsetIndex)
 		{
 			const pxr::UsdGeomSubset& GeomSubset = GeomSubsets[GeomSubsetIndex];
@@ -3570,25 +3561,49 @@ UsdUtils::FUsdPrimMaterialAssignmentInfo UsdUtils::GetPrimMaterialAssignments(
 			}
 		}
 
-		// Extra slot for unspecified faces.
-		// We need to fetch this even if we won't provide indices because we may need to create an additional slot for unassigned polygons
-		pxr::VtIntArray UnassignedIndices = pxr::UsdGeomSubset::GetUnassignedIndices(GeomSubsets, NumFaces);
-		if (UnassignedIndices.size() == 0)
-		{
-			bNeedsMainAssignment = false;
-		}
-		else
 		{
-			// Assign these leftover indices to the *next* material slot we'll create (doesn't exist yet),
-			// which will be the "main" material assignment slot
-			int32 LeftoverSlotIndex = Result.Slots.Num();
-			if (bProvideMaterialIndices)
+			TRACE_CPUPROFILER_EVENT_SCOPE(GetUnassignedFaces);
+
+			// Check if we have any unassigned faces
+			//
+			// Just summing the indices like this is an approximation, because they may be invalid/repeated indices, etc.
+			// Ideally we'd call pxr::UsdGeomSubset::GetUnassignedIndices right away here to know for sure if we need an additional slot.
+			// That can be slow however, and this function is in the hot path of the info cache build, so given that invalid partitions
+			// are invalid data in the first place, this approximation is hopefully a good enough compromise in order to not slow down the
+			// general case too much
+			uint64 SubsetIndexTotal = 0;
+			for (const pxr::UsdGeomSubset& Subset : GeomSubsets)
 			{
-				for (int PolygonIndex : UnassignedIndices)
+				pxr::VtIntArray PolygonIndicesInSubset;
+				Subset.GetIndicesAttr().Get(&PolygonIndicesInSubset, TimeCode);
+
+				SubsetIndexTotal += PolygonIndicesInSubset.size();
+			}
+
+			// Extra slot for unspecified faces.
+			// We need to fetch this even if we won't provide indices because we may need to create an additional slot for unassigned polygons
+			const bool bHasUnassignedIndices = SubsetIndexTotal != NumFaces;
+			if (bHasUnassignedIndices)
+			{
+				UE_LOG(LogUsd, Warning, TEXT("Found an invalid GeomSubsets partition in prim '%s'"), *UsdToUnreal::ConvertPath(UsdPrim.GetPath()));
+
+				// Assign these leftover indices to the *next* material slot we'll create (doesn't exist yet),
+				// which will be the "main" material assignment slot
+				int32 LeftoverSlotIndex = Result.Slots.Num();
+
+				if (bProvideMaterialIndices)
 				{
-					Result.MaterialIndices[PolygonIndex] = LeftoverSlotIndex;
+					pxr::VtIntArray UnassignedIndices = pxr::UsdGeomSubset::GetUnassignedIndices(GeomSubsets, NumFaces);
+					for (int PolygonIndex : UnassignedIndices)
+					{
+						Result.MaterialIndices[PolygonIndex] = LeftoverSlotIndex;
+					}
 				}
 			}
+			else
+			{
+				bNeedsMainAssignment = false;
+			}
 		}
 	}
 
@@ -4630,6 +4645,8 @@ void UsdUtils::RepairNormalsAndTangents(const FString& PrimPath, FMeshDescriptio
 
 TOptional<UsdUtils::FDisplayColorMaterial> UsdUtils::ExtractDisplayColorMaterial(const pxr::UsdGeomGprim& Gprim, const pxr::UsdTimeCode TimeCode)
 {
+	TRACE_CPUPROFILER_EVENT_SCOPE(UsdUtils::ExtractDisplayColorMaterial);
+
 	if (!Gprim)
 	{
 		return {};
@@ -5531,6 +5548,8 @@ UsdUtils::EMeshTopologyVariance UsdUtils::GetMeshTopologyVariance(const pxr::Usd
 
 uint64 UsdUtils::GetGprimVertexCount(const pxr::UsdGeomGprim& Gprim, double TimeCode)
 {
+	TRACE_CPUPROFILER_EVENT_SCOPE(UsdUtils::GetGprimVertexCount);
+
 	if (pxr::UsdGeomMesh Mesh{Gprim})
 	{
 		if (pxr::UsdAttribute Points = Mesh.GetPointsAttr())
diff --git a/Engine/Plugins/Runtime/USDCore/Source/USDUtilities/Private/USDShadeConversion.cpp b/Engine/Plugins/Runtime/USDCore/Source/USDUtilities/Private/USDShadeConversion.cpp
--- a/Engine/Plugins/Runtime/USDCore/Source/USDUtilities/Private/USDShadeConversion.cpp
+++ b/Engine/Plugins/Runtime/USDCore/Source/USDUtilities/Private/USDShadeConversion.cpp
@@ -3114,6 +3114,8 @@ void UsdUtils::AuthorUnrealMaterialBinding(pxr::UsdPrim& MeshOrGeomSubsetPrim, c
 
 TOptional<FString> UsdUtils::GetUnrealSurfaceOutput(const pxr::UsdPrim& MaterialPrim)
 {
+	TRACE_CPUPROFILER_EVENT_SCOPE(UsdUtils::GetUnrealSurfaceOutput);
+
 	if (!MaterialPrim)
 	{
 		return {};
diff --git a/Engine/Plugins/Runtime/USDCore/Source/UnrealUSDWrapper/Private/UsdWrappers/UsdPrim.cpp b/Engine/Plugins/Runtime/USDCore/Source/UnrealUSDWrapper/Private/UsdWrappers/UsdPrim.cpp
--- a/Engine/Plugins/Runtime/USDCore/Source/UnrealUSDWrapper/Private/UsdWrappers/UsdPrim.cpp
+++ b/Engine/Plugins/Runtime/USDCore/Source/UnrealUSDWrapper/Private/UsdWrappers/UsdPrim.cpp
@@ -215,6 +215,15 @@ namespace UE
 #endif	  // #if USE_USD_SDK
 	}
 
+	bool FUsdPrim::IsDefined() const
+	{
+#if USE_USD_SDK
+		return Impl->PxrUsdPrim.Get().IsDefined();
+#else
+		return false;
+#endif	  // #if USE_USD_SDK
+	}
+
 	bool FUsdPrim::IsGroup() const
 	{
 #if USE_USD_SDK
diff --git a/Engine/Plugins/Runtime/USDCore/Source/UnrealUSDWrapper/Public/UsdWrappers/UsdPrim.h b/Engine/Plugins/Runtime/USDCore/Source/UnrealUSDWrapper/Public/UsdWrappers/UsdPrim.h
--- a/Engine/Plugins/Runtime/USDCore/Source/UnrealUSDWrapper/Public/UsdWrappers/UsdPrim.h
+++ b/Engine/Plugins/Runtime/USDCore/Source/UnrealUSDWrapper/Public/UsdWrappers/UsdPrim.h
@@ -86,6 +86,7 @@ namespace UE
 		bool IsValid() const;
 		bool IsPseudoRoot() const;
 		bool IsModel() const;
+		bool IsDefined() const;
 		bool IsGroup() const;
 
 		TArray<FName> GetAppliedSchemas() const;
