USD: Fixed crash when opening a stage then undoing and redoing the stage open.

diff --git a/Engine/Plugins/Importers/USDImporter/Source/USDStage/Private/USDLevelSequenceHelper.cpp b/Engine/Plugins/Importers/USDImporter/Source/USDStage/Private/USDLevelSequenceHelper.cpp
--- a/Engine/Plugins/Importers/USDImporter/Source/USDStage/Private/USDLevelSequenceHelper.cpp
+++ b/Engine/Plugins/Importers/USDImporter/Source/USDStage/Private/USDLevelSequenceHelper.cpp
@@ -1031,7 +1031,7 @@ bool FUsdLevelSequenceHelperImpl::Serialize(FArchive& Ar)
 	Ar << StageActorBinding;
 
 	// Always keep SequenceHierarchyCache up-to-date given that it can't be serialized itself
-	if (Ar.IsLoading() && MainLevelSequence)
+	if (Ar.IsLoading() && MainLevelSequence && MainLevelSequence->GetMovieScene())
 	{
 		UMovieSceneCompiledDataManager::CompileHierarchy(MainLevelSequence, &SequenceHierarchyCache, EMovieSceneServerClientMask::All);
 	}
