USD: Add option for explicitly controlling prim collapsing regardless of kind.

diff --git a/Engine/Plugins/Importers/USDImporter/Source/USDSchemas/Private/USDGeomXformableTranslator.cpp b/Engine/Plugins/Importers/USDImporter/Source/USDSchemas/Private/USDGeomXformableTranslator.cpp
--- a/Engine/Plugins/Importers/USDImporter/Source/USDSchemas/Private/USDGeomXformableTranslator.cpp
+++ b/Engine/Plugins/Importers/USDImporter/Source/USDSchemas/Private/USDGeomXformableTranslator.cpp
@@ -1029,6 +1029,45 @@ void FUsdGeomXformableTranslator::CreateAlternativeDrawModeAssets(EUsdDrawMode D
 	HandleCardFace(GeomModelAPI.GetModelCardTextureZNegAttr());
 }
 
+namespace UE::UsdXformableTranslatorImpl::Private
+{
+	bool PrimCollapses(const pxr::UsdPrim& Prim, EUsdDefaultKind KindsToCollapse)
+	{
+		UsdUtils::ECollapsingPreference Preference = UsdUtils::GetCollapsingPreference(Prim);
+		switch (Preference)
+		{
+			case UsdUtils::ECollapsingPreference::Allow:
+			{
+				return true;
+				break;
+			}
+			case UsdUtils::ECollapsingPreference::ByKind:
+			{
+				EUsdDefaultKind PrimKind = UsdUtils::GetDefaultKind(Prim);
+
+				// Note that this is false if PrimKind is None
+				const bool bPrimKindCollapses = EnumHasAnyFlags(KindsToCollapse, PrimKind);
+				bool bCanBeCollapsed = bPrimKindCollapses || (PrimKind == EUsdDefaultKind::None && GCollapsePrimsWithoutKind);
+
+				if (!bCanBeCollapsed)
+				{
+					pxr::UsdModelAPI Model{Prim};
+					bCanBeCollapsed = bCanBeCollapsed || Model.IsKind(pxr::TfToken("prop"), pxr::UsdModelAPI::KindValidationNone);
+				}
+
+				return bCanBeCollapsed;
+			}
+			case UsdUtils::ECollapsingPreference::Never:	// Fallthrough
+			default:
+			{
+				break;
+			}
+		}
+
+		return false;
+	}
+}	 // namespace UE::UsdXformableTranslatorImpl::Private
+
 bool FUsdGeomXformableTranslator::CollapsesChildren(ECollapsingType CollapsingType) const
 {
 	TRACE_CPUPROFILER_EVENT_SCOPE(FUsdGeomXformableTranslator::CollapsesChildren);
@@ -1062,20 +1101,7 @@ bool FUsdGeomXformableTranslator::CollapsesChildren(ECollapsingType CollapsingTy
 
 	if (Model)
 	{
-		EUsdDefaultKind PrimKind = UsdUtils::GetDefaultKind(Prim);
-
-		// Note that this is false if PrimKind is None
-		const bool bPrimKindShouldCollapse = EnumHasAnyFlags(Context->KindsToCollapse, PrimKind);
-
-		// This indicates whether we *want* to collapse
-		bCollapsesChildren = Context->KindsToCollapse != EUsdDefaultKind::None
-							 && (bPrimKindShouldCollapse || (PrimKind == EUsdDefaultKind::None && GCollapsePrimsWithoutKind));
-
-		if (!bCollapsesChildren)
-		{
-			// Temp support for the prop kind
-			bCollapsesChildren = Model.IsKind(pxr::TfToken("prop"), pxr::UsdModelAPI::KindValidationNone);
-		}
+		bCollapsesChildren = UE::UsdXformableTranslatorImpl::Private::PrimCollapses(Prim, Context->KindsToCollapse);
 
 		if (bCollapsesChildren)
 		{
@@ -1107,10 +1133,7 @@ bool FUsdGeomXformableTranslator::CanBeCollapsed(ECollapsingType CollapsingType)
 		return false;
 	}
 
-	EUsdDefaultKind PrimKind = UsdUtils::GetDefaultKind(GetPrim());
-	// Note that this is false if PrimKind is None
-	const bool bThisPrimCanCollapse = EnumHasAnyFlags(Context->KindsToCollapse, PrimKind);
-	return bThisPrimCanCollapse || (PrimKind == EUsdDefaultKind::None && GCollapsePrimsWithoutKind);
+	return UE::UsdXformableTranslatorImpl::Private::PrimCollapses(UsdPrim, Context->KindsToCollapse);
 }
 
 TSet<UE::FSdfPath> FUsdGeomXformableTranslator::CollectAuxiliaryPrims() const
diff --git a/Engine/Plugins/Importers/USDImporter/Source/USDStage/Private/USDStageActor.cpp b/Engine/Plugins/Importers/USDImporter/Source/USDStage/Private/USDStageActor.cpp
--- a/Engine/Plugins/Importers/USDImporter/Source/USDStage/Private/USDStageActor.cpp
+++ b/Engine/Plugins/Importers/USDImporter/Source/USDStage/Private/USDStageActor.cpp
@@ -1402,6 +1402,8 @@ void AUsdStageActor::OnUsdObjectsChanged(const UsdUtils::FObjectChangesByPath& I
 					// for the prim (as it may now need/stop needing an alternate draw mode component)
 					UnrealIdentifiers::ModelDrawMode,
 					UnrealIdentifiers::ModelApplyDrawMode,
+					// Collapsing changes need to resync to build the new collapsed/uncollapsed meshes
+					*UsdToUnreal::ConvertToken(UnrealIdentifiers::UnrealCollapsingAttr),
 					// Physics collision attribute change needs to rebuild the collision shape
 					*UsdToUnreal::ConvertToken(pxr::UsdPhysicsTokens->physicsCollisionEnabled),
 					*UsdToUnreal::ConvertToken(pxr::UsdPhysicsTokens->physicsApproximation)};
diff --git a/Engine/Plugins/Importers/USDImporter/Source/USDStageEditor/Private/SUSDStageTreeView.cpp b/Engine/Plugins/Importers/USDImporter/Source/USDStageEditor/Private/SUSDStageTreeView.cpp
--- a/Engine/Plugins/Importers/USDImporter/Source/USDStageEditor/Private/SUSDStageTreeView.cpp
+++ b/Engine/Plugins/Importers/USDImporter/Source/USDStageEditor/Private/SUSDStageTreeView.cpp
@@ -798,6 +798,14 @@ TSharedPtr<SWidget> SUsdStageTreeView::ConstructPrimContextMenu()
 				}
 			))
 		);
+
+		PrimOptions.AddSubMenu(
+			LOCTEXT("Collapsing_Text", "Collapsing..."),
+			FText::GetEmpty(),
+			FNewMenuDelegate::CreateSP(this, &SUsdStageTreeView::FillCollapsingSubmenu),
+			bInOpenSubMenuOnClick,
+			FSlateIcon(FAppStyle::GetAppStyleSetName(), "Icons.Merge")
+		);
 	}
 	PrimOptions.EndSection();
 
@@ -1108,6 +1116,30 @@ void SUsdStageTreeView::OnRenamePrim()
 	}
 }
 
+void SUsdStageTreeView::OnSetCollapsingPreference(UsdUtils::ECollapsingPreference Preference)
+{
+	if (!UsdStage)
+	{
+		return;
+	}
+
+	TArray<FUsdPrimViewModelRef> MySelectedItems = GetSelectedItems();
+	for (FUsdPrimViewModelRef SelectedItem : MySelectedItems)
+	{
+		if (UsdUtils::NotifyIfInstanceProxy(SelectedItem->UsdPrim))
+		{
+			return;
+		}
+	}
+
+	FScopedTransaction Transaction(LOCTEXT("SetCollapsingPreferenceTransaction", "Set collapsing preference"));
+
+	for (FUsdPrimViewModelRef SelectedItem : MySelectedItems)
+	{
+		UsdUtils::SetCollapsingPreference(SelectedItem->UsdPrim, Preference);
+	}
+}
+
 void SUsdStageTreeView::OnAddReference()
 {
 	if (!UsdStage || !UsdStage.IsEditTargetValid())
@@ -1453,6 +1485,27 @@ bool SUsdStageTreeView::DoesPrimHaveSpecOnLocalLayerStack() const
 	return false;
 }
 
+bool SUsdStageTreeView::DoSelectedPrimsHaveCollapsingPreference(UsdUtils::ECollapsingPreference TargetPreference) const
+{
+	TArray<FUsdPrimViewModelRef> MySelectedItems = GetSelectedItems();
+
+	if (MySelectedItems.Num() == 0)
+	{
+		return false;
+	}
+
+	for (FUsdPrimViewModelRef SelectedItem : MySelectedItems)
+	{
+		UsdUtils::ECollapsingPreference PrimPreference = UsdUtils::GetCollapsingPreference(SelectedItem->UsdPrim);
+		if (PrimPreference != TargetPreference)
+		{
+			return false;
+		}
+	}
+
+	return true;
+}
+
 void SUsdStageTreeView::Tick(const FGeometry& AllottedGeometry, const double InCurrentTime, const float InDeltaTime)
 {
 	// Restore expansion states.
@@ -1639,6 +1692,48 @@ void SUsdStageTreeView::FillDuplicateSubmenu(FMenuBuilder& MenuBuilder)
 	);
 }
 
+void SUsdStageTreeView::FillCollapsingSubmenu(FMenuBuilder& MenuBuilder)
+{
+	MenuBuilder.AddMenuEntry(
+		LOCTEXT("AllowCollapse", "Allow collapsing"),
+		LOCTEXT("AllowCollapse_ToolTip", "Allow this prim to be collapsed and to try collapsing its subtree, regardless of its kind"),
+		FSlateIcon(),
+		FUIAction(
+			FExecuteAction::CreateSP(this, &SUsdStageTreeView::OnSetCollapsingPreference, UsdUtils::ECollapsingPreference::Allow),
+			FCanExecuteAction::CreateSP(this, &SUsdStageTreeView::DoesPrimExistOnStage),
+			FIsActionChecked::CreateSP(this, &SUsdStageTreeView::DoSelectedPrimsHaveCollapsingPreference, UsdUtils::ECollapsingPreference::Allow)
+		),
+		NAME_None,
+		EUserInterfaceActionType::RadioButton
+	);
+
+	MenuBuilder.AddMenuEntry(
+		LOCTEXT("CollapseOnKind", "Collapse according to kind"),
+		LOCTEXT("CollapseOnKind_ToolTip", "Collapse this prim only if its kind is allowed to collapse (default)"),
+		FSlateIcon(),
+		FUIAction(
+			FExecuteAction::CreateSP(this, &SUsdStageTreeView::OnSetCollapsingPreference, UsdUtils::ECollapsingPreference::ByKind),
+			FCanExecuteAction::CreateSP(this, &SUsdStageTreeView::DoesPrimExistOnStage),
+			FIsActionChecked::CreateSP(this, &SUsdStageTreeView::DoSelectedPrimsHaveCollapsingPreference, UsdUtils::ECollapsingPreference::ByKind)
+		),
+		NAME_None,
+		EUserInterfaceActionType::RadioButton
+	);
+
+	MenuBuilder.AddMenuEntry(
+		LOCTEXT("NeverCollapse", "Never collapse"),
+		LOCTEXT("NeverCollapse_ToolTip", "Never collapse this prim, regardless of its kind"),
+		FSlateIcon(),
+		FUIAction(
+			FExecuteAction::CreateSP(this, &SUsdStageTreeView::OnSetCollapsingPreference, UsdUtils::ECollapsingPreference::Never),
+			FCanExecuteAction::CreateSP(this, &SUsdStageTreeView::DoesPrimExistOnStage),
+			FIsActionChecked::CreateSP(this, &SUsdStageTreeView::DoSelectedPrimsHaveCollapsingPreference, UsdUtils::ECollapsingPreference::Never)
+		),
+		NAME_None,
+		EUserInterfaceActionType::RadioButton
+	);
+}
+
 void SUsdStageTreeView::FillAddSchemaSubmenu(FMenuBuilder& MenuBuilder)
 {
 	static TArray<TSharedPtr<FString>> AllKnownSchemas = {
diff --git a/Engine/Plugins/Importers/USDImporter/Source/USDStageEditor/Private/SUSDStageTreeView.h b/Engine/Plugins/Importers/USDImporter/Source/USDStageEditor/Private/SUSDStageTreeView.h
--- a/Engine/Plugins/Importers/USDImporter/Source/USDStageEditor/Private/SUSDStageTreeView.h
+++ b/Engine/Plugins/Importers/USDImporter/Source/USDStageEditor/Private/SUSDStageTreeView.h
@@ -14,6 +14,10 @@ namespace UE
 {
 	class FUsdPrim;
 }
+namespace UsdUtils
+{
+	enum class ECollapsingPreference : uint8;
+}
 
 #if USE_USD_SDK
 
@@ -65,6 +69,7 @@ private:
 	void OnToggleAllPayloads(EPayloadsTrigger PayloadsTrigger);
 
 	void FillDuplicateSubmenu(FMenuBuilder& MenuBuilder);
+	void FillCollapsingSubmenu(FMenuBuilder& MenuBuilder);
 	void FillAddSchemaSubmenu(FMenuBuilder& MenuBuilder);
 	void FillRemoveSchemaSubmenu(FMenuBuilder& MenuBuilder);
 
@@ -75,6 +80,7 @@ private:
 	void OnDuplicatePrim(EUsdDuplicateType DuplicateType);
 	void OnDeletePrim();
 	void OnRenamePrim();
+	void OnSetCollapsingPreference(UsdUtils::ECollapsingPreference Preference);
 
 	void OnAddReference();
 	void OnClearReferences();
@@ -92,6 +98,7 @@ private:
 	bool DoesPrimExistOnStage() const;
 	bool DoesPrimExistOnEditTarget() const;
 	bool DoesPrimHaveSpecOnLocalLayerStack() const;
+	bool DoSelectedPrimsHaveCollapsingPreference(UsdUtils::ECollapsingPreference Preference) const;
 
 	void RequestExpansionStateRestore();
 	virtual void Tick( const FGeometry& AllottedGeometry, const double InCurrentTime, const float InDeltaTime ) override;
diff --git a/Engine/Plugins/Runtime/USDCore/Resources/UsdResources/Linux/plugins/unreal/resources/plugInfo.json b/Engine/Plugins/Runtime/USDCore/Resources/UsdResources/Linux/plugins/unreal/resources/plugInfo.json
--- a/Engine/Plugins/Runtime/USDCore/Resources/UsdResources/Linux/plugins/unreal/resources/plugInfo.json
+++ b/Engine/Plugins/Runtime/USDCore/Resources/UsdResources/Linux/plugins/unreal/resources/plugInfo.json
@@ -6,6 +6,16 @@
         {
             "Info": {
                 "Types": {
+                    "unrealCollapsingAPI": {
+                        "alias": {
+                            "UsdSchemaBase": "CollapsingAPI"
+                        }, 
+                        "autoGenerated": true, 
+                        "bases": [
+                            "UsdAPISchemaBase"
+                        ], 
+                        "schemaKind": "singleApplyAPI"
+                    }, 
                     "unrealControlRigAPI": {
                         "alias": {
                             "UsdSchemaBase": "ControlRigAPI"
diff --git a/Engine/Plugins/Runtime/USDCore/Resources/UsdResources/Mac/plugins/unreal/resources/plugInfo.json b/Engine/Plugins/Runtime/USDCore/Resources/UsdResources/Mac/plugins/unreal/resources/plugInfo.json
--- a/Engine/Plugins/Runtime/USDCore/Resources/UsdResources/Mac/plugins/unreal/resources/plugInfo.json
+++ b/Engine/Plugins/Runtime/USDCore/Resources/UsdResources/Mac/plugins/unreal/resources/plugInfo.json
@@ -6,6 +6,16 @@
         {
             "Info": {
                 "Types": {
+                    "unrealCollapsingAPI": {
+                        "alias": {
+                            "UsdSchemaBase": "CollapsingAPI"
+                        }, 
+                        "autoGenerated": true, 
+                        "bases": [
+                            "UsdAPISchemaBase"
+                        ], 
+                        "schemaKind": "singleApplyAPI"
+                    }, 
                     "unrealControlRigAPI": {
                         "alias": {
                             "UsdSchemaBase": "ControlRigAPI"
diff --git a/Engine/Plugins/Runtime/USDCore/Resources/UsdResources/Win64/plugins/unreal/resources/plugInfo.json b/Engine/Plugins/Runtime/USDCore/Resources/UsdResources/Win64/plugins/unreal/resources/plugInfo.json
--- a/Engine/Plugins/Runtime/USDCore/Resources/UsdResources/Win64/plugins/unreal/resources/plugInfo.json
+++ b/Engine/Plugins/Runtime/USDCore/Resources/UsdResources/Win64/plugins/unreal/resources/plugInfo.json
@@ -6,6 +6,16 @@
         {
             "Info": {
                 "Types": {
+                    "unrealCollapsingAPI": {
+                        "alias": {
+                            "UsdSchemaBase": "CollapsingAPI"
+                        }, 
+                        "autoGenerated": true, 
+                        "bases": [
+                            "UsdAPISchemaBase"
+                        ], 
+                        "schemaKind": "singleApplyAPI"
+                    }, 
                     "unrealControlRigAPI": {
                         "alias": {
                             "UsdSchemaBase": "ControlRigAPI"
diff --git a/Engine/Plugins/Runtime/USDCore/Source/USDUtilities/Private/USDConversionUtils.cpp b/Engine/Plugins/Runtime/USDCore/Source/USDUtilities/Private/USDConversionUtils.cpp
--- a/Engine/Plugins/Runtime/USDCore/Source/USDUtilities/Private/USDConversionUtils.cpp
+++ b/Engine/Plugins/Runtime/USDCore/Source/USDUtilities/Private/USDConversionUtils.cpp
@@ -9,6 +9,7 @@
 #include "USDDuplicateType.h"
 #include "USDErrorUtils.h"
 #include "USDGeomMeshConversion.h"
+#include "USDIntegrationUtils.h"
 #include "USDLayerUtils.h"
 #include "USDLog.h"
 #include "USDObjectUtils.h"
@@ -1416,6 +1417,83 @@ bool UsdUtils::SetDefaultKind(pxr::UsdPrim& Prim, EUsdDefaultKind NewKind)
 	return IUsdPrim::SetKind(Prim, NewKindToken);
 }
 
+UsdUtils::ECollapsingPreference UsdUtils::GetCollapsingPreference(const pxr::UsdPrim& Prim)
+{
+	if (Prim)
+	{
+		if (UsdUtils::PrimHasSchema(Prim, UnrealIdentifiers::UnrealCollapsingAPI))
+		{
+			FScopedUsdAllocs UsdAllocs;
+
+			if (pxr::UsdAttribute Attr = Prim.GetAttribute(UnrealIdentifiers::UnrealCollapsingAttr))
+			{
+				pxr::TfToken Value;
+				if (Attr.Get(&Value))
+				{
+					if (Value == UnrealIdentifiers::CollapsingAllow)
+					{
+						return ECollapsingPreference::Allow;
+					}
+					else if (Value == UnrealIdentifiers::CollapsingNever)
+					{
+						return ECollapsingPreference::Never;
+					}
+				}
+			}
+		}
+	}
+
+	return ECollapsingPreference::ByKind;
+}
+
+bool UsdUtils::SetCollapsingPreference(const pxr::UsdPrim& Prim, UsdUtils::ECollapsingPreference NewPreference)
+{
+	if (!Prim)
+	{
+		return false;
+	}
+
+	FScopedUsdAllocs UsdAllocs;
+
+	pxr::SdfChangeBlock ChangeBlock;
+
+	const bool bAppliedSchema = UsdUtils::ApplySchema(Prim, UnrealIdentifiers::UnrealCollapsingAPI);
+	if (!bAppliedSchema)
+	{
+		return false;
+	}
+
+	const pxr::SdfVariability Variability = pxr::SdfVariabilityUniform;
+	if (pxr::UsdAttribute Attr = Prim.CreateAttribute(UnrealIdentifiers::UnrealCollapsingAttr, pxr::SdfValueTypeNames->Token, Variability))
+	{
+		switch (NewPreference)
+		{
+			case ECollapsingPreference::Allow:
+			{
+				return Attr.Set(UnrealIdentifiers::CollapsingAllow);
+				break;
+			}
+			case ECollapsingPreference::ByKind:
+			{
+				return Attr.Set(UnrealIdentifiers::CollapsingByKind);
+				break;
+			}
+			case ECollapsingPreference::Never:
+			{
+				return Attr.Set(UnrealIdentifiers::CollapsingNever);
+				break;
+			}
+			default:
+			{
+				ensure(false);
+				break;
+			}
+		}
+	}
+
+	return false;
+}
+
 EUsdDrawMode UsdUtils::GetAppliedDrawMode(const pxr::UsdPrim& Prim)
 {
 	// Reference: https://openusd.org/release/api/class_usd_geom_model_a_p_i.html#UsdGeomModelAPI_drawMode
diff --git a/Engine/Plugins/Runtime/USDCore/Source/USDUtilities/Public/USDConversionUtils.h b/Engine/Plugins/Runtime/USDCore/Source/USDUtilities/Public/USDConversionUtils.h
--- a/Engine/Plugins/Runtime/USDCore/Source/USDUtilities/Public/USDConversionUtils.h
+++ b/Engine/Plugins/Runtime/USDCore/Source/USDUtilities/Public/USDConversionUtils.h
@@ -294,6 +294,23 @@ namespace UsdUtils
 	 */
 	USDUTILITIES_API bool SetDefaultKind(pxr::UsdPrim& Prim, EUsdDefaultKind NewKind);
 
+	enum class ECollapsingPreference : uint8
+	{
+		Allow,
+		ByKind,
+		Never
+	};
+
+	/**
+	 * Gets how this prim wants to be collapsed: Whether we should allow/never try to do it, or whether we should check its kind for that purpose
+	 */
+	USDUTILITIES_API ECollapsingPreference GetCollapsingPreference(const pxr::UsdPrim& Prim);
+
+	/**
+	 * Sets how this prim wants to be collapsed: Whether we should always/never try to do it, or whether we should check its kind for that purpose
+	 */
+	USDUTILITIES_API bool SetCollapsingPreference(const pxr::UsdPrim& Prim, ECollapsingPreference NewPreference);
+
 	/**
 	 * Returns whether the prim has the UsdGeomModelAPI schema and should be drawn with one of the alternative draw modes, such as cards or bounds.
 	 * Will return EUsdDrawMode::Default in case the prim should be drawn as usual instead, or in case of error.
@@ -407,7 +424,7 @@ namespace UsdUtils
 	 * Removes any numbered suffix, followed by any number of underscores (e.g. Asset_2, Asset__232_31 or Asset94 all become 'Asset'), making
 	 * sure the string is kept at least one character long. Returns true if it removed anything.
 	 */
-	 UE_DEPRECATED(5.5, "This function has been moved to USDObjectUtils.h, within the USDClasses module")
+	UE_DEPRECATED(5.5, "This function has been moved to USDObjectUtils.h, within the USDClasses module")
 	USDUTILITIES_API bool RemoveNumberedSuffix(FString& Prefix);
 
 	/**
@@ -417,7 +434,7 @@ namespace UsdUtils
 	 * @param UsedNames - Strings that cannot be used for the result
 	 * @return Modified Name so that it doesn't match anything in UsedNames (e.g. "MyName" again, or "MyName_0" or "MyName_423")
 	 */
-	 UE_DEPRECATED(5.5, "This function has been moved to USDObjectUtils.h, within the USDClasses module")
+	UE_DEPRECATED(5.5, "This function has been moved to USDObjectUtils.h, within the USDClasses module")
 	USDUTILITIES_API FString GetUniqueName(FString Name, const TSet<FString>& UsedNames);
 
 #if USE_USD_SDK
diff --git a/Engine/Plugins/Runtime/USDCore/Source/UnrealUSDWrapper/Private/UnrealUSDWrapper.cpp b/Engine/Plugins/Runtime/USDCore/Source/UnrealUSDWrapper/Private/UnrealUSDWrapper.cpp
--- a/Engine/Plugins/Runtime/USDCore/Source/UnrealUSDWrapper/Private/UnrealUSDWrapper.cpp
+++ b/Engine/Plugins/Runtime/USDCore/Source/UnrealUSDWrapper/Private/UnrealUSDWrapper.cpp
@@ -166,6 +166,12 @@ namespace UnrealIdentifiers
 	const TfToken UnrealAssetType = TfToken("unreal:assetType");
 	const TfToken UnrealExportTime = TfToken("unreal:exportTime");
 	const TfToken UnrealEngineVersion = TfToken("unreal:engineVersion");
+
+	const TfToken UnrealCollapsingAPI = TfToken("CollapsingAPI");
+	const TfToken UnrealCollapsingAttr = TfToken("unreal:collapsing");
+	const TfToken CollapsingAllow = TfToken("allow");
+	const TfToken CollapsingByKind = TfToken("byKind");
+	const TfToken CollapsingNever = TfToken("never");
 }	 // namespace UnrealIdentifiers
 
 bool IUsdPrim::IsValidPrimName(const FString& Name, FText& OutReason)
diff --git a/Engine/Plugins/Runtime/USDCore/Source/UnrealUSDWrapper/Public/UnrealUSDWrapper.h b/Engine/Plugins/Runtime/USDCore/Source/UnrealUSDWrapper/Public/UnrealUSDWrapper.h
--- a/Engine/Plugins/Runtime/USDCore/Source/UnrealUSDWrapper/Public/UnrealUSDWrapper.h
+++ b/Engine/Plugins/Runtime/USDCore/Source/UnrealUSDWrapper/Public/UnrealUSDWrapper.h
@@ -30,27 +30,27 @@
 
 #if USE_USD_SDK
 PXR_NAMESPACE_OPEN_SCOPE
-    class GfMatrix4d;
-    class SdfPath;
-    class TfToken;
-    class UsdAttribute;
-    class UsdGeomMesh;
-    class UsdPrim;
-    class UsdStage;
-    class UsdStageCache;
-
-    template<typename T>
-    class TfRefPtr;
-    template<typename T>
-    class TfWeakPtr;
-
-    using UsdStageRefPtr = TfRefPtr<UsdStage>;
-    using UsdStageWeakPtr = TfWeakPtr<UsdStage>;
-
-    using SdfLayerRefPtr = TfRefPtr<SdfLayer>;
-    using SdfLayerWeakPtr = TfWeakPtr<SdfLayer>;
+	class GfMatrix4d;
+	class SdfPath;
+	class TfToken;
+	class UsdAttribute;
+	class UsdGeomMesh;
+	class UsdPrim;
+	class UsdStage;
+	class UsdStageCache;
+
+	template<typename T>
+	class TfRefPtr;
+	template<typename T>
+	class TfWeakPtr;
+
+	using UsdStageRefPtr = TfRefPtr<UsdStage>;
+	using UsdStageWeakPtr = TfWeakPtr<UsdStage>;
+
+	using SdfLayerRefPtr = TfRefPtr<SdfLayer>;
+	using SdfLayerWeakPtr = TfWeakPtr<SdfLayer>;
 PXR_NAMESPACE_CLOSE_SCOPE
-#endif      // #if USE_USD_SDK
+#endif	  // #if USE_USD_SDK
 
 class IUsdPrim;
 class FUsdDiagnosticDelegate;
@@ -469,6 +469,12 @@ namespace UnrealIdentifiers
 	extern UNREALUSDWRAPPER_API const pxr::TfToken UnrealAssetType;
 	extern UNREALUSDWRAPPER_API const pxr::TfToken UnrealExportTime;
 	extern UNREALUSDWRAPPER_API const pxr::TfToken UnrealEngineVersion;
+
+	extern UNREALUSDWRAPPER_API const pxr::TfToken UnrealCollapsingAPI;
+	extern UNREALUSDWRAPPER_API const pxr::TfToken UnrealCollapsingAttr;
+	extern UNREALUSDWRAPPER_API const pxr::TfToken CollapsingAllow;
+	extern UNREALUSDWRAPPER_API const pxr::TfToken CollapsingByKind;
+	extern UNREALUSDWRAPPER_API const pxr::TfToken CollapsingNever;
 #endif	  // #if USE_USD_SDK
 
 	extern UNREALUSDWRAPPER_API const TCHAR* LayerSavedComment;
@@ -539,4 +545,3 @@ struct UNREALUSDWRAPPER_API FUsdDelegates
 	static FUsdImportDelegate OnPreUsdImport;
 	static FUsdImportDelegate OnPostUsdImport;
 };
-
