USD: Add a File->Export flattened layer stack option.

diff --git a/Engine/Plugins/Importers/USDImporter/Source/USDStageEditor/Private/SUSDStage.cpp b/Engine/Plugins/Importers/USDImporter/Source/USDStageEditor/Private/SUSDStage.cpp
--- a/Engine/Plugins/Importers/USDImporter/Source/USDStageEditor/Private/SUSDStage.cpp
+++ b/Engine/Plugins/Importers/USDImporter/Source/USDStageEditor/Private/SUSDStage.cpp
@@ -9,6 +9,7 @@
 #include "UnrealUSDWrapper.h"
 #include "USDClassesModule.h"
 #include "USDConversionUtils.h"
+#include "USDErrorUtils.h"
 #include "USDLayerUtils.h"
 #include "USDLog.h"
 #include "USDProjectSettings.h"
@@ -1141,7 +1142,7 @@ void SUsdStage::FillExportSubMenu(FMenuBuilder& MenuBuilder)
 	const FString FilePath;
 
 	MenuBuilder.AddMenuEntry(
-		LOCTEXT("ExportAll", "Export all layers..."),
+		LOCTEXT("ExportAll", "All layers..."),
 		LOCTEXT("ExportAll_ToolTip", "Exports copies of all file-based layers in the stage's layer stack to a new folder"),
 		FSlateIcon(),
 		FUIAction(
@@ -1171,7 +1172,7 @@ void SUsdStage::FillExportSubMenu(FMenuBuilder& MenuBuilder)
 	);
 
 	MenuBuilder.AddMenuEntry(
-		LOCTEXT("ExportFlattened", "Export flattened stage..."),
+		LOCTEXT("ExportFlattened", "Flattened stage..."),
 		LOCTEXT("ExportFlattened_ToolTip", "Flattens the current stage to a single USD layer and exports it as a new USD file"),
 		FSlateIcon(),
 		FUIAction(
@@ -1199,6 +1200,36 @@ void SUsdStage::FillExportSubMenu(FMenuBuilder& MenuBuilder)
 		NAME_None,
 		EUserInterfaceActionType::Button
 	);
+
+	MenuBuilder.AddMenuEntry(
+		LOCTEXT("ExportStack", "Flattened layer stack..."),
+		LOCTEXT("ExportStack_ToolTip", "Flattens the current stage's local layer stack to a single layer, preserving most composition arcs"),
+		FSlateIcon(),
+		FUIAction(
+			FExecuteAction::CreateLambda(
+				[this]()
+				{
+					FileExportFlattenedLayerStack();
+				}
+			),
+			FCanExecuteAction::CreateLambda(
+				[this]()
+				{
+					if (const AUsdStageActor* StageActor = ViewModel.UsdStageActor.Get())
+					{
+						if (UE::FUsdStage Stage = StageActor->GetUsdStage())
+						{
+							return true;
+						}
+					}
+
+					return false;
+				}
+			)
+		),
+		NAME_None,
+		EUserInterfaceActionType::Button
+	);
 }
 
 void SUsdStage::FillStageStateSubMenu(FMenuBuilder& MenuBuilder)
@@ -2427,6 +2458,64 @@ void SUsdStage::FileExportFlattenedStage(const FString& OutputLayer)
 	}
 }
 
+void SUsdStage::FileExportFlattenedLayerStack(const FString& OutputLayer)
+{
+	const AUsdStageActor* StageActor = ViewModel.UsdStageActor.Get();
+	if (!StageActor)
+	{
+		return;
+	}
+
+	UE::FUsdStage UsdStage = StageActor->GetUsdStage();
+	if (!UsdStage)
+	{
+		return;
+	}
+
+	FString OutputLayerCopy = OutputLayer;
+	if (OutputLayerCopy.IsEmpty())
+	{
+		TOptional<FString> UsdFilePath = UsdUtils::BrowseUsdFile(UsdUtils::EBrowseFileMode::Save);
+		if (!UsdFilePath.IsSet())
+		{
+			return;
+		}
+
+		OutputLayerCopy = UsdFilePath.GetValue();
+	}
+
+	double StartTime = FPlatformTime::Cycles64();
+
+	UsdUtils::StartMonitoringErrors();
+	UE::FSdfLayer FlattenedLayer = UsdUtils::FlattenLayerStack(UsdStage);
+	if (UsdUtils::ShowErrorsAndStopMonitoring())
+	{
+		return;
+	}
+
+	const bool bResult = FlattenedLayer.Export(*OutputLayerCopy);
+	if (!bResult)
+	{
+		UE_LOG(LogUsd, Warning, TEXT("Failed to export flattened USD Stage to path '%s'!"), *OutputLayerCopy);
+	}
+
+	// Send analytics
+	if (FEngineAnalytics::IsAvailable())
+	{
+		bool bAutomated = false;
+		double ElapsedSeconds = FPlatformTime::ToSeconds64(FPlatformTime::Cycles64() - StartTime);
+		FString Extension = FPaths::GetExtension(OutputLayerCopy);
+		IUsdClassesModule::SendAnalytics(
+			{},
+			TEXT("ExportStageFlattenedLayerStack"),
+			bAutomated,
+			ElapsedSeconds,
+			UsdUtils::GetUsdStageNumFrames(UsdStage),
+			Extension
+		);
+	}
+}
+
 void SUsdStage::FileReload()
 {
 	FScopedTransaction Transaction(LOCTEXT("ReloadTransaction", "Reload USD stage"));
diff --git a/Engine/Plugins/Importers/USDImporter/Source/USDStageEditor/Private/SUSDStage.h b/Engine/Plugins/Importers/USDImporter/Source/USDStageEditor/Private/SUSDStage.h
--- a/Engine/Plugins/Importers/USDImporter/Source/USDStageEditor/Private/SUSDStage.h
+++ b/Engine/Plugins/Importers/USDImporter/Source/USDStageEditor/Private/SUSDStage.h
@@ -66,6 +66,7 @@ public:
 	void FileSave(const FString& OutputFileIfUnsaved = {});
 	void FileExportAllLayers(const FString& OutputDirectory = {});
 	void FileExportFlattenedStage(const FString& OutputLayer = {});
+	void FileExportFlattenedLayerStack(const FString& OutputLayer = {});
 	void FileReload();
 	void FileReset();
 	void FileClose();
diff --git a/Engine/Plugins/Importers/USDImporter/Source/USDStageEditor/Private/USDStageEditorBlueprintLibrary.cpp b/Engine/Plugins/Importers/USDImporter/Source/USDStageEditor/Private/USDStageEditorBlueprintLibrary.cpp
--- a/Engine/Plugins/Importers/USDImporter/Source/USDStageEditor/Private/USDStageEditorBlueprintLibrary.cpp
+++ b/Engine/Plugins/Importers/USDImporter/Source/USDStageEditor/Private/USDStageEditorBlueprintLibrary.cpp
@@ -167,6 +167,12 @@ void UUsdStageEditorBlueprintLibrary::FileExportFlattenedStage(const FString& Ou
 	StageEditorModule.FileExportFlattenedStage(OutputLayer);
 }
 
+void UUsdStageEditorBlueprintLibrary::FileExportFlattenedLayerStack(const FString& OutputLayer)
+{
+	IUsdStageEditorModule& StageEditorModule = FModuleManager::GetModuleChecked<IUsdStageEditorModule>("USDStageEditor");
+	StageEditorModule.FileExportFlattenedLayerStack(OutputLayer);
+}
+
 void UUsdStageEditorBlueprintLibrary::FileReload()
 {
 	IUsdStageEditorModule& StageEditorModule = FModuleManager::GetModuleChecked<IUsdStageEditorModule>("USDStageEditor");
diff --git a/Engine/Plugins/Importers/USDImporter/Source/USDStageEditor/Private/USDStageEditorModule.cpp b/Engine/Plugins/Importers/USDImporter/Source/USDStageEditor/Private/USDStageEditorModule.cpp
--- a/Engine/Plugins/Importers/USDImporter/Source/USDStageEditor/Private/USDStageEditorModule.cpp
+++ b/Engine/Plugins/Importers/USDImporter/Source/USDStageEditor/Private/USDStageEditorModule.cpp
@@ -627,6 +627,16 @@ void IUsdStageEditorModule::FileExportFlattenedStage(const FString& OutputLayer)
 #endif	  // USE_USD_SDK
 }
 
+void IUsdStageEditorModule::FileExportFlattenedLayerStack(const FString& OutputLayer) const
+{
+#if USE_USD_SDK
+	if (TSharedPtr<SUsdStage> UsdStageEditor = UE::UsdStageEditorModule::Private::GetUsdStageEditor())
+	{
+		UsdStageEditor->FileExportFlattenedLayerStack(OutputLayer);
+	}
+#endif	  // USE_USD_SDK
+}
+
 void IUsdStageEditorModule::FileReload() const
 {
 #if USE_USD_SDK
diff --git a/Engine/Plugins/Importers/USDImporter/Source/USDStageEditor/Public/USDStageEditorBlueprintLibrary.h b/Engine/Plugins/Importers/USDImporter/Source/USDStageEditor/Public/USDStageEditorBlueprintLibrary.h
--- a/Engine/Plugins/Importers/USDImporter/Source/USDStageEditor/Public/USDStageEditorBlueprintLibrary.h
+++ b/Engine/Plugins/Importers/USDImporter/Source/USDStageEditor/Public/USDStageEditorBlueprintLibrary.h
@@ -156,13 +156,22 @@ public:
 
 	/**
 	 * Exports the currently opened USD Stage to a single "flattened" USD layer.
-	 * Corresponds to the "File -> Export -> All Layers" action on the USD Stage Editor menu bar.
+	 * Corresponds to the "File -> Export -> Flattened stage" action on the USD Stage Editor menu bar.
 	 * @param OutputLayer - File path (e.g. "C:/ExportFolder/out.usda") to export the flattened layer to.
 	 *                      If this path is the empty string a dialog will be shown to let the user pick the file.
 	 */
 	UFUNCTION(BlueprintCallable, Category = "USD|Menu Actions")
 	static void FileExportFlattenedStage(const FString& OutputLayer);
 
+	/**
+	 * Exports the currently opened USD Stage to a single "flattened" USD layer.
+	 * Corresponds to the "File -> Export -> Flattened layer stack" action on the USD Stage Editor menu bar.
+	 * @param OutputLayer - File path (e.g. "C:/ExportFolder/out.usda") to export the flattened layer to.
+	 *                      If this path is the empty string a dialog will be shown to let the user pick the file.
+	 */
+	UFUNCTION(BlueprintCallable, Category = "USD|Menu Actions")
+	static void FileExportFlattenedLayerStack(const FString& OutputLayer);
+
 	/**
 	 * Reloads all layers of the current USD Stage.
 	 * Corresponds to the "File -> Reload" action on the USD Stage Editor menu bar.
diff --git a/Engine/Plugins/Importers/USDImporter/Source/USDStageEditor/Public/USDStageEditorModule.h b/Engine/Plugins/Importers/USDImporter/Source/USDStageEditor/Public/USDStageEditorModule.h
--- a/Engine/Plugins/Importers/USDImporter/Source/USDStageEditor/Public/USDStageEditorModule.h
+++ b/Engine/Plugins/Importers/USDImporter/Source/USDStageEditor/Public/USDStageEditorModule.h
@@ -49,6 +49,7 @@ public:
 	void FileSave(const FString& OutputFilePathIfUnsaved) const;
 	void FileExportAllLayers(const FString& OutputDirectory) const;
 	void FileExportFlattenedStage(const FString& OutputLayer) const;
+	void FileExportFlattenedLayerStack(const FString& OutputLayer) const;
 	void FileReload() const;
 	void FileReset() const;
 	void FileClose() const;
diff --git a/Engine/Plugins/Runtime/USDCore/Source/USDUtilities/Private/USDLayerUtils.cpp b/Engine/Plugins/Runtime/USDCore/Source/USDUtilities/Private/USDLayerUtils.cpp
--- a/Engine/Plugins/Runtime/USDCore/Source/USDUtilities/Private/USDLayerUtils.cpp
+++ b/Engine/Plugins/Runtime/USDCore/Source/USDUtilities/Private/USDLayerUtils.cpp
@@ -33,6 +33,7 @@
 #include "pxr/usd/usd/primCompositionQuery.h"
 #include "pxr/usd/usd/stage.h"
 #include "pxr/usd/usdGeom/xform.h"
+#include "pxr/usd/usdUtils/flattenLayerStack.h"
 #include "USDIncludesEnd.h"
 
 #define LOCTEXT_NAMESPACE "USDLayerUtils"
@@ -825,6 +826,12 @@ bool UsdUtils::IsSessionLayerWithinStage(const pxr::SdfLayerRefPtr& Layer, const
 	return false;
 }
 
+UE::FSdfLayer UsdUtils::FlattenLayerStack(const pxr::UsdStageRefPtr& Stage)
+{
+	FScopedUsdAllocs Allocs;
+	return UE::FSdfLayer{pxr::UsdUtilsFlattenLayerStack(Stage)};
+}
+
 void UsdUtils::ConvertAssetRelativePathsToAbsolute(UE::FSdfLayer& LayerToConvert, const UE::FSdfLayer& AnchorLayer)
 {
 	if (!LayerToConvert || !AnchorLayer)
diff --git a/Engine/Plugins/Runtime/USDCore/Source/USDUtilities/Public/USDLayerUtils.h b/Engine/Plugins/Runtime/USDCore/Source/USDUtilities/Public/USDLayerUtils.h
--- a/Engine/Plugins/Runtime/USDCore/Source/USDUtilities/Public/USDLayerUtils.h
+++ b/Engine/Plugins/Runtime/USDCore/Source/USDUtilities/Public/USDLayerUtils.h
@@ -146,6 +146,12 @@ namespace UsdUtils
 	/** Returns true if Layer is a session layer within Stage's layer stack */
 	USDUTILITIES_API bool IsSessionLayerWithinStage(const pxr::SdfLayerRefPtr& Layer, const pxr::UsdStageRefPtr& Stage);
 
+	/**
+	 * Uses the USD SDK's pxr::UsdUtilsFlattenLayerStack to essentially flatten all the sublayer composition into a single layer, preserving other
+	 * composition arcs like references and payloads
+	 */
+	USDUTILITIES_API UE::FSdfLayer FlattenLayerStack(const pxr::UsdStageRefPtr& Stage);
+
 	/**
 	 * Finds all fields that have asset paths as values (e.g. texture references) in LayerToConvert and (if they're relative paths)
 	 * updates them to be absolute with respect to AnchorLayer. LayerToConvert and AnchorLayer can be the same layer, but they don't have to
