USD: Return the full list of imported assets during a USD import. This list can be retrieved in Python through the imported_object_paths of an AssetImportTask.

diff --git a/Engine/Plugins/Importers/USDImporter/Source/USDStageImporter/Private/USDStageAssetImportFactory.cpp b/Engine/Plugins/Importers/USDImporter/Source/USDStageImporter/Private/USDStageAssetImportFactory.cpp
--- a/Engine/Plugins/Importers/USDImporter/Source/USDStageImporter/Private/USDStageAssetImportFactory.cpp
+++ b/Engine/Plugins/Importers/USDImporter/Source/USDStageImporter/Private/USDStageAssetImportFactory.cpp
@@ -84,6 +84,9 @@ UObject* UUsdStageAssetImportFactory::FactoryCreateFile(
 		GEditor->BroadcastLevelActorListChanged();
 
 		ImportedObject = ImportContext.ImportedAsset ? ToRawPtr(ImportContext.ImportedAsset) : Cast<UObject>(ImportContext.SceneActor);
+
+		ImportContext.ImportedAssets.Remove(ImportedObject);
+		AdditionalImportedObjects = ImportContext.ImportedAssets;
 	}
 	else
 	{
diff --git a/Engine/Plugins/Importers/USDImporter/Source/USDStageImporter/Private/USDStageImportContext.cpp b/Engine/Plugins/Importers/USDImporter/Source/USDStageImporter/Private/USDStageImportContext.cpp
--- a/Engine/Plugins/Importers/USDImporter/Source/USDStageImporter/Private/USDStageImportContext.cpp
+++ b/Engine/Plugins/Importers/USDImporter/Source/USDStageImporter/Private/USDStageImportContext.cpp
@@ -91,6 +91,7 @@ void FUsdStageImportContext::Reset()
 	FilePath = FString{};
 	ImportOptions = NewObject<UUsdStageImportOptions>();
 	ImportedAsset = nullptr;
+	ImportedAssets.Reset();
 	LevelSequenceHelper.Clear();
 	UsdAssetCache = nullptr;
 	PRAGMA_DISABLE_DEPRECATION_WARNINGS
diff --git a/Engine/Plugins/Importers/USDImporter/Source/USDStageImporter/Private/USDStageImportFactory.cpp b/Engine/Plugins/Importers/USDImporter/Source/USDStageImporter/Private/USDStageImportFactory.cpp
--- a/Engine/Plugins/Importers/USDImporter/Source/USDStageImporter/Private/USDStageImportFactory.cpp
+++ b/Engine/Plugins/Importers/USDImporter/Source/USDStageImporter/Private/USDStageImportFactory.cpp
@@ -66,6 +66,9 @@ UObject* UUsdStageImportFactory::FactoryCreateFile(
 		GEditor->RedrawLevelEditingViewports();
 
 		ImportedObject = ImportContext.ImportedAsset ? ToRawPtr(ImportContext.ImportedAsset) : Cast<UObject>(ImportContext.SceneActor);
+		
+		ImportContext.ImportedAssets.Remove(ImportedObject);
+		AdditionalImportedObjects = ImportContext.ImportedAssets;
 	}
 	else
 	{
diff --git a/Engine/Plugins/Importers/USDImporter/Source/USDStageImporter/Private/USDStageImporter.cpp b/Engine/Plugins/Importers/USDImporter/Source/USDStageImporter/Private/USDStageImporter.cpp
--- a/Engine/Plugins/Importers/USDImporter/Source/USDStageImporter/Private/USDStageImporter.cpp
+++ b/Engine/Plugins/Importers/USDImporter/Source/USDStageImporter/Private/USDStageImporter.cpp
@@ -764,6 +764,7 @@ namespace UE::USDStageImporter::Private
 		}
 
 		ImportContext.ImportedAsset = MovedAsset;
+		ImportContext.ImportedAssets.Add(MovedAsset);
 
 		return MovedAsset;
 	}
diff --git a/Engine/Plugins/Importers/USDImporter/Source/USDStageImporter/Public/USDStageImportContext.h b/Engine/Plugins/Importers/USDImporter/Source/USDStageImporter/Public/USDStageImportContext.h
--- a/Engine/Plugins/Importers/USDImporter/Source/USDStageImporter/Public/USDStageImportContext.h
+++ b/Engine/Plugins/Importers/USDImporter/Source/USDStageImporter/Public/USDStageImportContext.h
@@ -53,6 +53,9 @@ struct USDSTAGEIMPORTER_API FUsdStageImportContext
 	UPROPERTY()
 	TObjectPtr<UObject> ImportedAsset;
 
+	UPROPERTY()
+	TArray<TObjectPtr<UObject>> ImportedAssets;
+
 	/** Level sequence that will contain the animation data during the import process */
 	FUsdLevelSequenceHelper LevelSequenceHelper;
 
