USD: Early out of translating multiple instances of the same USD prototypes (for static meshes).

diff --git a/Engine/Plugins/Importers/USDImporter/Source/USDSchemas/Private/USDGeomMeshTranslator.cpp b/Engine/Plugins/Importers/USDImporter/Source/USDSchemas/Private/USDGeomMeshTranslator.cpp
--- a/Engine/Plugins/Importers/USDImporter/Source/USDSchemas/Private/USDGeomMeshTranslator.cpp
+++ b/Engine/Plugins/Importers/USDImporter/Source/USDSchemas/Private/USDGeomMeshTranslator.cpp
@@ -589,6 +589,8 @@ namespace UsdGeomMeshTranslatorImpl
 
 		FSHAHash AllLODHash;
 		{
+			TRACE_CPUPROFILER_EVENT_SCOPE(CreateStaticMesh::HashingMeshDescription);
+
 			FSHA1 SHA1;
 
 			for (const FMeshDescription& MeshDescription : LODIndexToMeshDescription)
@@ -1922,15 +1924,22 @@ void FUsdGeomMeshTranslator::CreateAssets()
 		return Super::CreateAssets();
 	}
 
+	if (ShouldSkipInstance())
+	{
+		return;
+	}
+
+	UE::FUsdPrim Prim = GetPrim();
+
 	// Don't bother generating assets if we're going to just draw some bounds for this prim instead
-	EUsdDrawMode DrawMode = UsdUtils::GetAppliedDrawMode(GetPrim());
+	EUsdDrawMode DrawMode = UsdUtils::GetAppliedDrawMode(Prim);
 	if (DrawMode != EUsdDrawMode::Default)
 	{
 		CreateAlternativeDrawModeAssets(DrawMode);
 		return;
 	}
 
-	if (UsdUtils::IsCollisionMesh(GetPrim()))
+	if (UsdUtils::IsCollisionMesh(Prim))
 	{
 		return;
 	}
@@ -1940,7 +1949,8 @@ void FUsdGeomMeshTranslator::CreateAssets()
 		return;
 	}
 
-	TSharedRef<FGeomMeshCreateAssetsTaskChain> AssetsTaskChain = MakeShared<FGeomMeshCreateAssetsTaskChain>(Context, PrimPath);
+	UE::FSdfPath PrototypePrimPath = GetPrototypePrimPath();
+	TSharedRef<FGeomMeshCreateAssetsTaskChain> AssetsTaskChain = MakeShared<FGeomMeshCreateAssetsTaskChain>(Context, PrimPath, PrototypePrimPath);
 
 	Context->TranslatorTasks.Add(MoveTemp(AssetsTaskChain));
 }
diff --git a/Engine/Plugins/Importers/USDImporter/Source/USDSchemas/Private/USDGeomXformableTranslator.cpp b/Engine/Plugins/Importers/USDImporter/Source/USDSchemas/Private/USDGeomXformableTranslator.cpp
--- a/Engine/Plugins/Importers/USDImporter/Source/USDSchemas/Private/USDGeomXformableTranslator.cpp
+++ b/Engine/Plugins/Importers/USDImporter/Source/USDSchemas/Private/USDGeomXformableTranslator.cpp
@@ -207,8 +207,12 @@ namespace UE::UsdXformableTranslatorImpl::Private
 class FUsdGeomXformableCreateAssetsTaskChain : public FBuildStaticMeshTaskChain
 {
 public:
-	explicit FUsdGeomXformableCreateAssetsTaskChain(const TSharedRef<FUsdSchemaTranslationContext>& InContext, const UE::FSdfPath& InPrimPath)
-		: FBuildStaticMeshTaskChain(InContext, InPrimPath)
+	explicit FUsdGeomXformableCreateAssetsTaskChain(
+		const TSharedRef<FUsdSchemaTranslationContext>& InContext,
+		const UE::FSdfPath& InPrimPath,
+		const TOptional<UE::FSdfPath>& AlternativePrimToLinkAssetsTo = {}
+	)
+		: FBuildStaticMeshTaskChain(InContext, InPrimPath, AlternativePrimToLinkAssetsTo)
 	{
 		SetupTasks();
 	}
@@ -274,15 +278,23 @@ void FUsdGeomXformableTranslator::CreateAssets()
 		return;
 	}
 
+	if (ShouldSkipInstance())
+	{
+		return;
+	}
+
+	UE::FUsdPrim Prim = GetPrim();
+
 	// Don't bother generating assets if we're going to just draw some bounds for this prim instead
-	EUsdDrawMode DrawMode = UsdUtils::GetAppliedDrawMode(GetPrim());
+	EUsdDrawMode DrawMode = UsdUtils::GetAppliedDrawMode(Prim);
 	if (DrawMode != EUsdDrawMode::Default)
 	{
 		CreateAlternativeDrawModeAssets(DrawMode);
 		return;
 	}
 
-	Context->TranslatorTasks.Add(MakeShared<FUsdGeomXformableCreateAssetsTaskChain>(Context, PrimPath));
+	UE::FSdfPath PrototypePrimPath = GetPrototypePrimPath();
+	Context->TranslatorTasks.Add(MakeShared<FUsdGeomXformableCreateAssetsTaskChain>(Context, PrimPath, PrototypePrimPath));
 }
 
 FUsdGeomXformableTranslator::FUsdGeomXformableTranslator(
@@ -322,7 +334,8 @@ USceneComponent* FUsdGeomXformableTranslator::CreateComponents()
 	{
 		if (Context->InfoCache)
 		{
-			if (UStaticMesh* StaticMesh = Context->InfoCache->GetSingleAssetForPrim<UStaticMesh>(PrimPath))
+			UE::FSdfPath PrototypePath = GetPrototypePrimPath();
+			if (UStaticMesh* StaticMesh = Context->InfoCache->GetSingleAssetForPrim<UStaticMesh>(PrototypePath))
 			{
 				TArray<UMaterialInterface*> ExistingAssignments;
 				for (FStaticMaterial& StaticMaterial : StaticMesh->GetStaticMaterials())
@@ -669,7 +682,8 @@ void FUsdGeomXformableTranslator::UpdateComponents(USceneComponent* SceneCompone
 		bool bHasMultipleLODs = false;
 		if (UStaticMeshComponent* StaticMeshComponent = Cast<UStaticMeshComponent>(SceneComponent))
 		{
-			UStaticMesh* PrimStaticMesh = Context->InfoCache->GetSingleAssetForPrim<UStaticMesh>(PrimPath);
+			UE::FSdfPath PrototypePath = GetPrototypePrimPath();
+			UStaticMesh* PrimStaticMesh = Context->InfoCache->GetSingleAssetForPrim<UStaticMesh>(PrototypePath);
 
 			if (PrimStaticMesh)
 			{
diff --git a/Engine/Plugins/Importers/USDImporter/Source/USDSchemas/Private/USDInfoCache.cpp b/Engine/Plugins/Importers/USDImporter/Source/USDSchemas/Private/USDInfoCache.cpp
--- a/Engine/Plugins/Importers/USDImporter/Source/USDSchemas/Private/USDInfoCache.cpp
+++ b/Engine/Plugins/Importers/USDImporter/Source/USDSchemas/Private/USDInfoCache.cpp
@@ -181,6 +181,12 @@ struct FUsdInfoCache::FUsdInfoCacheImpl
 	TArray<FString> TempPointInstancerPaths;
 	mutable FRWLock TempPointInstancerPathsLock;
 
+	// This is used to keep track of which prototypes are already being translated within this "translation session",
+	// so that the schema translators can early out if they're trying to translate multiple instances of the same
+	// prototype
+	TSet<UE::FSdfPath> TranslatedPrototypes;
+	mutable FRWLock TranslatedPrototypesLock;
+
 	// Geometry cache can come from a reference or payload of these file types
 	TArray<FString> AllowedExtensionsForGeometryCacheSource;
 
@@ -248,6 +254,10 @@ bool FUsdInfoCache::Serialize(FArchive& Ar)
 			Ar << ImplPtr->PrimPathToAssets;
 			Ar << ImplPtr->AssetToPrimPaths;
 		}
+		{
+			FWriteScopeLock ScopeLock(ImplPtr->TranslatedPrototypesLock);
+			Ar << ImplPtr->TranslatedPrototypes;
+		}
 	}
 
 	return true;
@@ -1803,6 +1813,35 @@ bool FUsdInfoCache::IsPotentialGeometryCacheRoot(const UE::FSdfPath& Path) const
 	return false;
 }
 
+void FUsdInfoCache::ResetTranslatedPrototypes()
+{
+	if (FUsdInfoCacheImpl* ImplPtr = Impl.Get())
+	{
+		FWriteScopeLock ScopeLock(ImplPtr->TranslatedPrototypesLock);
+		return ImplPtr->TranslatedPrototypes.Reset();
+	}
+}
+
+bool FUsdInfoCache::IsPrototypeTranslated(const UE::FSdfPath& PrototypePath)
+{
+	if (FUsdInfoCacheImpl* ImplPtr = Impl.Get())
+	{
+		FReadScopeLock ScopeLock(ImplPtr->TranslatedPrototypesLock);
+		return ImplPtr->TranslatedPrototypes.Contains(PrototypePath);
+	}
+
+	return false;
+}
+
+void FUsdInfoCache::MarkPrototypeAsTranslated(const UE::FSdfPath& PrototypePath)
+{
+	if (FUsdInfoCacheImpl* ImplPtr = Impl.Get())
+	{
+		FWriteScopeLock ScopeLock(ImplPtr->TranslatedPrototypesLock);
+		ImplPtr->TranslatedPrototypes.Add(PrototypePath);
+	}
+}
+
 TOptional<uint64> FUsdInfoCache::GetSubtreeVertexCount(const UE::FSdfPath& Path)
 {
 	if (FUsdInfoCacheImpl* ImplPtr = Impl.Get())
@@ -2113,6 +2152,8 @@ void FUsdInfoCache::Clear()
 			FWriteScopeLock PointInstancerLock(ImplPtr->TempPointInstancerPathsLock);
 			ImplPtr->TempPointInstancerPaths.Empty();
 		}
+
+		ResetTranslatedPrototypes();
 	}
 }
 
diff --git a/Engine/Plugins/Importers/USDImporter/Source/USDSchemas/Private/USDSchemaTranslator.cpp b/Engine/Plugins/Importers/USDImporter/Source/USDSchemas/Private/USDSchemaTranslator.cpp
--- a/Engine/Plugins/Importers/USDImporter/Source/USDSchemas/Private/USDSchemaTranslator.cpp
+++ b/Engine/Plugins/Importers/USDImporter/Source/USDSchemas/Private/USDSchemaTranslator.cpp
@@ -14,10 +14,6 @@
 #include "Misc/ScopedSlowTask.h"
 #include "Modules/ModuleManager.h"
 
-#define LOCTEXT_NAMESPACE "USDSchemaTranslator"
-
-int32 FRegisteredSchemaTranslatorHandle::CurrentSchemaTranslatorId = 0;
-
 #if USE_USD_SDK
 #include "USDIncludesStart.h"
 #include "pxr/base/tf/token.h"
@@ -25,9 +21,21 @@ int32 FRegisteredSchemaTranslatorHandle::CurrentSchemaTranslatorId = 0;
 #include "pxr/usd/usd/prim.h"
 #include "pxr/usd/usdShade/tokens.h"
 #include "USDIncludesEnd.h"
-
 #endif	  // #if USE_USD_SDK
 
+#define LOCTEXT_NAMESPACE "USDSchemaTranslator"
+
+int32 FRegisteredSchemaTranslatorHandle::CurrentSchemaTranslatorId = 0;
+
+static bool GInstancingAwareTranslation = true;
+static FAutoConsoleVariableRef CVarInstancingAwareTranslation(
+	TEXT("USD.InstancingAwareTranslation"),
+	GInstancingAwareTranslation,
+	TEXT(
+		"Enabling this lets the USDImporter skip some extra steps during translation when it encounters multiple instance prims of the same (static) Mesh prototype prim."
+	)
+);
+
 TSharedPtr<FUsdSchemaTranslator> FUsdSchemaTranslatorRegistry::CreateTranslatorForSchema(
 	TSharedRef<FUsdSchemaTranslationContext> InTranslationContext,
 	const UE::FUsdTyped& InSchema
@@ -224,6 +232,12 @@ void FUsdSchemaTranslationContext::CompleteTasks()
 
 		bFinished = (TranslatorTasks.Num() == 0);
 	}
+
+	// These only make sense within an asset translation period
+	if (InfoCache)
+	{
+		InfoCache->ResetTranslatedPrototypes();
+	}
 }
 
 bool FUsdSchemaTranslator::IsCollapsed(ECollapsingType CollapsingType) const
@@ -243,6 +257,50 @@ bool FUsdSchemaTranslator::IsCollapsed(ECollapsingType CollapsingType) const
 	return false;
 }
 
+UE::FSdfPath FUsdSchemaTranslator::GetPrototypePrimPath() const
+{
+	if (!GInstancingAwareTranslation)
+	{
+		return PrimPath;
+	}
+
+	UE::FUsdPrim Prim = GetPrim();
+	if (Prim)
+	{
+		if (Prim.IsInstance())
+		{
+			return Prim.GetPrototype().GetPrimPath();
+		}
+		else if (Prim.IsInstanceProxy())
+		{
+			return Prim.GetPrimInPrototype().GetPrimPath();
+		}
+	}
+
+	return PrimPath;
+}
+
+bool FUsdSchemaTranslator::ShouldSkipInstance() const
+{
+	if (!GInstancingAwareTranslation)
+	{
+		return false;
+	}
+
+	UE::FSdfPath PrototypePath = GetPrototypePrimPath();
+	if (!PrototypePath.IsEmpty())
+	{
+		if (Context->InfoCache->IsPrototypeTranslated(PrototypePath))
+		{
+			return true;
+		}
+
+		Context->InfoCache->MarkPrototypeAsTranslated(PrototypePath);
+	}
+
+	return false;
+}
+
 void FSchemaTranslatorTask::Start()
 {
 	if (LaunchPolicy == ESchemaTranslationLaunchPolicy::Async && IsInGameThread())
diff --git a/Engine/Plugins/Importers/USDImporter/Source/USDSchemas/Public/USDInfoCache.h b/Engine/Plugins/Importers/USDImporter/Source/USDSchemas/Public/USDInfoCache.h
--- a/Engine/Plugins/Importers/USDImporter/Source/USDSchemas/Public/USDInfoCache.h
+++ b/Engine/Plugins/Importers/USDImporter/Source/USDSchemas/Public/USDInfoCache.h
@@ -87,6 +87,14 @@ public:
 	// Returns true if Path could potentially be collapsed as a Geometry Cache asset
 	bool IsPotentialGeometryCacheRoot(const UE::FSdfPath& Path) const;
 
+public:
+	// Marks/checks if the provided path to a prototype prim is already being translated.
+	// This is used during scene translation with instanceables, so that the schema translators can early out
+	// in case they have been created to translate multiple instances of the same prototype
+	void ResetTranslatedPrototypes();
+	bool IsPrototypeTranslated(const UE::FSdfPath& PrototypePath);
+	void MarkPrototypeAsTranslated(const UE::FSdfPath& PrototypePath);
+
 public:
 	void LinkAssetToPrim(const UE::FSdfPath& Path, UObject* Asset);
 	void UnlinkAssetFromPrim(const UE::FSdfPath& Path, UObject* Asset);
diff --git a/Engine/Plugins/Importers/USDImporter/Source/USDSchemas/Public/USDSchemaTranslator.h b/Engine/Plugins/Importers/USDImporter/Source/USDSchemas/Public/USDSchemaTranslator.h
--- a/Engine/Plugins/Importers/USDImporter/Source/USDSchemas/Public/USDSchemaTranslator.h
+++ b/Engine/Plugins/Importers/USDImporter/Source/USDSchemas/Public/USDSchemaTranslator.h
@@ -404,6 +404,26 @@ public:
 		return false;
 	}
 
+	/**
+	 * This checks if the current prim is an instance, and if so, whether its
+	 * prototype is already being translated. Returns false otherwise.
+	 *
+	 * WARNING: In case this prim is an instance but the prototype is not being translated yet,
+	 * running this check will also mark that prototype as being currently translated on the info cache!
+	 *
+	 * The intent here is that the first schema translator that calls this for a prototype
+	 * will "own" the translation for that prototype, and any subsequent calls by other schema translators
+	 * with the same prototype will just return true so they can early out.
+	 */
+	bool ShouldSkipInstance() const;
+
+	/**
+	 * If this prim is a prototype or an instance proxy, returns the prototype path (or the path to the
+	 * analogue prim in the prototype's hierarchy).
+	 * If this prim is just a regular non-instance prim, this just returns our PrimPath member.
+	 */
+	UE::FSdfPath GetPrototypePrimPath() const;
+
 	UE::FUsdPrim GetPrim() const
 	{
 		return Context->Stage.GetPrimAtPath(PrimPath);
