USD: Improve the USD.DiscardUndoBufferOnStageOpenClose cvar to have it prevent serialization of assets/components and reset the transactor.

diff --git a/Engine/Plugins/Importers/USDImporter/Source/USDStage/Private/USDStageActor.cpp b/Engine/Plugins/Importers/USDImporter/Source/USDStage/Private/USDStageActor.cpp
--- a/Engine/Plugins/Importers/USDImporter/Source/USDStage/Private/USDStageActor.cpp
+++ b/Engine/Plugins/Importers/USDImporter/Source/USDStage/Private/USDStageActor.cpp
@@ -122,6 +122,15 @@ static FAutoConsoleVariableRef CVarTranslateOnlyUsedMaterialsWhenOpeningStage(
 		 "translated into Unreal material assets.")
 );
 
+static bool GDiscardUndoBufferOnStageOpenClose = false;
+static FAutoConsoleVariableRef CVarDiscardUndoBufferOnStageOpenClose(
+	TEXT("USD.DiscardUndoBufferOnStageOpenClose"),
+	GDiscardUndoBufferOnStageOpenClose,
+	TEXT(
+		"Enabling this will prevent the recording of open/close stage transactions, but also discard the undo buffer after they happen. This can help load times and also reduce memory usage, as sometimes recording all created assets and actors in the undo buffer can be expensive."
+	)
+);
+
 static const EObjectFlags DefaultObjFlag = EObjectFlags::RF_Transactional | EObjectFlags::RF_Public | EObjectFlags::RF_Standalone;
 
 AUsdStageActor::FOnActorLoaded AUsdStageActor::OnActorLoaded;
@@ -2906,6 +2915,13 @@ void AUsdStageActor::OpenUsdStage()
 
 	TRACE_CPUPROFILER_EVENT_SCOPE(AUsdStageActor::OpenUsdStage);
 
+	TOptional<TGuardValue<ITransaction*>> SuppressTransaction;
+	if (GDiscardUndoBufferOnStageOpenClose)
+	{
+		SuppressTransaction.Emplace(GUndo, nullptr);
+		RequestDelayedTransactorReset();
+	}
+
 	UsdUtils::StartMonitoringErrors();
 
 	FString AbsPath;
@@ -2959,6 +2975,13 @@ void AUsdStageActor::OpenUsdStage()
 
 void AUsdStageActor::CloseUsdStage(bool bUnloadIfNeeded)
 {
+	TOptional<TGuardValue<ITransaction*>> SuppressTransaction;
+	if (GDiscardUndoBufferOnStageOpenClose)
+	{
+		SuppressTransaction.Emplace(GUndo, nullptr);
+		RequestDelayedTransactorReset();
+	}
+
 	const bool bStageWasOpened = static_cast<bool>(UsdStage);
 	if (bStageWasOpened)
 	{
@@ -3142,6 +3165,13 @@ void AUsdStageActor::LoadUsdStage(bool bOpenIfNeeded)
 		}
 	}
 
+	TOptional<TGuardValue<ITransaction*>> SuppressTransaction;
+	if (GDiscardUndoBufferOnStageOpenClose)
+	{
+		SuppressTransaction.Emplace(GUndo, nullptr);
+		RequestDelayedTransactorReset();
+	}
+
 	if (!UsdStage && bOpenIfNeeded)
 	{
 		OpenUsdStage();
@@ -3275,6 +3305,13 @@ void AUsdStageActor::UnloadUsdStage()
 	const bool bMarkDirty = false;
 	Modify(bMarkDirty);
 
+	TOptional<TGuardValue<ITransaction*>> SuppressTransaction;
+	if (GDiscardUndoBufferOnStageOpenClose)
+	{
+		SuppressTransaction.Emplace(GUndo, nullptr);
+		RequestDelayedTransactorReset();
+	}
+
 	FUsdStageActorImpl::DeselectActorsAndComponents(this);
 
 	// Stop listening because we'll discard LevelSequence assets, which may trigger transactions
@@ -4082,6 +4119,39 @@ void AUsdStageActor::UpdateSpawnedObjectsTransientFlag(bool bTransient)
 	GetRootPrimTwin()->Iterate(UpdateTransient, bRecursive);
 }
 
+void AUsdStageActor::RequestDelayedTransactorReset()
+{
+#if WITH_EDITOR
+	bIsPendingTransactorReset = true;
+
+	// Wait for the next tick because many of our functions may all try to
+	// get the transactor reset
+	TWeakObjectPtr<AUsdStageActor> WeakThis{this};
+	ExecuteOnGameThread(
+		UE_SOURCE_LOCATION,
+		[WeakThis]()
+		{
+			if (AUsdStageActor* Actor = WeakThis.Get())
+			{
+				if (Actor->bIsPendingTransactorReset && GEditor)
+				{
+					Actor->bIsPendingTransactorReset = false;
+
+					if (UTransactor* EditorTransactor = GEditor->Trans)
+					{
+						const FText Reason = LOCTEXT(
+							"DiscardTransactionReason",
+							"Resetting because USD.DiscardUndoBufferOnStageOpenClose is enabled"
+						);
+						EditorTransactor->Reset(Reason);
+					}
+				}
+			}
+		}
+	);
+#endif	  // WITH_EDITOR
+}
+
 void AUsdStageActor::OnUsdPrimTwinDestroyed(const UUsdPrimTwin& UsdPrimTwin)
 {
 	PrimsToAnimate.Remove(UsdPrimTwin.PrimPath);
@@ -4442,27 +4512,11 @@ void AUsdStageActor::HandlePropertyChangedEvent(FPropertyChangedEvent& PropertyC
 
 	if (MemberPropertyName == GET_MEMBER_NAME_CHECKED(AUsdStageActor, RootLayer))
 	{
-#if WITH_EDITOR
-		static IConsoleVariable* CVar = IConsoleManager::Get().FindConsoleVariable(TEXT("USD.DiscardUndoBufferOnStageOpenClose"));
-		const bool bDiscardUndo = CVar && CVar->GetBool();
-#endif	  // WITH_EDITOR
-
 		// Technically we don't need this guard value for the root layer itself, since SetRootLayer can compare
 		// RootLayer with the path of the current stage's root layer, but let's just do this for consistency.
 		const FString CorrectRootLayer = RootLayer.FilePath;
 		RootLayer.FilePath = RootLayer.FilePath + TEXT("dummy");
 		SetRootLayer(CorrectRootLayer);
-
-#if WITH_EDITOR
-		if (bDiscardUndo && GEditor)
-		{
-			if (UTransactor* EditorTransactor = GEditor->Trans)
-			{
-				const FText Reason = LOCTEXT("DiscardTransactionReason", "Resetting because USD.DiscardUndoBufferOnStageOpenClose is enabled");
-				EditorTransactor->Reset(Reason);
-			}
-		}
-#endif	  // WITH_EDITOR
 	}
 	else if (PropertyName == GET_MEMBER_NAME_CHECKED(AUsdStageActor, StageState))
 	{
diff --git a/Engine/Plugins/Importers/USDImporter/Source/USDStage/Public/USDStageActor.h b/Engine/Plugins/Importers/USDImporter/Source/USDStage/Public/USDStageActor.h
--- a/Engine/Plugins/Importers/USDImporter/Source/USDStage/Public/USDStageActor.h
+++ b/Engine/Plugins/Importers/USDImporter/Source/USDStage/Public/USDStageActor.h
@@ -395,6 +395,8 @@ protected:
 
 	void UpdateSpawnedObjectsTransientFlag(bool bTransient);
 
+	void RequestDelayedTransactorReset();
+
 #if WITH_EDITOR
 	void OnBeginPIE(bool bIsSimulating);
 	void OnPostPIEStarted(bool bIsSimulating);
@@ -481,6 +483,7 @@ protected:
 	bool bIsModifyingAProperty;
 	bool bIsUndoRedoing;
 	bool bIsAddingBinding = false;
+	bool bIsPendingTransactorReset = false;
 
 	FDelegateHandle OnRedoHandle;
 	FDelegateHandle OnSequencerCreatedHandle;
diff --git a/Engine/Plugins/Importers/USDImporter/Source/USDStageEditor/Private/SUSDStage.cpp b/Engine/Plugins/Importers/USDImporter/Source/USDStageEditor/Private/SUSDStage.cpp
--- a/Engine/Plugins/Importers/USDImporter/Source/USDStageEditor/Private/SUSDStage.cpp
+++ b/Engine/Plugins/Importers/USDImporter/Source/USDStageEditor/Private/SUSDStage.cpp
@@ -48,14 +48,6 @@
 
 #define LOCTEXT_NAMESPACE "SUsdStage"
 
-static bool GDiscardUndoBufferOnStageOpenClose = false;
-static FAutoConsoleVariableRef CVarDiscardUndoBufferOnStageOpenClose(
-	TEXT("USD.DiscardUndoBufferOnStageOpenClose"),
-	GDiscardUndoBufferOnStageOpenClose,
-	TEXT("Enabling this will prevent the recording of open/close stage transactions, but also discard the undo buffer after they happen. Use this "
-		 "when memory-constrained, as sometimes recording all created assets and actors in the undo buffer can be expensive.")
-);
-
 #if USE_USD_SDK
 
 namespace SUSDStageConstants
@@ -2548,23 +2540,10 @@ void SUsdStage::FileClose()
 {
 	TRACE_CPUPROFILER_EVENT_SCOPE(SUsdStage::FileClose);
 
-	TOptional<FScopedTransaction> Transaction;
-	if (!GDiscardUndoBufferOnStageOpenClose)
-	{
-		Transaction.Emplace(LOCTEXT("CloseTransaction", "Close USD stage"));
-	}
+	FScopedTransaction Transaction{LOCTEXT("CloseTransaction", "Close USD stage")};
 
 	ViewModel.CloseStage();
 	RequestFullRefresh();
-
-	if (GDiscardUndoBufferOnStageOpenClose && GEditor)
-	{
-		if (UTransactor* Transactor = GEditor->Trans)
-		{
-			const FText Reason = LOCTEXT("DiscardCloseTransactionReason", "Resetting because the cvar USD.DiscardUndoBufferOnStageOpenClose is true");
-			Transactor->Reset(Reason);
-		}
-	}
 }
 
 void SUsdStage::OnLayerIsolated(const UE::FSdfLayer& IsolatedLayer)
@@ -2648,13 +2627,9 @@ void SUsdStage::OnPrimSelectionChanged(const TArray<FString>& PrimPaths)
 
 void SUsdStage::OpenStage(const TCHAR* FilePath)
 {
-	TOptional<FScopedTransaction> Transaction;
-	if (!GDiscardUndoBufferOnStageOpenClose)
-	{
-		// Create the transaction before calling UsdStageModule.GetUsdStageActor as that may create the actor, and we want
-		// the actor spawning to be part of the transaction
-		Transaction.Emplace(FText::Format(LOCTEXT("OpenStageTransaction", "Open USD stage '{0}'"), FText::FromString(FilePath)));
-	}
+	// Create the transaction before calling UsdStageModule.GetUsdStageActor as that may create the actor, and we want
+	// the actor spawning to be part of the transaction
+	FScopedTransaction Transaction{FText::Format(LOCTEXT("OpenStageTransaction", "Open USD stage '{0}'"), FText::FromString(FilePath))};
 
 	if (!ViewModel.UsdStageActor.IsValid())
 	{
@@ -2663,15 +2638,6 @@ void SUsdStage::OpenStage(const TCHAR* FilePath)
 	}
 
 	ViewModel.OpenStage(FilePath);
-
-	if (GDiscardUndoBufferOnStageOpenClose && GEditor)
-	{
-		if (UTransactor* Transactor = GEditor->Trans)
-		{
-			const FText Reason = LOCTEXT("DiscardOpenTransactionReason", "Resetting because the cvar USD.DiscardUndoBufferOnStageOpenClose is true");
-			Transactor->Reset(Reason);
-		}
-	}
 }
 
 void SUsdStage::RequestLayersTreeViewRefresh()
