USD: Fix AUsdStageActor responding to the same notice twice if the original notice handling took long enough to show a progress bar.

diff --git a/Engine/Plugins/Importers/USDImporter/Source/USDStageEditorViewModels/Private/USDObjectFieldViewModel.cpp b/Engine/Plugins/Importers/USDImporter/Source/USDStageEditorViewModels/Private/USDObjectFieldViewModel.cpp
--- a/Engine/Plugins/Importers/USDImporter/Source/USDStageEditorViewModels/Private/USDObjectFieldViewModel.cpp
+++ b/Engine/Plugins/Importers/USDImporter/Source/USDStageEditorViewModels/Private/USDObjectFieldViewModel.cpp
@@ -159,6 +159,23 @@ void FUsdObjectFieldsViewModel::SetFieldValue(const FString& FieldName, const Us
 		return;
 	}
 
+	// Avoid a reentry here because of this sequence of events:
+	//  - The user changes one of the prim property input boxes and presses Enter, calling e.g. OnSpinboxValueCommitted;
+	//  - We set the new value directly on the prim in the stage (within this function), triggering a USD stage change;
+	//  - USD sends a change notice;
+	//  - We respond to the notice on AUsdStageActor::OnUsdObjectsChanged;
+	//  - The body of that function uses a FScopedSlowTask to show a progress bar;
+	//  - The progress bar being displayed steals focus from the input box;
+	//  - The input box losing focus automatically triggers OnSpinboxValueCommitted again!
+	//  - We set the new value directly on the prim in the stage, again. Even though it's the same value, we get a new notice!
+	//  - etc.
+	static bool bIsReentrant = false;
+	if (bIsReentrant)
+	{
+		return;
+	}
+	TGuardValue<bool> ReentrantGuard(bIsReentrant, true);
+
 	// Transact here as setting this attribute may trigger USD events that affect assets/components
 	FScopedTransaction Transaction(FText::Format(
 		LOCTEXT("SetFieldValue", "Set value for field '{0}' of prim '{1}'"),
