USD: Prevent multiple resyncs for the same prim during the same UE transaction;

Prevent multiple event broadcasts for the same prim when undoing/redoing a transaction.

diff --git a/Engine/Plugins/Importers/USDImporter/Source/USDStage/Private/USDStageActor.cpp b/Engine/Plugins/Importers/USDImporter/Source/USDStage/Private/USDStageActor.cpp
--- a/Engine/Plugins/Importers/USDImporter/Source/USDStage/Private/USDStageActor.cpp
+++ b/Engine/Plugins/Importers/USDImporter/Source/USDStage/Private/USDStageActor.cpp
@@ -106,6 +106,15 @@
 
 #define LOCTEXT_NAMESPACE "USDStageActor"
 
+static bool GHandleNoticesImmediately = false;
+static FAutoConsoleVariableRef CVarHandleNoticesImmediately(
+	TEXT("USD.HandleNoticesImmediately"),
+	GHandleNoticesImmediately,
+	TEXT(
+		"When this is false, we will collect all USD notices emitted during a UE transaction and handle them in one pass, potentially eliminating some unnecessary updates. When this is true, we respond to each USD notice immediately"
+	)
+);
+
 static bool GRegenerateSkeletalAssetsOnControlRigBake = true;
 static FAutoConsoleVariableRef CVarRegenerateSkeletalAssetsOnControlRigBake(
 	TEXT("USD.RegenerateSkeletalAssetsOnControlRigBake"),
@@ -1262,6 +1271,8 @@ void AUsdStageActor::OnUsdObjectsChanged(const UsdUtils::FObjectChangesByPath& I
 		return;
 	}
 
+	TRACE_CPUPROFILER_EVENT_SCOPE(AUsdStageActor::OnUsdObjectsChanged);
+
 	// If we're opened, we shouldn't have any actor, component or asset. We shouldn't even have a built info cache!
 	// This means we don't really need to do anything here, except trigger the USD Stage Editor to refresh.
 	// We'd otherwise go through a lot of work to find out the prims to update, but given that this is only for UI refresh
@@ -1296,6 +1307,19 @@ void AUsdStageActor::OnUsdObjectsChanged(const UsdUtils::FObjectChangesByPath& I
 		return;
 	}
 
+	for (const TPair<FString, TArray<UsdUtils::FSdfChangeListEntry>>& InfoChange : InfoChanges)
+	{
+		AccumulatedInfoChanges.FindOrAdd(InfoChange.Key).Append(InfoChange.Value);
+	}
+	for (const TPair<FString, TArray<UsdUtils::FSdfChangeListEntry>>& ResyncChange : ResyncChanges)
+	{
+		AccumulatedResyncChanges.FindOrAdd(ResyncChange.Key).Append(ResyncChange.Value);
+	}
+
+	// We want to update our transactor right away, even if we may call HandleAccumulatedNotices() only at the end of the transaction.
+	// This because we want to store the current edit target when these change notices were emitted, so that we can reply them
+	// back to the right layer when undo/redoing. It's all going to always be on the same transaction anyway though
+	//
 	// Only update the transactor if we're listening to USD notices. Within OnObjectPropertyChanged we will stop listening when writing stage changes
 	// from our component changes, and this will also make sure we're not duplicating the events we store and replicate via multi-user: If a
 	// modification can be described purely via UObject changes, then those changes will be responsible for the whole modification and we won't record
@@ -1306,6 +1330,24 @@ void AUsdStageActor::OnUsdObjectsChanged(const UsdUtils::FObjectChangesByPath& I
 		Transactor->Update(InfoChanges, ResyncChanges);
 	}
 
+	// If we don't have a transaction currently, then HandleTransactionStateChanged will never be called in order to eventually
+	// call HandleAccumulatedNotices, so we have no choice but to do it now. If we didn't do this, users would always need to put
+	// their Python USD changes within a UE scoped transaction to get the stage actor to actually respond
+	if (GHandleNoticesImmediately || !GUndo)
+	{
+		HandleAccumulatedNotices();
+	}
+}
+
+void AUsdStageActor::HandleAccumulatedNotices()
+{
+	TRACE_CPUPROFILER_EVENT_SCOPE(AUsdStageActor::HandleAccumulatedNotices);
+
+	if (AccumulatedInfoChanges.Num() == 0 && AccumulatedResyncChanges.Num() == 0)
+	{
+		return;
+	}
+
 	const UE::FUsdStage& Stage = GetOrOpenUsdStage();
 	if (!Stage)
 	{
@@ -1351,120 +1393,124 @@ void AUsdStageActor::OnUsdObjectsChanged(const UsdUtils::FObjectChangesByPath& I
 	// back to USD when these objects change though.
 	FScopedBlockMonitoringChangesForTransaction BlockMonitoring{LevelSequenceHelper};
 
-	bool bHasResync = ResyncChanges.Num() > 0;
+	bool bHasResync = AccumulatedResyncChanges.Num() > 0;
 
 	bool bNeedsAnimationReload = false;
 
 	// The most important thing here is to iterate in parent to child order, so build SortedPrimsChangedList
 	TMap<UE::FSdfPath, bool> SortedPrimsChangedList;
-	for (const TPair<FString, TArray<UsdUtils::FSdfChangeListEntry>>& InfoChange : InfoChanges)
 	{
-		UE::FSdfPath PrimPath = UE::FSdfPath(*InfoChange.Key).StripAllVariantSelections();
+		TRACE_CPUPROFILER_EVENT_SCOPE(FillSortedPrimsChangedList);
 
-		// Upgrade some info changes into resync changes
-		bool bIsResync = false;
-		for (const UsdUtils::FSdfChangeListEntry& ObjectChange : InfoChange.Value)
+		for (const TPair<FString, TArray<UsdUtils::FSdfChangeListEntry>>& InfoChange : AccumulatedInfoChanges)
 		{
-			// This is in charge of resyncing components (and so calling CreateComponents instead of just
-			// UpdateComponents) whenever we change material assignments. This is important because
-			// CreateComponents is where we set our material overrides
-			if (ObjectChange.Flags.bDidChangeRelationshipTargets)
-			{
-				bIsResync = true;
-				bHasResync = true;
-
-				// If the material that was changed was on a geom subset, we need to instead pretend it
-				// happened to its parent Mesh prim instead. If the mesh is collapsed this won't matter,
-				// but if it's not collapsed we currently need this to make sure we resync the static mesh
-				// component. The Mesh prim will declare that the subset is its aux prim, but we don't propagate
-				// resyncs from aux to main prims now, so even if we resynced the UsdGeomSubset the Mesh
-				// would not resync.
-				// TODO: Better way of handling material override updates. Ideally we wouldn't be resyncing
-				// the Mesh prim just to update material overrides...
-				UE::FUsdPrim ChangedPrim = Stage.GetPrimAtPath(PrimPath);
-				if (ChangedPrim && ChangedPrim.IsA(TEXT("GeomSubset")))
-				{
-					PrimPath = PrimPath.GetParentPath();
-				}
-
-				continue;
-			}
+			UE::FSdfPath PrimPath = UE::FSdfPath(*InfoChange.Key).StripAllVariantSelections();
 
-			for (const UsdUtils::FAttributeChange& AttributeChange : ObjectChange.AttributeChanges)
+			// Upgrade some info changes into resync changes
+			bool bIsResync = false;
+			for (const UsdUtils::FSdfChangeListEntry& ObjectChange : InfoChange.Value)
 			{
-				static const TSet<FString> StageResyncProperties = {TEXT("metersPerUnit"), TEXT("upAxis")};
-
-				// Upgrade these to resync so that the prim twins are regenerated, which clears all the existing
-				// animation tracks and adds new ones, automatically re-baking to control rig
-				static const TSet<FString> PrimResyncProperties = {
-					*UsdToUnreal::ConvertToken(UnrealIdentifiers::UnrealControlRigPath),
-					*UsdToUnreal::ConvertToken(UnrealIdentifiers::UnrealUseFKControlRig),
-					*UsdToUnreal::ConvertToken(UnrealIdentifiers::UnrealControlRigReduceKeys),
-					*UsdToUnreal::ConvertToken(UnrealIdentifiers::UnrealControlRigReductionTolerance),
-					// For now we need to do this as we need to refresh the material slot info on the info cache if these
-					// update... this is of course way too aggressive, although it's unlikely people will be manually editing these.
-					// TODO: More nuanced info cache updates
-					UnrealIdentifiers::PrimvarsDisplayColor,
-					UnrealIdentifiers::PrimvarsDisplayOpacity,
-					UnrealIdentifiers::DoubleSided,
-					// When we change these UsdGeomModelAPI attributes we may need to create a new component type
-					// for the prim (as it may now need/stop needing an alternate draw mode component)
-					UnrealIdentifiers::ModelDrawMode,
-					UnrealIdentifiers::ModelApplyDrawMode,
-					// Collapsing changes need to resync to build the new collapsed/uncollapsed meshes
-					*UsdToUnreal::ConvertToken(UnrealIdentifiers::UnrealCollapsingAttr),
-					// Physics collision attribute change needs to rebuild the collision shape
-					*UsdToUnreal::ConvertToken(pxr::UsdPhysicsTokens->physicsCollisionEnabled),
-					*UsdToUnreal::ConvertToken(pxr::UsdPhysicsTokens->physicsApproximation)};
-
-				// Some stage info should trigger some resyncs because they should trigger reparsing of geometry
-				if ((PrimPath.IsAbsoluteRootPath() && StageResyncProperties.Contains(AttributeChange.PropertyName))
-					|| PrimResyncProperties.Contains(AttributeChange.PropertyName))
+				// This is in charge of resyncing components (and so calling CreateComponents instead of just
+				// UpdateComponents) whenever we change material assignments. This is important because
+				// CreateComponents is where we set our material overrides
+				if (ObjectChange.Flags.bDidChangeRelationshipTargets)
 				{
 					bIsResync = true;
 					bHasResync = true;
-					break;
+
+					// If the material that was changed was on a geom subset, we need to instead pretend it
+					// happened to its parent Mesh prim instead. If the mesh is collapsed this won't matter,
+					// but if it's not collapsed we currently need this to make sure we resync the static mesh
+					// component. The Mesh prim will declare that the subset is its aux prim, but we don't propagate
+					// resyncs from aux to main prims now, so even if we resynced the UsdGeomSubset the Mesh
+					// would not resync.
+					// TODO: Better way of handling material override updates. Ideally we wouldn't be resyncing
+					// the Mesh prim just to update material overrides...
+					UE::FUsdPrim ChangedPrim = Stage.GetPrimAtPath(PrimPath);
+					if (ChangedPrim && ChangedPrim.IsA(TEXT("GeomSubset")))
+					{
+						PrimPath = PrimPath.GetParentPath();
+					}
+
+					continue;
 				}
 
-				// We put all of the audio info directly on the section, so if any of these change then we need to re-add the prim
-				// to the sequencer, and potentially generate a new audio asset
-				static const TSet<FString> AudioSectionProperties = {
-					*UsdToUnreal::ConvertToken(pxr::UsdMediaTokens->filePath),
-					*UsdToUnreal::ConvertToken(pxr::UsdMediaTokens->auralMode),
-					*UsdToUnreal::ConvertToken(pxr::UsdMediaTokens->playbackMode),
-					*UsdToUnreal::ConvertToken(pxr::UsdMediaTokens->startTime),
-					*UsdToUnreal::ConvertToken(pxr::UsdMediaTokens->endTime),
-					*UsdToUnreal::ConvertToken(pxr::UsdMediaTokens->mediaOffset),
-					*UsdToUnreal::ConvertToken(pxr::UsdMediaTokens->gain)};
-				if (AudioSectionProperties.Contains(AttributeChange.PropertyName))
+				for (const UsdUtils::FAttributeChange& AttributeChange : ObjectChange.AttributeChanges)
 				{
-					bIsResync = true;
-					bHasResync = true;
+					static const TSet<FString> StageResyncProperties = {TEXT("metersPerUnit"), TEXT("upAxis")};
+
+					// Upgrade these to resync so that the prim twins are regenerated, which clears all the existing
+					// animation tracks and adds new ones, automatically re-baking to control rig
+					static const TSet<FString> PrimResyncProperties = {
+						*UsdToUnreal::ConvertToken(UnrealIdentifiers::UnrealControlRigPath),
+						*UsdToUnreal::ConvertToken(UnrealIdentifiers::UnrealUseFKControlRig),
+						*UsdToUnreal::ConvertToken(UnrealIdentifiers::UnrealControlRigReduceKeys),
+						*UsdToUnreal::ConvertToken(UnrealIdentifiers::UnrealControlRigReductionTolerance),
+						// For now we need to do this as we need to refresh the material slot info on the info cache if these
+						// update... this is of course way too aggressive, although it's unlikely people will be manually editing these.
+						// TODO: More nuanced info cache updates
+						UnrealIdentifiers::PrimvarsDisplayColor,
+						UnrealIdentifiers::PrimvarsDisplayOpacity,
+						UnrealIdentifiers::DoubleSided,
+						// When we change these UsdGeomModelAPI attributes we may need to create a new component type
+						// for the prim (as it may now need/stop needing an alternate draw mode component)
+						UnrealIdentifiers::ModelDrawMode,
+						UnrealIdentifiers::ModelApplyDrawMode,
+						// Collapsing changes need to resync to build the new collapsed/uncollapsed meshes
+						*UsdToUnreal::ConvertToken(UnrealIdentifiers::UnrealCollapsingAttr),
+						// Physics collision attribute change needs to rebuild the collision shape
+						*UsdToUnreal::ConvertToken(pxr::UsdPhysicsTokens->physicsCollisionEnabled),
+						*UsdToUnreal::ConvertToken(pxr::UsdPhysicsTokens->physicsApproximation)};
+
+					// Some stage info should trigger some resyncs because they should trigger reparsing of geometry
+					if ((PrimPath.IsAbsoluteRootPath() && StageResyncProperties.Contains(AttributeChange.PropertyName))
+						|| PrimResyncProperties.Contains(AttributeChange.PropertyName))
+					{
+						bIsResync = true;
+						bHasResync = true;
+						break;
+					}
+
+					// We put all of the audio info directly on the section, so if any of these change then we need to re-add the prim
+					// to the sequencer, and potentially generate a new audio asset
+					static const TSet<FString> AudioSectionProperties = {
+						*UsdToUnreal::ConvertToken(pxr::UsdMediaTokens->filePath),
+						*UsdToUnreal::ConvertToken(pxr::UsdMediaTokens->auralMode),
+						*UsdToUnreal::ConvertToken(pxr::UsdMediaTokens->playbackMode),
+						*UsdToUnreal::ConvertToken(pxr::UsdMediaTokens->startTime),
+						*UsdToUnreal::ConvertToken(pxr::UsdMediaTokens->endTime),
+						*UsdToUnreal::ConvertToken(pxr::UsdMediaTokens->mediaOffset),
+						*UsdToUnreal::ConvertToken(pxr::UsdMediaTokens->gain)};
+					if (AudioSectionProperties.Contains(AttributeChange.PropertyName))
+					{
+						bIsResync = true;
+						bHasResync = true;
+						break;
+					}
+				}
+
+				// Any sublayer change (even offsets) means we need to regenerate our LevelSequence to add (or shift)
+				// the corresponding subsequences
+				if (PrimPath.IsAbsoluteRootPath() && ObjectChange.SubLayerChanges.Num() > 0)
+				{
+					bNeedsAnimationReload = true;
 					break;
 				}
 			}
 
-			// Any sublayer change (even offsets) means we need to regenerate our LevelSequence to add (or shift)
-			// the corresponding subsequences
-			if (PrimPath.IsAbsoluteRootPath() && ObjectChange.SubLayerChanges.Num() > 0)
-			{
-				bNeedsAnimationReload = true;
-				break;
-			}
+			// We may need the full spec path with variant selections later, but for traversal and retrieving prims from the stage we always need
+			// the prim path without any variant selections in it (i.e. GetPrimAtPath("/Root{Varset=Var}Child") doesn't work, we need
+			// GetPrimAtPath("/Root/Child")), and USD sometimes emits changes with the variant selection path (like during renames).
+			SortedPrimsChangedList.Add(PrimPath, bIsResync);
 		}
+		// Do Resyncs after so that they overwrite pure info changes if we have any
+		for (const TPair<FString, TArray<UsdUtils::FSdfChangeListEntry>>& ResyncChange : AccumulatedResyncChanges)
+		{
+			UE::FSdfPath PrimPath = UE::FSdfPath(*ResyncChange.Key).StripAllVariantSelections();
 
-		// We may need the full spec path with variant selections later, but for traversal and retrieving prims from the stage we always need
-		// the prim path without any variant selections in it (i.e. GetPrimAtPath("/Root{Varset=Var}Child") doesn't work, we need
-		// GetPrimAtPath("/Root/Child")), and USD sometimes emits changes with the variant selection path (like during renames).
-		SortedPrimsChangedList.Add(PrimPath, bIsResync);
-	}
-	// Do Resyncs after so that they overwrite pure info changes if we have any
-	for (const TPair<FString, TArray<UsdUtils::FSdfChangeListEntry>>& ResyncChange : ResyncChanges)
-	{
-		UE::FSdfPath PrimPath = UE::FSdfPath(*ResyncChange.Key).StripAllVariantSelections();
-
-		const bool bIsResync = true;
-		SortedPrimsChangedList.Add(PrimPath, bIsResync);
+			const bool bIsResync = true;
+			SortedPrimsChangedList.Add(PrimPath, bIsResync);
+		}
 	}
 
 	// During PIE, the PIE and the editor world will respond to notices. We have to prevent any PIE
@@ -1488,7 +1534,10 @@ void AUsdStageActor::OnUsdObjectsChanged(const UsdUtils::FObjectChangesByPath& I
 	{
 		// Take a copy of the info cache here: We want to keep the old one for this function call as it helps us cleanup the old assets and components
 		// For now we don't need to do this for info changes, only for resync changes
-		InfoCache = MakeShared<FUsdInfoCache>(*OldInfoCache);
+		{
+			TRACE_CPUPROFILER_EVENT_SCOPE(CopyingOldInfoCache);
+			InfoCache = MakeShared<FUsdInfoCache>(*OldInfoCache);
+		}
 
 		// TODO: Selective rebuild of only the required parts of the cache.
 		// If a prim changes from CanBeCollapsed to not (or vice-versa), that means its parent may change, and its grandparent, etc. so we'd need
@@ -1511,6 +1560,8 @@ void AUsdStageActor::OnUsdObjectsChanged(const UsdUtils::FObjectChangesByPath& I
 
 	TFunction<void(TMap<UE::FSdfPath, bool>&)> SortAndCleanPrimsToUpdate = [](TMap<UE::FSdfPath, bool>& InOutMap)
 	{
+		TRACE_CPUPROFILER_EVENT_SCOPE(SortAndCleanPrimsToUpdate);
+
 		// Force traversal of parent before child (not needed for cleaning up, but very much needed later for reloading assets/components)
 		InOutMap.KeySort(
 			[](const UE::FSdfPath& A, const UE::FSdfPath& B) -> bool
@@ -1578,6 +1629,8 @@ void AUsdStageActor::OnUsdObjectsChanged(const UsdUtils::FObjectChangesByPath& I
 		 &KnownOldPrims,
 		 &Stage](const UE::FSdfPath& PrimPath, bool bIsResync, TMap<UE::FSdfPath, bool>& OutPrimsToUpdate, TSet<UE::FSdfPath>& InOutVisitedPaths)
 	{
+		TRACE_CPUPROFILER_EVENT_SCOPE(RecursiveCollectPrimsToUpdate);
+
 		// Allow revisiting a prim if we're resyncing: This is a cheap way of making sure we mark the
 		// resynced prims as resynced, or else we may recurse into them while traversing another prim and preemptively
 		// mark them as visited, and because we always recurse with resync=false we'd assume those aren't resyncs.
@@ -1675,6 +1728,8 @@ void AUsdStageActor::OnUsdObjectsChanged(const UsdUtils::FObjectChangesByPath& I
 	// Collect all the paths to update
 	TMap<UE::FSdfPath, bool> PrimsToUpdate;
 	{
+		TRACE_CPUPROFILER_EVENT_SCOPE(CollectPathsToUpdate);
+
 		PrimsToUpdate.Reserve(SortedPrimsChangedList.Num());
 
 		// Recursively append main prims to the list of PrimsToUpdate
@@ -1702,23 +1757,27 @@ void AUsdStageActor::OnUsdObjectsChanged(const UsdUtils::FObjectChangesByPath& I
 
 	RefreshStageTask.EnterProgressFrame();
 	FScopedSlowTask CleanUpAssetsTask(PrimsToUpdate.Num(), LOCTEXT("CleaningUpAssets", "Cleaning up old assets"));
-	for (const TPair<UE::FSdfPath, bool>& PrimChangedInfo : PrimsToUpdate)
 	{
-		CleanUpAssetsTask.EnterProgressFrame();
-
-		const UE::FSdfPath& PrimPath = PrimChangedInfo.Key;
-		const bool bIsResync = PrimChangedInfo.Value;
+		TRACE_CPUPROFILER_EVENT_SCOPE(CleaningUpAssets);
 
-		// If it's a new prim, we may not have old info about it
-		if (!OldInfoCache->ContainsInfoAboutPrim(PrimPath))
+		for (const TPair<UE::FSdfPath, bool>& PrimChangedInfo : PrimsToUpdate)
 		{
-			continue;
-		}
+			CleanUpAssetsTask.EnterProgressFrame();
+
+			const UE::FSdfPath& PrimPath = PrimChangedInfo.Key;
+			const bool bIsResync = PrimChangedInfo.Value;
+
+			// If it's a new prim, we may not have old info about it
+			if (!OldInfoCache->ContainsInfoAboutPrim(PrimPath))
+			{
+				continue;
+			}
 
-		UE_LOG(LogUsd, Verbose, TEXT("Cleaning up assets for prim path '%s'"), *PrimPath.GetString());
+			UE_LOG(LogUsd, Verbose, TEXT("Cleaning up assets for prim path '%s'"), *PrimPath.GetString());
 
-		const bool bForEntireSubtree = bIsResync;
-		bHasLoadedOrAbandonedAssets |= UnloadAssets(PrimPath, bForEntireSubtree);
+			const bool bForEntireSubtree = bIsResync;
+			bHasLoadedOrAbandonedAssets |= UnloadAssets(PrimPath, bForEntireSubtree);
+		}
 	}
 
 	// Regenerate assets before cleaning up components because if material assignments change, we may
@@ -1726,57 +1785,63 @@ void AUsdStageActor::OnUsdObjectsChanged(const UsdUtils::FObjectChangesByPath& I
 	TSet<UE::FSdfPath> MaterialUserPrims;
 	RefreshStageTask.EnterProgressFrame();
 	FScopedSlowTask RegenerateAssetsTask(PrimsToUpdate.Num(), LOCTEXT("RegeneratingAssets", "Regenerating assets"));
-	for (const TPair<UE::FSdfPath, bool>& PrimChangedInfo : PrimsToUpdate)
 	{
-		RegenerateAssetsTask.EnterProgressFrame();
+		TRACE_CPUPROFILER_EVENT_SCOPE(RegeneratingAssets);
 
-		const UE::FSdfPath& PrimPath = PrimChangedInfo.Key;
-		const bool bIsResync = PrimChangedInfo.Value;
+		for (const TPair<UE::FSdfPath, bool>& PrimChangedInfo : PrimsToUpdate)
+		{
+			RegenerateAssetsTask.EnterProgressFrame();
 
-		UE::FUsdPrim PrimToUpdate = Stage.GetPrimAtPath(PrimPath);
+			const UE::FSdfPath& PrimPath = PrimChangedInfo.Key;
+			const bool bIsResync = PrimChangedInfo.Value;
 
-		// It's OK to not have info about a prim if it's an old prim that only exists on the old info cache.
-		// If the new info cache has info about this prim then it must exist on the stage right now
-		if (!InfoCache->ContainsInfoAboutPrim(PrimPath))
-		{
-			ensure(!PrimToUpdate);
-			continue;
-		}
-		ensure(PrimToUpdate);
+			UE::FUsdPrim PrimToUpdate = Stage.GetPrimAtPath(PrimPath);
 
-		UE_LOG(LogUsd, Verbose, TEXT("Regenerating assets for prim path '%s'"), *PrimPath.GetString());
+			// It's OK to not have info about a prim if it's an old prim that only exists on the old info cache.
+			// If the new info cache has info about this prim then it must exist on the stage right now
+			if (!InfoCache->ContainsInfoAboutPrim(PrimPath))
+			{
+				ensure(!PrimToUpdate);
+				continue;
+			}
+			ensure(PrimToUpdate);
 
-		TSharedRef<FUsdSchemaTranslationContext> TranslationContext = FUsdStageActorImpl::CreateUsdSchemaTranslationContext(
-			this,
-			PrimPath.GetString()
-		);
+			UE_LOG(LogUsd, Verbose, TEXT("Regenerating assets for prim path '%s'"), *PrimPath.GetString());
 
-		bool bThisPrimLoadedAssets = false;
-		if (bIsResync)
-		{
-			bThisPrimLoadedAssets |= LoadAssets(*TranslationContext, PrimToUpdate);
-		}
-		else
-		{
-			bThisPrimLoadedAssets |= LoadAsset(*TranslationContext, PrimToUpdate);
-		}
-		bHasLoadedOrAbandonedAssets |= bThisPrimLoadedAssets;
+			TSharedRef<FUsdSchemaTranslationContext> TranslationContext = FUsdStageActorImpl::CreateUsdSchemaTranslationContext(
+				this,
+				PrimPath.GetString()
+			);
 
-		// For UE-120185: If we recreated a material for a prim path we also need to update all components that were using it.
-		// This could be fleshed out further if other asset types require this refresh of "consumer components" but materials
-		// seem to be the only ones that do at the moment.
-		// Note that even after UE-157644 this is still useful: Material prims are not marked as a dependency
-		// of Mesh prims, otherwise we'd have to regenerate the StaticMesh itself when material info changed.
-		// Ideally we'd just set a new material override on the component, which is what this does
-		const bool bPrimGeneratesMaterials = InfoCache->GetSingleAssetForPrim<UMaterialInterface>(PrimPath) != nullptr;
-		if (bThisPrimLoadedAssets && bPrimGeneratesMaterials)
-		{
-			MaterialUserPrims.Append(InfoCache->GetMaterialUsers(PrimPath));
+			bool bThisPrimLoadedAssets = false;
+			if (bIsResync)
+			{
+				bThisPrimLoadedAssets |= LoadAssets(*TranslationContext, PrimToUpdate);
+			}
+			else
+			{
+				bThisPrimLoadedAssets |= LoadAsset(*TranslationContext, PrimToUpdate);
+			}
+			bHasLoadedOrAbandonedAssets |= bThisPrimLoadedAssets;
+
+			// For UE-120185: If we recreated a material for a prim path we also need to update all components that were using it.
+			// This could be fleshed out further if other asset types require this refresh of "consumer components" but materials
+			// seem to be the only ones that do at the moment.
+			// Note that even after UE-157644 this is still useful: Material prims are not marked as a dependency
+			// of Mesh prims, otherwise we'd have to regenerate the StaticMesh itself when material info changed.
+			// Ideally we'd just set a new material override on the component, which is what this does
+			const bool bPrimGeneratesMaterials = InfoCache->GetSingleAssetForPrim<UMaterialInterface>(PrimPath) != nullptr;
+			if (bThisPrimLoadedAssets && bPrimGeneratesMaterials)
+			{
+				MaterialUserPrims.Append(InfoCache->GetMaterialUsers(PrimPath));
+			}
 		}
 	}
 
 	if (MaterialUserPrims.Num())
 	{
+		TRACE_CPUPROFILER_EVENT_SCOPE(MaterialUserPrimsPass);
+
 		TSet<UE::FSdfPath> VisitedPaths;
 		PrimsToUpdate.Reserve(PrimsToUpdate.Num() + MaterialUserPrims.Num());
 		for (const UE::FSdfPath& MaterialUserPrim : MaterialUserPrims)
@@ -1792,26 +1857,30 @@ void AUsdStageActor::OnUsdObjectsChanged(const UsdUtils::FObjectChangesByPath& I
 
 	RefreshStageTask.EnterProgressFrame();
 	FScopedSlowTask CleanUpComponentsTask(PrimsToUpdate.Num(), LOCTEXT("CleaningUpComponents", "Cleaning up actors and components"));
-	for (const TPair<UE::FSdfPath, bool>& PrimChangedInfo : PrimsToUpdate)
 	{
-		CleanUpComponentsTask.EnterProgressFrame();
+		TRACE_CPUPROFILER_EVENT_SCOPE(CleaningUpComponents);
 
-		const UE::FSdfPath PrimPath = PrimChangedInfo.Key;
-		const bool bIsResync = PrimChangedInfo.Value;
-
-		// If it's a new prim, we may not have old info about it
-		if (!OldInfoCache->ContainsInfoAboutPrim(PrimPath))
+		for (const TPair<UE::FSdfPath, bool>& PrimChangedInfo : PrimsToUpdate)
 		{
-			continue;
-		}
+			CleanUpComponentsTask.EnterProgressFrame();
+
+			const UE::FSdfPath PrimPath = PrimChangedInfo.Key;
+			const bool bIsResync = PrimChangedInfo.Value;
 
-		UE_LOG(LogUsd, Verbose, TEXT("Cleaning up actors and components for prim path '%s'"), *PrimPath.GetString());
+			// If it's a new prim, we may not have old info about it
+			if (!OldInfoCache->ContainsInfoAboutPrim(PrimPath))
+			{
+				continue;
+			}
 
-		if (bIsResync && PrimPath.IsAbsoluteRootOrPrimPath())
-		{
-			if (UUsdPrimTwin* UsdPrimTwin = GetRootPrimTwin()->Find(PrimPath.GetString()))
+			UE_LOG(LogUsd, Verbose, TEXT("Cleaning up actors and components for prim path '%s'"), *PrimPath.GetString());
+
+			if (bIsResync && PrimPath.IsAbsoluteRootOrPrimPath())
 			{
-				UsdPrimTwin->Clear();
+				if (UUsdPrimTwin* UsdPrimTwin = GetRootPrimTwin()->Find(PrimPath.GetString()))
+				{
+					UsdPrimTwin->Clear();
+				}
 			}
 		}
 	}
@@ -1825,34 +1894,42 @@ void AUsdStageActor::OnUsdObjectsChanged(const UsdUtils::FObjectChangesByPath& I
 
 	RefreshStageTask.EnterProgressFrame();
 	FScopedSlowTask RegenerateComponentsTask(PrimsToUpdate.Num(), LOCTEXT("RegeneratingComponents", "Regenerating components"));
-	for (const TPair<UE::FSdfPath, bool>& PrimChangedInfo : PrimsToUpdate)
 	{
-		RegenerateComponentsTask.EnterProgressFrame();
-
-		const UE::FSdfPath PrimPath = PrimChangedInfo.Key;
-		const bool bIsResync = PrimChangedInfo.Value;
+		TRACE_CPUPROFILER_EVENT_SCOPE(RegeneratingComponents);
 
-		if (!InfoCache->ContainsInfoAboutPrim(PrimPath))
+		for (const TPair<UE::FSdfPath, bool>& PrimChangedInfo : PrimsToUpdate)
 		{
-			continue;
-		}
+			RegenerateComponentsTask.EnterProgressFrame();
 
-		UE_LOG(LogUsd, Verbose, TEXT("Regenerating components for prim path '%s'"), *PrimPath.GetString());
+			const UE::FSdfPath PrimPath = PrimChangedInfo.Key;
+			const bool bIsResync = PrimChangedInfo.Value;
 
-		TSharedRef<FUsdSchemaTranslationContext> TranslationContext = FUsdStageActorImpl::CreateUsdSchemaTranslationContext(
-			this,
-			PrimPath.GetString()
-		);
+			if (!InfoCache->ContainsInfoAboutPrim(PrimPath))
+			{
+				continue;
+			}
+
+			UE_LOG(LogUsd, Verbose, TEXT("Regenerating components for prim path '%s'"), *PrimPath.GetString());
 
-		UpdatePrim(PrimPath, bIsResync, *TranslationContext);
-		TranslationContext->CompleteTasks();
+			TSharedRef<FUsdSchemaTranslationContext> TranslationContext = FUsdStageActorImpl::CreateUsdSchemaTranslationContext(
+				this,
+				PrimPath.GetString()
+			);
+
+			UpdatePrim(PrimPath, bIsResync, *TranslationContext);
+			TranslationContext->CompleteTasks();
+		}
 	}
 
 	// Separate pass because we need to update the stage editor even if we don't have info about this prim
 	// anymore, as that's how it refreshes whenever we delete a prim
-	for (const TPair<UE::FSdfPath, bool>& PrimChangedInfo : PrimsToUpdate)
 	{
-		OnPrimChanged.Broadcast(PrimChangedInfo.Key.GetString(), PrimChangedInfo.Value);
+		TRACE_CPUPROFILER_EVENT_SCOPE(OnPrimChangedBroadcast);
+
+		for (const TPair<UE::FSdfPath, bool>& PrimChangedInfo : PrimsToUpdate)
+		{
+			OnPrimChanged.Broadcast(PrimChangedInfo.Key.GetString(), PrimChangedInfo.Value);
+		}
 	}
 
 	if (bHasResync)
@@ -1863,6 +1940,8 @@ void AUsdStageActor::OnUsdObjectsChanged(const UsdUtils::FObjectChangesByPath& I
 #if WITH_EDITOR
 	if (GIsEditor && GEditor && !IsGarbageCollecting())	   // Make sure we're not in standalone either
 	{
+		TRACE_CPUPROFILER_EVENT_SCOPE(BroadcastingEditorEvents);
+
 		if (bHasResync)
 		{
 			GEditor->BroadcastLevelActorListChanged();
@@ -1873,6 +1952,9 @@ void AUsdStageActor::OnUsdObjectsChanged(const UsdUtils::FObjectChangesByPath& I
 #endif	  // WITH_EDITOR
 
 #endif	  // USE_USD_SDK
+
+	AccumulatedInfoChanges.Reset();
+	AccumulatedResyncChanges.Reset();
 }
 
 USDSTAGE_API void AUsdStageActor::Reset()
@@ -3699,6 +3781,21 @@ void AUsdStageActor::HandleTransactionStateChanged(
 	const ETransactionStateEventType InTransactionState
 )
 {
+	// Handle any accumulated USD notices we received during the transaction
+	if (!GHandleNoticesImmediately && GIsEditor)
+	{
+		if (InTransactionState == ETransactionStateEventType::PreTransactionFinalized)
+		{
+			HandleAccumulatedNotices();
+		}
+		else if (InTransactionState == ETransactionStateEventType::TransactionFinalized)
+		{
+			ensureAlways(AccumulatedInfoChanges.Num() == 0 && AccumulatedResyncChanges.Num() == 0);
+			AccumulatedInfoChanges.Reset();
+			AccumulatedResyncChanges.Reset();
+		}
+	}
+
 	// Hack for solving UE-127253
 	// When we Reload (or open a new stage), we call ReloadAnimations which will close the Sequencer (if opened), recreate our LevelSequence, and get
 	// the Sequencer to show that one instead. If we undo the Reload, that new LevelSequence will be deleted and the Sequencer will be left open
diff --git a/Engine/Plugins/Importers/USDImporter/Source/USDStage/Private/USDTransactor.cpp b/Engine/Plugins/Importers/USDImporter/Source/USDStage/Private/USDTransactor.cpp
--- a/Engine/Plugins/Importers/USDImporter/Source/USDStage/Private/USDTransactor.cpp
+++ b/Engine/Plugins/Importers/USDImporter/Source/USDStage/Private/USDTransactor.cpp
@@ -703,7 +703,7 @@ namespace UsdUtils
 	}
 
 	/** Applies the field value pairs to all prims on the stage, and returns a list of prim paths for modified prims */
-	TArray<FString> ApplyFieldMapToStage(
+	TSet<FString> ApplyFieldMapToStage(
 		const FTransactorEditStorage& EditStorage,
 		EApplicationDirection Direction,
 		UE::FUsdStage& Stage,
@@ -715,7 +715,7 @@ namespace UsdUtils
 			return {};
 		}
 
-		TArray<FString> PrimsChanged;
+		TSet<FString> PrimsChanged;
 
 		int32 Start = 0;
 		int32 End = 0;
@@ -992,7 +992,7 @@ namespace UsdUtils
 
 			UE::FUsdStage& Stage = StageActor->GetOrOpenUsdStage();
 
-			TArray<FString>
+			TSet<FString>
 				PrimsChanged = UsdUtils::ApplyFieldMapToStage(Values, UsdUtils::EApplicationDirection::Reverse, Stage, StageActor->GetTime());
 
 			if (PrimsChanged.Num() > 0)
@@ -1039,7 +1039,7 @@ namespace UsdUtils
 
 			UE::FUsdStage& Stage = StageActor->GetOrOpenUsdStage();
 
-			TArray<FString> PrimsChanged;
+			TSet<FString> PrimsChanged;
 			if (bIsApplyingConcertSync && ReceivedValuesBeforeUndo.IsSet())
 			{
 				// If we're applying a received ConcertSync transaction that actually is an undo on the source client then we want to use it's
@@ -1172,6 +1172,8 @@ void UUsdTransactor::Initialize(AUsdStageActor* InStageActor)
 
 void UUsdTransactor::Update(const UsdUtils::FObjectChangesByPath& NewInfoChanges, const UsdUtils::FObjectChangesByPath& NewResyncChanges)
 {
+	TRACE_CPUPROFILER_EVENT_SCOPE(UUsdTransactor::Update);
+
 	// We always send notices even when we're undoing/redoing changes (so that multi-user can broadcast them).
 	// Make sure that we only ever update our OldValues/NewValues when we receive *new* updates though
 	if (Impl.IsValid() && (Impl->IsTransactionUndoing() || Impl->IsTransactionRedoing() || Impl->IsApplyingConcertSyncTransaction()))
diff --git a/Engine/Plugins/Importers/USDImporter/Source/USDStage/Public/USDStageActor.h b/Engine/Plugins/Importers/USDImporter/Source/USDStage/Public/USDStageActor.h
--- a/Engine/Plugins/Importers/USDImporter/Source/USDStage/Public/USDStageActor.h
+++ b/Engine/Plugins/Importers/USDImporter/Source/USDStage/Public/USDStageActor.h
@@ -408,6 +408,7 @@ protected:
 	void OnPreUsdImport(FString FilePath);
 	void OnPostUsdImport(FString FilePath);
 	void OnUsdObjectsChanged(const UsdUtils::FObjectChangesByPath& InfoChanges, const UsdUtils::FObjectChangesByPath& ResyncChanges);
+	void HandleAccumulatedNotices();
 	void OnUsdPrimTwinDestroyed(const UUsdPrimTwin& UsdPrimTwin);
 	void OnObjectPropertyChanged(UObject* ObjectBeingModified, FPropertyChangedEvent& PropertyChangedEvent);
 	void HandlePropertyChangedEvent(FPropertyChangedEvent& PropertyChangedEvent);
@@ -445,6 +446,13 @@ protected:
 	UPROPERTY(Transient)
 	TObjectPtr<UUsdTransactor> Transactor;
 
+	/**
+	 * USD can emit multiple notices within the same UE transaction. We accumulate these in here, and respond
+	 * to them only once, when the transaction is about to finish
+	 */
+	UsdUtils::FObjectChangesByPath AccumulatedInfoChanges;
+	UsdUtils::FObjectChangesByPath AccumulatedResyncChanges;
+
 	/** Caches various information about prims that are expensive to query */
 	TSharedPtr<FUsdInfoCache> InfoCache;
 
