USD: Fix broken ContentWorkers build.

diff --git a/Engine/Plugins/Importers/USDImporter/Source/USDStage/Private/USDStageActor.cpp b/Engine/Plugins/Importers/USDImporter/Source/USDStage/Private/USDStageActor.cpp
--- a/Engine/Plugins/Importers/USDImporter/Source/USDStage/Private/USDStageActor.cpp
+++ b/Engine/Plugins/Importers/USDImporter/Source/USDStage/Private/USDStageActor.cpp
@@ -3626,7 +3626,6 @@ const FUsdListener& AUsdStageActor::GetUsdListener() const
 }
 
 #if WITH_EDITOR
-
 void AUsdStageActor::PostEditChangeProperty(FPropertyChangedEvent& PropertyChangedEvent)
 {
 	// For handling root layer changes via direct changes to properties we want to go through OnObjectPropertyChanged -> HandlePropertyChangedEvent ->
@@ -3800,7 +3799,6 @@ void AUsdStageActor::HandleTransactionStateChanged(
 	// When we Reload (or open a new stage), we call ReloadAnimations which will close the Sequencer (if opened), recreate our LevelSequence, and get
 	// the Sequencer to show that one instead. If we undo the Reload, that new LevelSequence will be deleted and the Sequencer will be left open
 	// trying to display it, which leads to crashes. Here we try detecting for that case and close/reopen the sequencer to show the correct one.
-#if WITH_EDITOR
 	if (GIsEditor && LevelSequence
 		&& (InTransactionState == ETransactionStateEventType::UndoRedoStarted || InTransactionState == ETransactionStateEventType::UndoRedoFinalized))
 	{
@@ -3902,7 +3900,6 @@ void AUsdStageActor::HandleTransactionStateChanged(
 			}
 		}
 	}
-#endif	  // WITH_EDITOR
 
 	if (InTransactionState == ETransactionStateEventType::TransactionFinalized || InTransactionState == ETransactionStateEventType::UndoRedoFinalized
 		|| InTransactionState == ETransactionStateEventType::TransactionCanceled)
@@ -3910,7 +3907,6 @@ void AUsdStageActor::HandleTransactionStateChanged(
 		OldRootLayer = RootLayer;
 	}
 }
-
 #endif	  // WITH_EDITOR
 
 void AUsdStageActor::PostDuplicate(bool bDuplicateForPIE)
diff --git a/Engine/Plugins/Importers/USDImporter/Source/USDStage/Public/USDStageActor.h b/Engine/Plugins/Importers/USDImporter/Source/USDStage/Public/USDStageActor.h
--- a/Engine/Plugins/Importers/USDImporter/Source/USDStage/Public/USDStageActor.h
+++ b/Engine/Plugins/Importers/USDImporter/Source/USDStage/Public/USDStageActor.h
@@ -408,7 +408,7 @@ protected:
 	void OnPreUsdImport(FString FilePath);
 	void OnPostUsdImport(FString FilePath);
 	void OnUsdObjectsChanged(const UsdUtils::FObjectChangesByPath& InfoChanges, const UsdUtils::FObjectChangesByPath& ResyncChanges);
-	void HandleAccumulatedNotices();
+	USDSTAGE_API void HandleAccumulatedNotices();
 	void OnUsdPrimTwinDestroyed(const UUsdPrimTwin& UsdPrimTwin);
 	void OnObjectPropertyChanged(UObject* ObjectBeingModified, FPropertyChangedEvent& PropertyChangedEvent);
 	void HandlePropertyChangedEvent(FPropertyChangedEvent& PropertyChangedEvent);
