USD: Actually fix the broken ContentWorkers build.

diff --git a/Engine/Plugins/Importers/USDImporter/Source/USDStage/Private/USDStageActor.cpp b/Engine/Plugins/Importers/USDImporter/Source/USDStage/Private/USDStageActor.cpp
--- a/Engine/Plugins/Importers/USDImporter/Source/USDStage/Private/USDStageActor.cpp
+++ b/Engine/Plugins/Importers/USDImporter/Source/USDStage/Private/USDStageActor.cpp
@@ -1337,12 +1337,15 @@ void AUsdStageActor::OnUsdObjectsChanged(const UsdUtils::FObjectChangesByPath& I
 	{
 		HandleAccumulatedNotices();
 	}
+#endif	  // USE_USD_SDK
 }
 
 void AUsdStageActor::HandleAccumulatedNotices()
 {
 	TRACE_CPUPROFILER_EVENT_SCOPE(AUsdStageActor::HandleAccumulatedNotices);
 
+#if USE_USD_SDK
+
 	if (AccumulatedInfoChanges.Num() == 0 && AccumulatedResyncChanges.Num() == 0)
 	{
 		return;
diff --git a/Engine/Plugins/Importers/USDImporter/Source/USDStage/Public/USDStageActor.h b/Engine/Plugins/Importers/USDImporter/Source/USDStage/Public/USDStageActor.h
--- a/Engine/Plugins/Importers/USDImporter/Source/USDStage/Public/USDStageActor.h
+++ b/Engine/Plugins/Importers/USDImporter/Source/USDStage/Public/USDStageActor.h
@@ -408,7 +408,7 @@ protected:
 	void OnPreUsdImport(FString FilePath);
 	void OnPostUsdImport(FString FilePath);
 	void OnUsdObjectsChanged(const UsdUtils::FObjectChangesByPath& InfoChanges, const UsdUtils::FObjectChangesByPath& ResyncChanges);
-	USDSTAGE_API void HandleAccumulatedNotices();
+	void HandleAccumulatedNotices();
 	void OnUsdPrimTwinDestroyed(const UUsdPrimTwin& UsdPrimTwin);
 	void OnObjectPropertyChanged(UObject* ObjectBeingModified, FPropertyChangedEvent& PropertyChangedEvent);
 	void HandlePropertyChangedEvent(FPropertyChangedEvent& PropertyChangedEvent);
