USD: Made the USD components outer'ed to their actors instead of a component.

diff --git a/Engine/Plugins/Importers/USDImporter/Source/USDSchemas/Private/USDGeomXformableTranslator.cpp b/Engine/Plugins/Importers/USDImporter/Source/USDSchemas/Private/USDGeomXformableTranslator.cpp
--- a/Engine/Plugins/Importers/USDImporter/Source/USDSchemas/Private/USDGeomXformableTranslator.cpp
+++ b/Engine/Plugins/Importers/USDImporter/Source/USDSchemas/Private/USDGeomXformableTranslator.cpp
@@ -515,7 +515,7 @@ USceneComponent* FUsdGeomXformableTranslator::CreateComponentsEx(TOptional<TSubc
 	}
 	else
 	{
-		ComponentOuter = Context->ParentComponent;
+		ComponentOuter = Context->ParentComponent->GetOwner();
 	}
 
 	if (!ComponentOuter)
