USD: Refactor FUsdInfoCache into two UObjects so that we can better control when things get serialized.

diff --git a/Engine/Plugins/Importers/USDImporter/Source/USDSchemas/Private/Custom/MDLUSDShadeMaterialTranslator.cpp b/Engine/Plugins/Importers/USDImporter/Source/USDSchemas/Private/Custom/MDLUSDShadeMaterialTranslator.cpp
--- a/Engine/Plugins/Importers/USDImporter/Source/USDSchemas/Private/Custom/MDLUSDShadeMaterialTranslator.cpp
+++ b/Engine/Plugins/Importers/USDImporter/Source/USDSchemas/Private/Custom/MDLUSDShadeMaterialTranslator.cpp
@@ -78,9 +78,9 @@ void FMdlUsdShadeMaterialTranslator::CreateAssets()
 		return;
 	}
 
-	if (Context->bTranslateOnlyUsedMaterials && Context->InfoCache)
+	if (Context->bTranslateOnlyUsedMaterials && Context->UsdInfoCache)
 	{
-		if (!Context->InfoCache->IsMaterialUsed(PrimPath))
+		if (!Context->UsdInfoCache->IsMaterialUsed(PrimPath))
 		{
 			UE_LOG(
 				LogUsd,
diff --git a/Engine/Plugins/Importers/USDImporter/Source/USDSchemas/Private/Custom/MaterialXUSDShadeMaterialTranslator.cpp b/Engine/Plugins/Importers/USDImporter/Source/USDSchemas/Private/Custom/MaterialXUSDShadeMaterialTranslator.cpp
--- a/Engine/Plugins/Importers/USDImporter/Source/USDSchemas/Private/Custom/MaterialXUSDShadeMaterialTranslator.cpp
+++ b/Engine/Plugins/Importers/USDImporter/Source/USDSchemas/Private/Custom/MaterialXUSDShadeMaterialTranslator.cpp
@@ -269,11 +269,8 @@ namespace UE::USDMaterialXTranslator::Private
 
 							// Rename the UMaterialInterface into the target UPackage the asset cache created for us.
 							// SanitizedName will already match it.
-							const bool bRenamed = ImportedObject->Rename(
-								*SanitizedName.ToString(),
-								Outer,
-								REN_NonTransactional | REN_DontCreateRedirectors
-							);
+							const bool bRenamed = ImportedObject
+													  ->Rename(*SanitizedName.ToString(), Outer, REN_NonTransactional | REN_DontCreateRedirectors);
 							ensure(bRenamed);
 
 							ImportedObject->ClearFlags(ImportedObject->GetFlags());
@@ -372,12 +369,12 @@ void FMaterialXUsdShadeMaterialTranslator::CreateAssets()
 		return;
 	}
 
-	if (!Context->UsdAssetCache || !Context->InfoCache)
+	if (!Context->UsdAssetCache || !Context->UsdInfoCache)
 	{
 		return;
 	}
 
-	if (Context->bTranslateOnlyUsedMaterials && !Context->InfoCache->IsMaterialUsed(PrimPath))
+	if (Context->bTranslateOnlyUsedMaterials && !Context->UsdInfoCache->IsMaterialUsed(PrimPath))
 	{
 		UE_LOG(
 			LogUsd,
diff --git a/Engine/Plugins/Importers/USDImporter/Source/USDSchemas/Private/MeshTranslationImpl.cpp b/Engine/Plugins/Importers/USDImporter/Source/USDSchemas/Private/MeshTranslationImpl.cpp
--- a/Engine/Plugins/Importers/USDImporter/Source/USDSchemas/Private/MeshTranslationImpl.cpp
+++ b/Engine/Plugins/Importers/USDImporter/Source/USDSchemas/Private/MeshTranslationImpl.cpp
@@ -11,6 +11,7 @@
 #include "USDLog.h"
 #include "USDMemory.h"
 #include "USDObjectUtils.h"
+#include "USDPrimLinkCache.h"
 #include "USDProjectSettings.h"
 #include "USDShadeConversion.h"
 #include "USDTypesConversion.h"
@@ -170,7 +171,7 @@ namespace UE::MeshTranslationImplInternal::Private
 		UMaterialInterface& Material,
 		const TMap<FString, int32>& MeshPrimvarToUVIndex,
 		UUsdAssetCache3* AssetCache,
-		FUsdInfoCache* InfoCache,
+		UUsdPrimLinkCache* PrimLinkCache,
 		const FString& MaterialHashPrefix,
 		bool bShareAssetsForIdenticalPrims
 	)
@@ -424,11 +425,11 @@ namespace UE::MeshTranslationImplInternal::Private
 #endif	  // WITH_EDITOR
 		}
 
-		if (CompatibleMaterial && CompatibleMaterial != &Material && InfoCache)
+		if (CompatibleMaterial && CompatibleMaterial != &Material && PrimLinkCache)
 		{
-			for (const UE::FSdfPath& Prim : InfoCache->GetPrimsForAsset(&Material))
+			for (const UE::FSdfPath& Prim : PrimLinkCache->GetPrimsForAsset(&Material))
 			{
-				InfoCache->LinkAssetToPrim(Prim, CompatibleMaterial);
+				PrimLinkCache->LinkAssetToPrim(Prim, CompatibleMaterial);
 			}
 		}
 
@@ -440,7 +441,7 @@ TMap<const UsdUtils::FUsdPrimMaterialSlot*, UMaterialInterface*> MeshTranslation
 	const pxr::UsdPrim& UsdPrim,
 	const TArray<UsdUtils::FUsdPrimMaterialAssignmentInfo>& AssignmentInfo,
 	UUsdAssetCache3& AssetCache,
-	FUsdInfoCache& InfoCache,
+	UUsdPrimLinkCache& PrimLinkCache,
 	EObjectFlags Flags,
 	bool bShareAssetsForIdenticalPrims
 )
@@ -566,10 +567,10 @@ TMap<const UsdUtils::FUsdPrimMaterialSlot*, UMaterialInterface*> MeshTranslation
 					bool bMaterialIsDirectReference = false;
 
 					// Here we have to pick the "best" material to use as reference, in case we need compatible/TwoSided versions.
-					// They are returned from InfoCache.GetAssetsForPrim in the most recent to least recent order, so
+					// They are returned from PrimLinkCache.GetAssetsForPrim in the most recent to least recent order, so
 					// in general we want to pick the first ones we find that match our criteria (as the older assets may be leftover from
 					// before we resynced something)
-					TArray<UMaterialInterface*> ExistingMaterials = InfoCache.GetAssetsForPrim<UMaterialInterface>(MaterialPrimPath);
+					TArray<UMaterialInterface*> ExistingMaterials = PrimLinkCache.GetAssetsForPrim<UMaterialInterface>(MaterialPrimPath);
 					for (UMaterialInterface* ExistingMaterial : ExistingMaterials)
 					{
 						const bool bExistingIsTwoSided = ExistingMaterial->IsTwoSided();
@@ -613,7 +614,7 @@ TMap<const UsdUtils::FUsdPrimMaterialSlot*, UMaterialInterface*> MeshTranslation
 								bMaterialIsDirectReference = bExistingIsDirectReference;
 							}
 						}
-						else // if (!Slot.bMeshIsDoubleSided && bExistingIsTwoSided)
+						else	// if (!Slot.bMeshIsDoubleSided && bExistingIsTwoSided)
 						{
 							// We can ignore this case: If we're searching for a one sided material and just ran into
 							// an existing two-sided one we should just keep iterating: If a two-sided material is within
@@ -690,7 +691,7 @@ TMap<const UsdUtils::FUsdPrimMaterialSlot*, UMaterialInterface*> MeshTranslation
 						// TouchAsset/ActiveAssets stuff)
 						AssetCache.TouchAssetPath(Material);
 
-						InfoCache.LinkAssetToPrim(UE::FSdfPath{*Slot.MaterialSource}, Material);
+						PrimLinkCache.LinkAssetToPrim(UE::FSdfPath{*Slot.MaterialSource}, Material);
 
 						// Finally, try to make our generated material primvar-compatible. We do this last because this will
 						// create another instance with the non-compatible material as reference material, which means we also
@@ -701,7 +702,7 @@ TMap<const UsdUtils::FUsdPrimMaterialSlot*, UMaterialInterface*> MeshTranslation
 							Material = AlreadyHandledMaterial;
 
 							AssetCache.TouchAssetPath(Material);
-							InfoCache.LinkAssetToPrim(UE::FSdfPath{*Slot.MaterialSource}, Material);
+							PrimLinkCache.LinkAssetToPrim(UE::FSdfPath{*Slot.MaterialSource}, Material);
 						}
 						else
 						{
@@ -710,7 +711,7 @@ TMap<const UsdUtils::FUsdPrimMaterialSlot*, UMaterialInterface*> MeshTranslation
 									*Material,
 									MeshPrimvarToUVIndex,
 									&AssetCache,
-									&InfoCache,
+									&PrimLinkCache,
 									HashPrefix,
 									bShareAssetsForIdenticalPrims
 								);
@@ -784,7 +785,8 @@ void MeshTranslationImpl::SetMaterialOverrides(
 	const TArray<UMaterialInterface*>& ExistingAssignments,
 	UMeshComponent& MeshComponent,
 	UUsdAssetCache3& AssetCache,
-	FUsdInfoCache& InfoCache,
+	UUsdInfoCache& UsdInfoCache,
+	UUsdPrimLinkCache& PrimLinkCache,
 	float Time,
 	EObjectFlags Flags,
 	bool bInterpretLODs,
@@ -860,7 +862,7 @@ void MeshTranslationImpl::SetMaterialOverrides(
 		// likely just some root Xform prim.
 		// Note: This only works because we'll rebuild the cache when our material purpose/render context changes,
 		// and because in USD relationships (and so material bindings) can't vary with time
-		TOptional<TArray<UsdUtils::FUsdPrimMaterialSlot>> SubtreeSlots = InfoCache.GetSubtreeMaterialSlots(UE::FSdfPath{PrimPath});
+		TOptional<TArray<UsdUtils::FUsdPrimMaterialSlot>> SubtreeSlots = UsdInfoCache.GetSubtreeMaterialSlots(UE::FSdfPath{PrimPath});
 		if (SubtreeSlots.IsSet())
 		{
 			UsdUtils::FUsdPrimMaterialAssignmentInfo& NewInfo = LODIndexToAssignments.Emplace_GetRef();
@@ -928,7 +930,7 @@ void MeshTranslationImpl::SetMaterialOverrides(
 			ValidPrim,
 			LODIndexToAssignments,
 			AssetCache,
-			InfoCache,
+			PrimLinkCache,
 			Flags,
 			bShareAssetsForIdenticalPrims
 		);
diff --git a/Engine/Plugins/Importers/USDImporter/Source/USDSchemas/Private/MeshTranslationImpl.h b/Engine/Plugins/Importers/USDImporter/Source/USDSchemas/Private/MeshTranslationImpl.h
--- a/Engine/Plugins/Importers/USDImporter/Source/USDSchemas/Private/MeshTranslationImpl.h
+++ b/Engine/Plugins/Importers/USDImporter/Source/USDSchemas/Private/MeshTranslationImpl.h
@@ -19,6 +19,8 @@ class FUsdInfoCache;
 class UMaterialInterface;
 class UMeshComponent;
 class UUsdAssetCache3;
+class UUsdPrimLinkCache;
+class UUsdInfoCache;
 class UUsdMeshAssetUserData;
 namespace UsdUtils
 {
@@ -34,7 +36,7 @@ namespace MeshTranslationImpl
 		const pxr::UsdPrim& UsdPrim,
 		const TArray<UsdUtils::FUsdPrimMaterialAssignmentInfo>& AssignmentInfo,
 		UUsdAssetCache3& AssetCache,
-		FUsdInfoCache& InfoCache,
+		UUsdPrimLinkCache& PrimLinkCache,
 		EObjectFlags Flags,
 		bool bShareAssetsForIdenticalPrims
 	);
@@ -48,7 +50,8 @@ namespace MeshTranslationImpl
 		const TArray<UMaterialInterface*>& ExistingAssignments,
 		UMeshComponent& MeshComponent,
 		UUsdAssetCache3& AssetCache,
-		FUsdInfoCache& InfoCache,
+		UUsdInfoCache& UsdInfoCache,
+		UUsdPrimLinkCache& PrimLinkCache,
 		float Time,
 		EObjectFlags Flags,
 		bool bInterpretLODs,
diff --git a/Engine/Plugins/Importers/USDImporter/Source/USDSchemas/Private/USDGeomMeshTranslator.cpp b/Engine/Plugins/Importers/USDImporter/Source/USDSchemas/Private/USDGeomMeshTranslator.cpp
--- a/Engine/Plugins/Importers/USDImporter/Source/USDSchemas/Private/USDGeomMeshTranslator.cpp
+++ b/Engine/Plugins/Importers/USDImporter/Source/USDSchemas/Private/USDGeomMeshTranslator.cpp
@@ -157,7 +157,7 @@ namespace UsdGeomMeshTranslatorImpl
 		}
 
 		// We want Nanite because the mesh is large enough for the threshold, which is set to something valid
-		if (!bHasNaniteOverrideEnabled && Context.InfoCache.IsValid())
+		if (!bHasNaniteOverrideEnabled)
 		{
 			const int32 NumTriangles = LODIndexToMeshDescription[0].Triangles().Num();
 			if (NumTriangles >= Context.NaniteTriangleThreshold)
@@ -202,9 +202,9 @@ namespace UsdGeomMeshTranslatorImpl
 			return false;
 		}
 
-		if (Context.InfoCache.IsValid())
+		if (Context.UsdInfoCache)
 		{
-			TOptional<uint64> SubtreeSectionCount = Context.InfoCache->GetSubtreeMaterialSlotCount(PrimPath);
+			TOptional<uint64> SubtreeSectionCount = Context.UsdInfoCache->GetSubtreeMaterialSlotCount(PrimPath);
 
 			if (!SubtreeSectionCount.IsSet() || SubtreeSectionCount.GetValue() > NANITE_MAX_CLUSTER_MATERIALS)
 			{
@@ -240,17 +240,12 @@ namespace UsdGeomMeshTranslatorImpl
 		const TArray<UsdUtils::FUsdPrimMaterialAssignmentInfo>& LODIndexToMaterialInfo,
 		UStaticMesh& StaticMesh,
 		UUsdAssetCache3& AssetCache,
-		FUsdInfoCache* InfoCache,
+		UUsdPrimLinkCache& PrimLinkCache,
 		float Time,
 		EObjectFlags Flags,
 		bool bShareAssetsForIdenticalPrims
 	)
 	{
-		if (!InfoCache)
-		{
-			return false;
-		}
-
 		bool bMaterialAssignementsHaveChanged = false;
 
 		TArray<UMaterialInterface*> ExistingAssignments;
@@ -263,7 +258,7 @@ namespace UsdGeomMeshTranslatorImpl
 			UsdPrim,
 			LODIndexToMaterialInfo,
 			AssetCache,
-			*InfoCache,
+			PrimLinkCache,
 			Flags,
 			bShareAssetsForIdenticalPrims
 		);
@@ -1626,10 +1621,10 @@ void FBuildStaticMeshTaskChain::SetupTasks()
 
 		   if (StaticMesh)
 		   {
-			   if (Context->InfoCache)
+			   if (Context->PrimLinkCache)
 			   {
 				   const UE::FSdfPath& TargetPath = AlternativePrimToLinkAssetsTo.IsSet() ? AlternativePrimToLinkAssetsTo.GetValue() : PrimPath;
-				   Context->InfoCache->LinkAssetToPrim(TargetPath, StaticMesh);
+				   Context->PrimLinkCache->LinkAssetToPrim(TargetPath, StaticMesh);
 			   }
 
 			   if (UUsdMeshAssetUserData* UserData = UsdUnreal::ObjectUtils::GetOrCreateAssetUserData<UUsdMeshAssetUserData>(StaticMesh))
@@ -1669,7 +1664,7 @@ void FBuildStaticMeshTaskChain::SetupTasks()
 			   // Only the original creator of the prim at creation time gets to set the material assignments
 			   // directly on the mesh, all others prims ensure their materials via material overrides on the
 			   // components
-			   if (bIsNew)
+			   if (bIsNew && Context->UsdAssetCache && Context->PrimLinkCache)
 			   {
 #if WITH_EDITOR
 				   StaticMesh->NaniteSettings.bEnabled = bShouldEnableNanite;
@@ -1681,8 +1676,8 @@ void FBuildStaticMeshTaskChain::SetupTasks()
 					   GetPrim(),
 					   LODIndexToMaterialInfo,
 					   *StaticMesh,
-					   *Context->UsdAssetCache.Get(),
-					   Context->InfoCache.Get(),
+					   *Context->UsdAssetCache,
+					   *Context->PrimLinkCache,
 					   Context->Time,
 					   Context->ObjectFlags,
 					   Context->bShareAssetsForIdenticalPrims
@@ -1784,7 +1779,7 @@ void FBuildStaticMeshTaskChain::SetupTasks()
 						 "were degenerate)"),
 					*PrimPath.GetString()
 				);
-				UsdUnreal::TranslatorUtils::AbandonFailedAsset(StaticMesh, Context->UsdAssetCache.Get(), Context->InfoCache.Get());
+				UsdUnreal::TranslatorUtils::AbandonFailedAsset(StaticMesh, Context->UsdAssetCache.Get(), Context->PrimLinkCache.Get());
 				return false;
 			}
 
@@ -2103,7 +2098,7 @@ TSet<UE::FSdfPath> FUsdGeomMeshTranslator::CollectAuxiliaryPrims() const
 
 	if (!Context->bIsBuildingInfoCache)
 	{
-		return Context->InfoCache->GetAuxiliaryPrims(PrimPath);
+		return Context->UsdInfoCache->GetAuxiliaryPrims(PrimPath);
 	}
 
 	TSet<UE::FSdfPath> Result;
diff --git a/Engine/Plugins/Importers/USDImporter/Source/USDSchemas/Private/USDGeomPointInstancerTranslator.cpp b/Engine/Plugins/Importers/USDImporter/Source/USDSchemas/Private/USDGeomPointInstancerTranslator.cpp
--- a/Engine/Plugins/Importers/USDImporter/Source/USDSchemas/Private/USDGeomPointInstancerTranslator.cpp
+++ b/Engine/Plugins/Importers/USDImporter/Source/USDSchemas/Private/USDGeomPointInstancerTranslator.cpp
@@ -301,11 +301,10 @@ USceneComponent* FUsdGeomPointInstancerTranslator::CreateComponents()
 			return MainSceneComponent;
 		}
 
-		if (!Context->InfoCache.IsValid())
+		if (!Context->UsdInfoCache)
 		{
 			return MainSceneComponent;
 		}
-		FUsdInfoCache& InfoCache = *Context->InfoCache.Get();
 
 		// Lets pretend ParentComponent is pointing to the parent USceneComponent while we create the child HISMs, so they get
 		// automatically attached to it as children
@@ -404,12 +403,10 @@ void FUsdGeomPointInstancerTranslator::UpdateComponents(USceneComponent* PointIn
 			return;
 		}
 
-		if (!Context->UsdAssetCache.IsValid() || !Context->InfoCache.IsValid())
+		if (!Context->UsdAssetCache.IsValid() || !Context->PrimLinkCache || !Context->UsdInfoCache)
 		{
 			return;
 		}
-		UUsdAssetCache3& AssetCache = *Context->UsdAssetCache.Get();
-		FUsdInfoCache& InfoCache = *Context->InfoCache.Get();
 
 		// Lets pretend ParentComponent is pointing to the parent USceneComponent while we create the child HISMs, so they get
 		// automatically attached to it as children
@@ -430,7 +427,7 @@ void FUsdGeomPointInstancerTranslator::UpdateComponents(USceneComponent* PointIn
 		// prototype, as we translate these with task pools and some of those prototypes may have generated nullptr.
 		// We always put the prototype path on the asset import data though, so here we use that to figure out where
 		// each mesh should go
-		TArray<UStaticMesh*> PrototypeMeshArr = InfoCache.GetAssetsForPrim<UStaticMesh>(PrimPath);
+		TArray<UStaticMesh*> PrototypeMeshArr = Context->PrimLinkCache->GetAssetsForPrim<UStaticMesh>(PrimPath);
 		std::unordered_map<pxr::SdfPath, UStaticMesh*, pxr::SdfPath::Hash> PrototypeMeshes;
 		PrototypeMeshes.reserve(PrototypeMeshArr.Num());
 		for (UStaticMesh* PrototypeMesh : PrototypeMeshArr)
@@ -539,8 +536,9 @@ void FUsdGeomPointInstancerTranslator::UpdateComponents(USceneComponent* PointIn
 					PrototypeUsdPrim.Get(),
 					ExistingAssignments,
 					*HISMComponent,
-					AssetCache,
-					InfoCache,
+					*Context->UsdAssetCache,
+					*Context->UsdInfoCache,
+					*Context->PrimLinkCache,
 					Context->Time,
 					Context->ObjectFlags,
 					Context->bAllowInterpretingLODs,
@@ -576,7 +574,7 @@ TSet<UE::FSdfPath> FUsdGeomPointInstancerTranslator::CollectAuxiliaryPrims() con
 {
 	if (!Context->bIsBuildingInfoCache)
 	{
-		return Context->InfoCache->GetAuxiliaryPrims(PrimPath);
+		return Context->UsdInfoCache->GetAuxiliaryPrims(PrimPath);
 	}
 
 	FScopedUsdAllocs UsdAllocs;
diff --git a/Engine/Plugins/Importers/USDImporter/Source/USDSchemas/Private/USDGeomPrimitiveTranslator.cpp b/Engine/Plugins/Importers/USDImporter/Source/USDSchemas/Private/USDGeomPrimitiveTranslator.cpp
--- a/Engine/Plugins/Importers/USDImporter/Source/USDSchemas/Private/USDGeomPrimitiveTranslator.cpp
+++ b/Engine/Plugins/Importers/USDImporter/Source/USDSchemas/Private/USDGeomPrimitiveTranslator.cpp
@@ -150,9 +150,9 @@ USceneComponent* FUsdGeomPrimitiveTranslator::CreateComponents()
 	// Handle material overrides
 	if (UStaticMeshComponent* StaticMeshComponent = Cast<UStaticMeshComponent>(SceneComponent))
 	{
-		if (Context->InfoCache)
+		if (Context->PrimLinkCache)
 		{
-			if (UStaticMesh* StaticMesh = Context->InfoCache->GetSingleAssetForPrim<UStaticMesh>(PrimPath))
+			if (UStaticMesh* StaticMesh = Context->PrimLinkCache->GetSingleAssetForPrim<UStaticMesh>(PrimPath))
 			{
 				TArray<UMaterialInterface*> ExistingAssignments;
 				for (FStaticMaterial& StaticMaterial : StaticMesh->GetStaticMaterials())
@@ -164,8 +164,9 @@ USceneComponent* FUsdGeomPrimitiveTranslator::CreateComponents()
 					GetPrim(),
 					ExistingAssignments,
 					*StaticMeshComponent,
-					*Context->UsdAssetCache.Get(),
-					*Context->InfoCache.Get(),
+					*Context->UsdAssetCache,
+					*Context->UsdInfoCache,
+					*Context->PrimLinkCache,
 					Context->Time,
 					Context->ObjectFlags,
 					Context->bAllowInterpretingLODs,
@@ -218,7 +219,7 @@ TSet<UE::FSdfPath> FUsdGeomPrimitiveTranslator::CollectAuxiliaryPrims() const
 {
 	if (!Context->bIsBuildingInfoCache)
 	{
-		return Context->InfoCache->GetAuxiliaryPrims(PrimPath);
+		return Context->UsdInfoCache->GetAuxiliaryPrims(PrimPath);
 	}
 
 	// Let's assume we can't specify UsdGeomSubsets on primitives for now
diff --git a/Engine/Plugins/Importers/USDImporter/Source/USDSchemas/Private/USDGeomXformableTranslator.cpp b/Engine/Plugins/Importers/USDImporter/Source/USDSchemas/Private/USDGeomXformableTranslator.cpp
--- a/Engine/Plugins/Importers/USDImporter/Source/USDSchemas/Private/USDGeomXformableTranslator.cpp
+++ b/Engine/Plugins/Importers/USDImporter/Source/USDSchemas/Private/USDGeomXformableTranslator.cpp
@@ -332,10 +332,10 @@ USceneComponent* FUsdGeomXformableTranslator::CreateComponents()
 	// Note how we don't have to handle geometry caches in here as they're handled by the geometry cache translator now
 	if (UStaticMeshComponent* StaticMeshComponent = Cast<UStaticMeshComponent>(SceneComponent))
 	{
-		if (Context->InfoCache)
+		if (Context->PrimLinkCache)
 		{
 			UE::FSdfPath PrototypePath = GetPrototypePrimPath();
-			if (UStaticMesh* StaticMesh = Context->InfoCache->GetSingleAssetForPrim<UStaticMesh>(PrototypePath))
+			if (UStaticMesh* StaticMesh = Context->PrimLinkCache->GetSingleAssetForPrim<UStaticMesh>(PrototypePath))
 			{
 				TArray<UMaterialInterface*> ExistingAssignments;
 				for (FStaticMaterial& StaticMaterial : StaticMesh->GetStaticMaterials())
@@ -347,8 +347,9 @@ USceneComponent* FUsdGeomXformableTranslator::CreateComponents()
 					GetPrim(),
 					ExistingAssignments,
 					*StaticMeshComponent,
-					*Context->UsdAssetCache.Get(),
-					*Context->InfoCache.Get(),
+					*Context->UsdAssetCache,
+					*Context->UsdInfoCache,
+					*Context->PrimLinkCache,
 					Context->Time,
 					Context->ObjectFlags,
 					Context->bAllowInterpretingLODs,
@@ -664,7 +665,7 @@ void FUsdGeomXformableTranslator::UpdateComponents(USceneComponent* SceneCompone
 {
 	TRACE_CPUPROFILER_EVENT_SCOPE(FUsdGeomXformableTranslator::UpdateComponents);
 
-	if (SceneComponent && Context->InfoCache)
+	if (SceneComponent && Context->PrimLinkCache)
 	{
 		SceneComponent->Modify();
 
@@ -683,7 +684,7 @@ void FUsdGeomXformableTranslator::UpdateComponents(USceneComponent* SceneCompone
 		if (UStaticMeshComponent* StaticMeshComponent = Cast<UStaticMeshComponent>(SceneComponent))
 		{
 			UE::FSdfPath PrototypePath = GetPrototypePrimPath();
-			UStaticMesh* PrimStaticMesh = Context->InfoCache->GetSingleAssetForPrim<UStaticMesh>(PrototypePath);
+			UStaticMesh* PrimStaticMesh = Context->PrimLinkCache->GetSingleAssetForPrim<UStaticMesh>(PrototypePath);
 
 			if (PrimStaticMesh)
 			{
@@ -794,7 +795,7 @@ namespace UE::UsdXformableTranslatorImpl::Private
 		UE::FUsdPrim Prim,
 		UUsdDrawModeComponent* DrawModeComponent,
 		UUsdAssetCache3& AssetCache,
-		FUsdInfoCache& InfoCache
+		UUsdPrimLinkCache& PrimLinkCache
 	)
 	{
 		if (!Prim)
@@ -817,7 +818,7 @@ namespace UE::UsdXformableTranslatorImpl::Private
 		// Collect textures from the info cache (they will all be linked to this Prim but will have within their AssetUserData the attribute
 		// that they originated from).
 		// The info cache may have old textures in case we're changing editing the stage, so we track multiple textures per attribute.
-		TArray<UTexture2D*> Textures = InfoCache.GetAssetsForPrim<UTexture2D>(Prim.GetPrimPath());
+		TArray<UTexture2D*> Textures = PrimLinkCache.GetAssetsForPrim<UTexture2D>(Prim.GetPrimPath());
 		std::unordered_map<pxr::SdfPath, TArray<UTexture2D*>, pxr::SdfPath::Hash> AttrPathToTextures;
 		for (UTexture2D* Texture : Textures)
 		{
@@ -939,7 +940,7 @@ USceneComponent* FUsdGeomXformableTranslator::CreateAlternativeDrawModeComponent
 		case EUsdDrawMode::Cards:
 		{
 			UUsdDrawModeComponent* Component = Cast<UUsdDrawModeComponent>(CreateComponentsEx({UUsdDrawModeComponent::StaticClass()}, bNeedsActor));
-			if (ensure(Component) && Context->UsdAssetCache && Context->InfoCache)
+			if (ensure(Component) && Context->UsdAssetCache && Context->PrimLinkCache)
 			{
 				// For now we only assign textures when creating components, not when updating. Maybe in the future we can
 				// add support for "texture animations"
@@ -947,7 +948,7 @@ USceneComponent* FUsdGeomXformableTranslator::CreateAlternativeDrawModeComponent
 					GetPrim(),
 					Component,
 					*Context->UsdAssetCache,
-					*Context->InfoCache
+					*Context->PrimLinkCache
 				);
 			}
 			return Component;
@@ -968,7 +969,7 @@ void FUsdGeomXformableTranslator::CreateAlternativeDrawModeAssets(EUsdDrawMode D
 {
 	// Currently we just use this function to create the textures that we're going to use on the bounds components,
 	// if applicable
-	if (DrawMode != EUsdDrawMode::Cards || !Context->UsdAssetCache || !Context->InfoCache)
+	if (DrawMode != EUsdDrawMode::Cards || !Context->UsdAssetCache || !Context->PrimLinkCache)
 	{
 		return;
 	}
@@ -1027,7 +1028,7 @@ void FUsdGeomXformableTranslator::CreateAlternativeDrawModeAssets(EUsdDrawMode D
 				// We link the textures to the prim, so that if the prim is reloaded the AUsdStageActor knows to potentially
 				// drop the textures. However we put the full attribute path on AssetUserData, so that when we're filling in
 				// our UUsdDrawModeComponent later, we know which texture came from which attribute
-				Context->InfoCache->LinkAssetToPrim(PrimPath, Texture);
+				Context->PrimLinkCache->LinkAssetToPrim(PrimPath, Texture);
 				if (UUsdAssetUserData* TextureUserData = UsdUnreal::ObjectUtils::GetOrCreateAssetUserData(Texture))
 				{
 					TextureUserData->PrimPaths.AddUnique(UsdToUnreal::ConvertPath(Attr.GetPath()));
@@ -1088,7 +1089,7 @@ bool FUsdGeomXformableTranslator::CollapsesChildren(ECollapsingType CollapsingTy
 
 	if (!Context->bIsBuildingInfoCache)
 	{
-		return Context->InfoCache->DoesPathCollapseChildren(PrimPath, CollapsingType);
+		return Context->UsdInfoCache->DoesPathCollapseChildren(PrimPath, CollapsingType);
 	}
 
 	// If we have a custom draw mode, it means we should draw bounds/cards/etc. instead
@@ -1120,7 +1121,7 @@ bool FUsdGeomXformableTranslator::CollapsesChildren(ECollapsingType CollapsingTy
 		if (bCollapsesChildren)
 		{
 			// This indicates whether the subtree *can* be collapsed
-			TOptional<bool> bSubtreeCanBeCollapsed = Context->InfoCache->CanXformableSubtreeBeCollapsed(PrimPath, *Context);
+			TOptional<bool> bSubtreeCanBeCollapsed = Context->UsdInfoCache->CanXformableSubtreeBeCollapsed(PrimPath, *Context);
 			if (bSubtreeCanBeCollapsed.IsSet())
 			{
 				return bSubtreeCanBeCollapsed.GetValue();
@@ -1154,10 +1155,10 @@ TSet<UE::FSdfPath> FUsdGeomXformableTranslator::CollectAuxiliaryPrims() const
 {
 	if (!Context->bIsBuildingInfoCache)
 	{
-		return Context->InfoCache->GetAuxiliaryPrims(PrimPath);
+		return Context->UsdInfoCache->GetAuxiliaryPrims(PrimPath);
 	}
 
-	if (!Context->InfoCache->DoesPathCollapseChildren(PrimPath, ECollapsingType::Assets))
+	if (!Context->UsdInfoCache->DoesPathCollapseChildren(PrimPath, ECollapsingType::Assets))
 	{
 		return {};
 	}
diff --git a/Engine/Plugins/Importers/USDImporter/Source/USDSchemas/Private/USDGeometryCacheTranslator.cpp b/Engine/Plugins/Importers/USDImporter/Source/USDSchemas/Private/USDGeometryCacheTranslator.cpp
--- a/Engine/Plugins/Importers/USDImporter/Source/USDSchemas/Private/USDGeometryCacheTranslator.cpp
+++ b/Engine/Plugins/Importers/USDImporter/Source/USDSchemas/Private/USDGeometryCacheTranslator.cpp
@@ -83,7 +83,7 @@ namespace UsdGeometryCacheTranslatorImpl
 		const TArray<UsdUtils::FUsdPrimMaterialAssignmentInfo>& LODIndexToMaterialInfo,
 		UGeometryCache& GeometryCache,
 		UUsdAssetCache3& AssetCache,
-		FUsdInfoCache* InfoCache,
+		UUsdPrimLinkCache& PrimLinkCache,
 		float Time,
 		EObjectFlags Flags,
 		bool bShareAssetsForIdenticalPrims
@@ -93,7 +93,7 @@ namespace UsdGeometryCacheTranslatorImpl
 			UsdPrim,
 			LODIndexToMaterialInfo,
 			AssetCache,
-			*InfoCache,
+			PrimLinkCache,
 			Flags,
 			bShareAssetsForIdenticalPrims
 		);
@@ -913,7 +913,7 @@ void FGeometryCacheCreateAssetsTaskChain::SetupTasks()
 					MeshTranslationImpl::RecordSourcePrimsForMaterialSlots(LODIndexToMaterialInfo, UserData);
 				}
 
-				if (bIsNew)
+				if (bIsNew && Context->UsdAssetCache && Context->PrimLinkCache)
 				{
 					// Only the original creator of the prim at creation time gets to set the material assignments
 					// directly on the geometry cache, all others prims ensure their materials via material overrides on the
@@ -922,18 +922,18 @@ void FGeometryCacheCreateAssetsTaskChain::SetupTasks()
 						GetPrim(),
 						LODIndexToMaterialInfo,
 						*GeometryCache,
-						*Context->UsdAssetCache.Get(),
-						Context->InfoCache.Get(),
+						*Context->UsdAssetCache,
+						*Context->PrimLinkCache,
 						Context->Time,
 						Context->ObjectFlags,
 						Context->bShareAssetsForIdenticalPrims
 					);
 				}
 
-				if (Context->InfoCache)
+				if (Context->PrimLinkCache)
 				{
 					const UE::FSdfPath& TargetPath = AlternativePrimToLinkAssetsTo.IsSet() ? AlternativePrimToLinkAssetsTo.GetValue() : PrimPath;
-					Context->InfoCache->LinkAssetToPrim(TargetPath, GeometryCache.Get());
+					Context->PrimLinkCache->LinkAssetToPrim(TargetPath, GeometryCache.Get());
 				}
 			}
 
@@ -1032,9 +1032,9 @@ USceneComponent* FUsdGeometryCacheTranslator::CreateComponents()
 
 	if (UGeometryCacheComponent* Component = Cast<UGeometryCacheComponent>(SceneComponent))
 	{
-		if (Context->InfoCache && Context->UsdAssetCache)
+		if (Context->PrimLinkCache && Context->UsdAssetCache)
 		{
-			if (UGeometryCache* GeometryCache = Context->InfoCache->GetSingleAssetForPrim<UGeometryCache>(PrimPath))
+			if (UGeometryCache* GeometryCache = Context->PrimLinkCache->GetSingleAssetForPrim<UGeometryCache>(PrimPath))
 			{
 				// Geometry caches don't support LODs
 				const bool bAllowInterpretingLODs = false;
@@ -1044,7 +1044,8 @@ USceneComponent* FUsdGeometryCacheTranslator::CreateComponents()
 					GeometryCache->Materials,
 					*Component,
 					*Context->UsdAssetCache,
-					*Context->InfoCache,
+					*Context->UsdInfoCache,
+					*Context->PrimLinkCache,
 					Context->Time,
 					Context->ObjectFlags,
 					bAllowInterpretingLODs,
@@ -1059,7 +1060,7 @@ USceneComponent* FUsdGeometryCacheTranslator::CreateComponents()
 					UsdGroomTranslatorUtils::CreateGroomBindingAsset(
 						GetPrim(),
 						*Context->UsdAssetCache,
-						*Context->InfoCache,
+						*Context->PrimLinkCache,
 						Context->ObjectFlags,
 						Context->bShareAssetsForIdenticalPrims
 					);
@@ -1108,9 +1109,9 @@ void FUsdGeometryCacheTranslator::UpdateComponents(USceneComponent* SceneCompone
 	if (GeometryCacheComponent)
 	{
 		UGeometryCache* GeometryCache = nullptr;
-		if (Context->InfoCache)
+		if (Context->PrimLinkCache)
 		{
-			GeometryCache = Context->InfoCache->GetSingleAssetForPrim<UGeometryCache>(PrimPath);
+			GeometryCache = Context->PrimLinkCache->GetSingleAssetForPrim<UGeometryCache>(PrimPath);
 		}
 
 		bool bShouldRegister = false;
@@ -1180,7 +1181,7 @@ void FUsdGeometryCacheTranslator::UpdateComponents(USceneComponent* SceneCompone
 		// If the prim has a GroomBinding schema, apply the target groom to its associated GroomComponent
 		if (UsdUtils::PrimHasSchema(GetPrim(), UnrealIdentifiers::GroomBindingAPI))
 		{
-			UsdGroomTranslatorUtils::SetGroomFromPrim(GetPrim(), *Context->InfoCache, SceneComponent);
+			UsdGroomTranslatorUtils::SetGroomFromPrim(GetPrim(), *Context->PrimLinkCache, SceneComponent);
 		}
 
 		const bool bIsImporting = GForceImport || Context->bIsImporting;
@@ -1239,7 +1240,7 @@ TSet<UE::FSdfPath> FUsdGeometryCacheTranslator::CollectAuxiliaryPrims() const
 
 	if (!Context->bIsBuildingInfoCache)
 	{
-		return Context->InfoCache->GetAuxiliaryPrims(PrimPath);
+		return Context->UsdInfoCache->GetAuxiliaryPrims(PrimPath);
 	}
 
 	if (ShouldSkipSkinnablePrim())
@@ -1264,8 +1265,8 @@ TSet<UE::FSdfPath> FUsdGeometryCacheTranslator::CollectAuxiliaryPrims() const
 
 bool FUsdGeometryCacheTranslator::IsPotentialGeometryCacheRoot() const
 {
-	// The logic to check for GeometryCache is completely in the InfoCache
-	return Context->InfoCache->IsPotentialGeometryCacheRoot(PrimPath);
+	// The logic to check for GeometryCache is completely in the UsdInfoCache
+	return Context->UsdInfoCache->IsPotentialGeometryCacheRoot(PrimPath);
 }
 
 #endif	  // #if USE_USD_SDK
diff --git a/Engine/Plugins/Importers/USDImporter/Source/USDSchemas/Private/USDGroomTranslator.cpp b/Engine/Plugins/Importers/USDImporter/Source/USDSchemas/Private/USDGroomTranslator.cpp
--- a/Engine/Plugins/Importers/USDImporter/Source/USDSchemas/Private/USDGroomTranslator.cpp
+++ b/Engine/Plugins/Importers/USDImporter/Source/USDSchemas/Private/USDGroomTranslator.cpp
@@ -233,9 +233,9 @@ protected:
 					}
 				}
 
-				if (GroomAsset && Context->InfoCache)
+				if (GroomAsset && Context->PrimLinkCache)
 				{
-					Context->InfoCache->LinkAssetToPrim(PrimPath, GroomAsset);
+					Context->PrimLinkCache->LinkAssetToPrim(PrimPath, GroomAsset);
 				}
 
 				// Next step is to parse the GroomCache data if it was determined that the groom has animated attributes
@@ -248,12 +248,12 @@ protected:
 			ESchemaTranslationLaunchPolicy::Async,
 			[this]() -> bool
 			{
-				if (!Context->InfoCache)
+				if (!Context->PrimLinkCache)
 				{
 					return false;
 				}
 
-				GroomAsset = Context->InfoCache->GetSingleAssetForPrim<UGroomAsset>(PrimPath);
+				GroomAsset = Context->PrimLinkCache->GetSingleAssetForPrim<UGroomAsset>(PrimPath);
 				if (!GroomAsset)
 				{
 					return false;
@@ -303,9 +303,9 @@ protected:
 					&bGroomCacheIsNew
 				);
 
-				if (GroomCache && Context->InfoCache)
+				if (GroomCache && Context->PrimLinkCache)
 				{
-					Context->InfoCache->LinkAssetToPrim(PrimPath, GroomCache);
+					Context->PrimLinkCache->LinkAssetToPrim(PrimPath, GroomCache);
 
 					if (UUsdAssetUserData* UserData = UsdUnreal::ObjectUtils::GetOrCreateAssetUserData(GroomCache))
 					{
@@ -450,7 +450,7 @@ protected:
 					const FString StrandsGroomCachePrimPath = UsdGroomTranslatorUtils::GetStrandsGroomCachePrimPath(PrimPath);
 					UE_LOG(LogUsd, Warning, TEXT("Failed to create GroomCache for prim '%s'"), *StrandsGroomCachePrimPath);
 
-					UsdUnreal::TranslatorUtils::AbandonFailedAsset(GroomCache, Context->UsdAssetCache.Get(), Context->InfoCache.Get());
+					UsdUnreal::TranslatorUtils::AbandonFailedAsset(GroomCache, Context->UsdAssetCache.Get(), Context->PrimLinkCache.Get());
 				}
 
 				return bSuccess;
@@ -539,9 +539,9 @@ void FUsdGroomTranslator::UpdateComponents(USceneComponent* SceneComponent)
 			GroomComponent->Modify();
 
 			UGroomAsset* Groom = nullptr;
-			if (Context->InfoCache)
+			if (Context->PrimLinkCache)
 			{
-				Groom = Context->InfoCache->GetSingleAssetForPrim<UGroomAsset>(PrimPath);
+				Groom = Context->PrimLinkCache->GetSingleAssetForPrim<UGroomAsset>(PrimPath);
 			}
 
 			bool bShouldRegister = false;
@@ -558,7 +558,7 @@ void FUsdGroomTranslator::UpdateComponents(USceneComponent* SceneComponent)
 
 				if (Groom)
 				{
-					UGroomCache* GroomCache = Context->InfoCache->GetSingleAssetForPrim<UGroomCache>(PrimPath);
+					UGroomCache* GroomCache = Context->PrimLinkCache->GetSingleAssetForPrim<UGroomCache>(PrimPath);
 					if (GroomCache != GroomComponent->GroomCache.Get())
 					{
 						GroomComponent->SetGroomCache(GroomCache);
@@ -612,10 +612,10 @@ TSet<UE::FSdfPath> FUsdGroomTranslator::CollectAuxiliaryPrims() const
 
 	if (!Context->bIsBuildingInfoCache)
 	{
-		return Context->InfoCache->GetAuxiliaryPrims(PrimPath);
+		return Context->UsdInfoCache->GetAuxiliaryPrims(PrimPath);
 	}
 
-	if (!Context->InfoCache->DoesPathCollapseChildren(PrimPath, ECollapsingType::Assets))
+	if (!Context->UsdInfoCache->DoesPathCollapseChildren(PrimPath, ECollapsingType::Assets))
 	{
 		return {};
 	}
diff --git a/Engine/Plugins/Importers/USDImporter/Source/USDSchemas/Private/USDGroomTranslatorUtils.cpp b/Engine/Plugins/Importers/USDImporter/Source/USDSchemas/Private/USDGroomTranslatorUtils.cpp
--- a/Engine/Plugins/Importers/USDImporter/Source/USDSchemas/Private/USDGroomTranslatorUtils.cpp
+++ b/Engine/Plugins/Importers/USDImporter/Source/USDSchemas/Private/USDGroomTranslatorUtils.cpp
@@ -9,6 +9,7 @@
 #include "USDClassesModule.h"
 #include "USDInfoCache.h"
 #include "USDIntegrationUtils.h"
+#include "USDPrimLinkCache.h"
 #include "USDTypesConversion.h"
 #include "UsdWrappers/SdfPath.h"
 
@@ -115,7 +116,7 @@ namespace UE::UsdGroomTranslatorUtils::Private
 		return {};
 	}
 
-	UObject* GetGroomBindingSourceMesh(const pxr::UsdPrim& Prim, const FUsdInfoCache& InfoCache, EGroomBindingMeshType BindingType)
+	UObject* GetGroomBindingSourceMesh(const pxr::UsdPrim& Prim, const UUsdPrimLinkCache& PrimLinkCache, EGroomBindingMeshType BindingType)
 	{
 		FScopedUsdAllocs Allocs;
 
@@ -133,11 +134,11 @@ namespace UE::UsdGroomTranslatorUtils::Private
 				pxr::UsdPrim TargetPrim = Prim.GetPrimAtPath(TargetPrimPath);
 				if (BindingType == EGroomBindingMeshType::SkeletalMesh && (pxr::UsdSkelRoot(TargetPrim) || pxr::UsdSkelSkeleton(TargetPrim)))
 				{
-					return InfoCache.GetSingleAssetForPrim<USkeletalMesh>(UE::FSdfPath{TargetPrimPath});
+					return PrimLinkCache.GetSingleAssetForPrim<USkeletalMesh>(UE::FSdfPath{TargetPrimPath});
 				}
 				else if (BindingType == EGroomBindingMeshType::GeometryCache && pxr::UsdGeomMesh(TargetPrim))
 				{
-					return InfoCache.GetSingleAssetForPrim<UGeometryCache>(UE::FSdfPath{TargetPrimPath});
+					return PrimLinkCache.GetSingleAssetForPrim<UGeometryCache>(UE::FSdfPath{TargetPrimPath});
 				}
 			}
 		}
@@ -152,7 +153,7 @@ namespace UsdGroomTranslatorUtils
 	void CreateGroomBindingAsset(
 		const pxr::UsdPrim& Prim,
 		UUsdAssetCache3& AssetCache,
-		FUsdInfoCache& InfoCache,
+		UUsdPrimLinkCache& PrimLinkCache,
 		EObjectFlags ObjectFlags,
 		bool bShareAssetsForIdenticalPrims
 	)
@@ -168,7 +169,7 @@ namespace UsdGroomTranslatorUtils
 		}
 
 		// The GroomAsset should already be processed and cached by the USDGroomTranslator
-		UGroomAsset* GroomAsset = InfoCache.GetSingleAssetForPrim<UGroomAsset>(UE::FSdfPath{*GroomPrimPath});
+		UGroomAsset* GroomAsset = PrimLinkCache.GetSingleAssetForPrim<UGroomAsset>(UE::FSdfPath{*GroomPrimPath});
 		if (!GroomAsset)
 		{
 			return;
@@ -177,10 +178,10 @@ namespace UsdGroomTranslatorUtils
 		// Determine the type of binding needed based on the prim mesh type
 		const FString PrimPath(UsdToUnreal::ConvertPath(Prim.GetPath()));
 		EGroomBindingMeshType GroomBindingType = EGroomBindingMeshType::SkeletalMesh;
-		UObject* TargetMesh = InfoCache.GetSingleAssetForPrim<USkeletalMesh>(UE::FSdfPath{*PrimPath});
+		UObject* TargetMesh = PrimLinkCache.GetSingleAssetForPrim<USkeletalMesh>(UE::FSdfPath{*PrimPath});
 		if (!TargetMesh)
 		{
-			TargetMesh = InfoCache.GetSingleAssetForPrim<UGeometryCache>(UE::FSdfPath{*PrimPath});
+			TargetMesh = PrimLinkCache.GetSingleAssetForPrim<UGeometryCache>(UE::FSdfPath{*PrimPath});
 			if (!TargetMesh)
 			{
 				return;
@@ -192,9 +193,10 @@ namespace UsdGroomTranslatorUtils
 		Settings.GroomBindingType = GroomBindingType;
 		Settings.Groom = GroomAsset;
 		Settings.TargetMesh = TargetMesh;
-		Settings.SourceMesh = GetGroomBindingSourceMesh(Prim, InfoCache, GroomBindingType);
+		Settings.SourceMesh = GetGroomBindingSourceMesh(Prim, PrimLinkCache, GroomBindingType);
 
-		const FString PrefixedBindingHash = UsdUtils::GetAssetHashPrefix(Prim, bShareAssetsForIdenticalPrims) + ComputeGroomBindingHash(Settings).ToString();
+		const FString PrefixedBindingHash = UsdUtils::GetAssetHashPrefix(Prim, bShareAssetsForIdenticalPrims)
+											+ ComputeGroomBindingHash(Settings).ToString();
 
 		const FString GroomBindingPath = FString::Printf(TEXT("%s_groombinding"), *PrimPath);
 
@@ -238,11 +240,11 @@ namespace UsdGroomTranslatorUtils
 				GroomBinding->Build();
 			}
 
-			InfoCache.LinkAssetToPrim(UE::FSdfPath{*GroomBindingPath}, GroomBinding);
+			PrimLinkCache.LinkAssetToPrim(UE::FSdfPath{*GroomBindingPath}, GroomBinding);
 		}
 	}
 
-	void SetGroomFromPrim(const pxr::UsdPrim& Prim, const FUsdInfoCache& InfoCache, USceneComponent* SceneComponent)
+	void SetGroomFromPrim(const pxr::UsdPrim& Prim, const UUsdPrimLinkCache& PrimLinkCache, USceneComponent* SceneComponent)
 	{
 		if (!SceneComponent)
 		{
@@ -257,7 +259,7 @@ namespace UsdGroomTranslatorUtils
 			return;
 		}
 
-		UGroomAsset* GroomAsset = InfoCache.GetSingleAssetForPrim<UGroomAsset>(UE::FSdfPath{*GroomPrimPath});
+		UGroomAsset* GroomAsset = PrimLinkCache.GetSingleAssetForPrim<UGroomAsset>(UE::FSdfPath{*GroomPrimPath});
 		if (!GroomAsset)
 		{
 			return;
@@ -265,7 +267,7 @@ namespace UsdGroomTranslatorUtils
 
 		const FString PrimPath(UsdToUnreal::ConvertPath(Prim.GetPath()));
 		const FString GroomBindingPath = FString::Printf(TEXT("%s_groombinding"), *PrimPath);
-		UGroomBindingAsset* GroomBinding = InfoCache.GetSingleAssetForPrim<UGroomBindingAsset>(UE::FSdfPath{*GroomBindingPath});
+		UGroomBindingAsset* GroomBinding = PrimLinkCache.GetSingleAssetForPrim<UGroomBindingAsset>(UE::FSdfPath{*GroomBindingPath});
 
 		// Set the GroomAsset and GroomBindingAsset on the child GroomComponent of SceneComponent that was set up in the translator
 		TArray<USceneComponent*> Children;
diff --git a/Engine/Plugins/Importers/USDImporter/Source/USDSchemas/Private/USDGroomTranslatorUtils.h b/Engine/Plugins/Importers/USDImporter/Source/USDSchemas/Private/USDGroomTranslatorUtils.h
--- a/Engine/Plugins/Importers/USDImporter/Source/USDSchemas/Private/USDGroomTranslatorUtils.h
+++ b/Engine/Plugins/Importers/USDImporter/Source/USDSchemas/Private/USDGroomTranslatorUtils.h
@@ -15,7 +15,7 @@ PXR_NAMESPACE_OPEN_SCOPE
 	class UsdPrim;
 PXR_NAMESPACE_CLOSE_SCOPE
 
-class FUsdInfoCache;
+class UUsdPrimLinkCache;
 class USceneComponent;
 class UUsdAssetCache3;
 
@@ -30,13 +30,13 @@ namespace UsdGroomTranslatorUtils
 	void CreateGroomBindingAsset(
 		const pxr::UsdPrim& Prim,
 		UUsdAssetCache3& AssetCache,
-		FUsdInfoCache& InfoCache,
+		UUsdPrimLinkCache& PrimLinkCache,
 		EObjectFlags ObjectFlags,
 		bool bShareAssetsForIdenticalPrims
 	);
 
 	// Set the groom asset targeted by the given prim with GroomBindingAPI on the SceneComponent (with a GroomComponent as a direct child)
-	void SetGroomFromPrim(const pxr::UsdPrim& Prim, const FUsdInfoCache& InfoCache, USceneComponent* SceneComponent);
+	void SetGroomFromPrim(const pxr::UsdPrim& Prim, const UUsdPrimLinkCache& PrimLinkCache, USceneComponent* SceneComponent);
 
 	// Get the (pseudo) prim path for the strands GroomCache associated with a groom prim
 	FString GetStrandsGroomCachePrimPath(const UE::FSdfPath& PrimPath);
diff --git a/Engine/Plugins/Importers/USDImporter/Source/USDSchemas/Private/USDInfoCache.cpp b/Engine/Plugins/Importers/USDImporter/Source/USDSchemas/Private/USDInfoCache.cpp
--- a/Engine/Plugins/Importers/USDImporter/Source/USDSchemas/Private/USDInfoCache.cpp
+++ b/Engine/Plugins/Importers/USDImporter/Source/USDSchemas/Private/USDInfoCache.cpp
@@ -66,6 +66,8 @@ static FAutoConsoleVariableRef CVarNumPerPrimLocks(
 	)
 );
 
+PRAGMA_DISABLE_DEPRECATION_WARNINGS
+
 namespace UE::UsdInfoCache::Private
 {
 	// Flags to hint at the state of a prim for the purpose of geometry cache
@@ -139,7 +141,7 @@ FArchive& operator<<(FArchive& Ar, UE::UsdInfoCache::Private::FUsdPrimInfo& Info
 	return Ar;
 }
 
-struct FUsdInfoCache::FUsdInfoCacheImpl
+struct FUsdInfoCacheImpl
 {
 	FUsdInfoCacheImpl()
 		: AllowedExtensionsForGeometryCacheSource(UnrealUSDWrapper::GetNativeFileFormats())
@@ -153,14 +155,30 @@ struct FUsdInfoCache::FUsdInfoCacheImpl
 	FUsdInfoCacheImpl(const FUsdInfoCacheImpl& Other)
 		: FUsdInfoCacheImpl()
 	{
-		// Probably not necessary
 		FReadScopeLock ScopedInfoMapLock(Other.InfoMapLock);
 		FReadScopeLock ScopedPrimPathToAssetsLock(Other.PrimPathToAssetsLock);
 
+		FWriteScopeLock ThisScopedInfoMapLock(InfoMapLock);
+		FWriteScopeLock ThisScopedPrimPathToAssetsLock(PrimPathToAssetsLock);
+
 		InfoMap = Other.InfoMap;
+
 		PrimPathToAssets = Other.PrimPathToAssets;
 		AssetToPrimPaths = Other.AssetToPrimPaths;
+
+		AllowedExtensionsForGeometryCacheSource = Other.AllowedExtensionsForGeometryCacheSource;
+	}
+
+	FUsdInfoCacheImpl& operator=(const FUsdInfoCacheImpl& Other)
+	{
+		FReadScopeLock OtherScopedInfoMapLock(Other.InfoMapLock);
+
+		FWriteScopeLock ThisScopedInfoMapLock(InfoMapLock);
+
+		InfoMap = Other.InfoMap;
 		AllowedExtensionsForGeometryCacheSource = Other.AllowedExtensionsForGeometryCacheSource;
+		TempStage = Other.TempStage;
+		return *this;
 	}
 
 	~FUsdInfoCacheImpl()
@@ -173,13 +191,14 @@ struct FUsdInfoCache::FUsdInfoCacheImpl
 	mutable FRWLock InfoMapLock;
 
 	// Information we may have about a subset of prims
+	// Only used by the deprecated FUsdInfoCache
 	TMap<UE::FSdfPath, TArray<TWeakObjectPtr<UObject>>> PrimPathToAssets;
 	TMap<TWeakObjectPtr<UObject>, TArray<UE::FSdfPath>> AssetToPrimPaths;
 	mutable FRWLock PrimPathToAssetsLock;
 
 	// Temporarily used during the info cache build, as we need to do another pass on point instancers afterwards
-	TArray<FString> TempPointInstancerPaths;
-	mutable FRWLock TempPointInstancerPathsLock;
+	TArray<FString> PointInstancerPaths;
+	mutable FRWLock PointInstancerPathsLock;
 
 	// This is used to keep track of which prototypes are already being translated within this "translation session",
 	// so that the schema translators can early out if they're trying to translate multiple instances of the same
@@ -229,12 +248,12 @@ public:
 
 FUsdInfoCache::FUsdInfoCache()
 {
-	Impl = MakeUnique<FUsdInfoCache::FUsdInfoCacheImpl>();
+	Impl = MakeUnique<FUsdInfoCacheImpl>();
 }
 
 FUsdInfoCache::FUsdInfoCache(const FUsdInfoCache& Other)
 {
-	Impl = MakeUnique<FUsdInfoCache::FUsdInfoCacheImpl>(*Other.Impl);
+	Impl = MakeUnique<FUsdInfoCacheImpl>(*Other.Impl);
 }
 
 FUsdInfoCache::~FUsdInfoCache()
@@ -338,7 +357,7 @@ namespace UE::USDInfoCache::Private
 	bool RecursiveQueryCanBeCollapsed(
 		const pxr::UsdPrim& UsdPrim,
 		FUsdSchemaTranslationContext& Context,
-		FUsdInfoCache::FUsdInfoCacheImpl& Impl,
+		FUsdInfoCacheImpl& Impl,
 		FUsdSchemaTranslatorRegistry& Registry
 	)
 	{
@@ -602,7 +621,7 @@ namespace UE::USDInfoCache::Private
 	void GetPrimVertexCountAndSlots(
 		const pxr::UsdPrim& UsdPrim,
 		const FUsdSchemaTranslationContext& Context,
-		const FUsdInfoCache::FUsdInfoCacheImpl& Impl,
+		const FUsdInfoCacheImpl& Impl,
 		uint64& OutVertexCount,
 		TArray<UsdUtils::FUsdPrimMaterialSlot>& OutMaterialSlots
 	)
@@ -699,7 +718,7 @@ namespace UE::USDInfoCache::Private
 		}
 	}
 
-	void RepopulateInfoMap(const pxr::UsdPrim& UsdPrim, FUsdInfoCache::FUsdInfoCacheImpl& Impl)
+	void RepopulateInfoMap(const pxr::UsdPrim& UsdPrim, FUsdInfoCacheImpl& Impl)
 	{
 		using namespace UE::UsdInfoCache::Private;
 
@@ -720,7 +739,7 @@ namespace UE::USDInfoCache::Private
 		const pxr::UsdPrim& UsdPrim,
 		FUsdSchemaTranslationContext& Context,
 		const pxr::TfToken& MaterialPurposeToken,
-		FUsdInfoCache::FUsdInfoCacheImpl& Impl,
+		FUsdInfoCacheImpl& Impl,
 		FUsdSchemaTranslatorRegistry& Registry,
 		uint64& OutSubtreeVertexCount,
 		TArray<UsdUtils::FUsdPrimMaterialSlot>& OutSubtreeSlots,
@@ -984,8 +1003,8 @@ namespace UE::USDInfoCache::Private
 				// could technically be anywhere, so store them here for a later pass
 				if (bIsPointInstancer)
 				{
-					FWriteScopeLock PointInstancerLock(Impl.TempPointInstancerPathsLock);
-					Impl.TempPointInstancerPaths.Emplace(UE::FSdfPath{UsdPrimPath}.GetString());
+					FWriteScopeLock PointInstancerLock(Impl.PointInstancerPathsLock);
+					Impl.PointInstancerPaths.Emplace(UE::FSdfPath{UsdPrimPath}.GetString());
 				}
 				// While we will compute the totals for any and all children normally, don't just append the regular
 				// traversal vertex count to the point instancer prim itself just yet, as that doesn't really represent
@@ -1010,7 +1029,7 @@ namespace UE::USDInfoCache::Private
 	 * of prim subtrees all over to build the final counts of point instancers that use them as prototypes, and
 	 * then update their parents.
 	 */
-	void UpdateInfoForPointInstancers(const FUsdSchemaTranslationContext& Context, FUsdInfoCache::FUsdInfoCacheImpl& Impl)
+	void UpdateInfoForPointInstancers(const FUsdSchemaTranslationContext& Context, FUsdInfoCacheImpl& Impl)
 	{
 		TRACE_CPUPROFILER_EVENT_SCOPE(UpdateInfoForPointInstancers);
 
@@ -1061,12 +1080,12 @@ namespace UE::USDInfoCache::Private
 			return LHS < RHS;
 		};
 		{
-			FWriteScopeLock PointInstancerLock{Impl.TempPointInstancerPathsLock};
-			Impl.TempPointInstancerPaths.Sort(SortFunction);
+			FWriteScopeLock PointInstancerLock{Impl.PointInstancerPathsLock};
+			Impl.PointInstancerPaths.Sort(SortFunction);
 		}
 
-		FReadScopeLock PointInstancerLock{Impl.TempPointInstancerPathsLock};
-		for (const FString& PointInstancerPath : Impl.TempPointInstancerPaths)
+		FReadScopeLock PointInstancerLock{Impl.PointInstancerPathsLock};
+		for (const FString& PointInstancerPath : Impl.PointInstancerPaths)
 		{
 			UE::FSdfPath UsdPointInstancerPath{*PointInstancerPath};
 
@@ -1130,7 +1149,7 @@ namespace UE::USDInfoCache::Private
 	 * the main recursive pass just adds them to arrays, and we're allowed to handle bMergeIdenticalSlots
 	 * only here.
 	 */
-	void CollectMaterialSlotCounts(FUsdInfoCache::FUsdInfoCacheImpl& Impl, bool bContextMergeIdenticalSlots)
+	void CollectMaterialSlotCounts(FUsdInfoCacheImpl& Impl, bool bContextMergeIdenticalSlots)
 	{
 		TRACE_CPUPROFILER_EVENT_SCOPE(CollectMaterialSlotCounts);
 
@@ -1173,7 +1192,7 @@ namespace UE::USDInfoCache::Private
 	bool CanMeshSubtreeBeCollapsed(
 		const pxr::UsdPrim& UsdPrim,
 		FUsdSchemaTranslationContext& Context,
-		FUsdInfoCache::FUsdInfoCacheImpl& Impl,
+		FUsdInfoCacheImpl& Impl,
 		const TSharedPtr<FUsdSchemaTranslator>& Translator
 	)
 	{
@@ -1209,7 +1228,7 @@ namespace UE::USDInfoCache::Private
 	void RecursiveQueryCollapsesChildren(
 		const pxr::UsdPrim& UsdPrim,
 		FUsdSchemaTranslationContext& Context,
-		FUsdInfoCache::FUsdInfoCacheImpl& Impl,
+		FUsdInfoCacheImpl& Impl,
 		FUsdSchemaTranslatorRegistry& Registry
 	)
 	{
@@ -1330,7 +1349,7 @@ namespace UE::USDInfoCache::Private
 		return Result;
 	}
 
-	void RegisterInstanceableAuxPrims(FUsdSchemaTranslationContext& Context, FUsdInfoCache::FUsdInfoCacheImpl& Impl)
+	void RegisterInstanceableAuxPrims(FUsdSchemaTranslationContext& Context, FUsdInfoCacheImpl& Impl)
 	{
 		TRACE_CPUPROFILER_EVENT_SCOPE(UE::USDInfoCache::Private::RegisterInstanceableAuxPrims);
 		FScopedUsdAllocs Allocs;
@@ -1442,7 +1461,7 @@ namespace UE::USDInfoCache::Private
 	void FindValidGeometryCacheRoot(
 		const pxr::UsdPrim& UsdPrim,
 		FUsdSchemaTranslationContext& Context,
-		FUsdInfoCache::FUsdInfoCacheImpl& Impl,
+		FUsdInfoCacheImpl& Impl,
 		UE::UsdInfoCache::Private::EGeometryCachePrimState& OutState
 	)
 	{
@@ -1505,7 +1524,7 @@ namespace UE::USDInfoCache::Private
 	void RecursiveCheckForGeometryCache(
 		const pxr::UsdPrim& UsdPrim,
 		FUsdSchemaTranslationContext& Context,
-		FUsdInfoCache::FUsdInfoCacheImpl& Impl,
+		FUsdInfoCacheImpl& Impl,
 		bool bIsInsideSkelRoot,
 		int32& OutDepth,
 		UE::UsdInfoCache::Private::EGeometryCachePrimState& OutState
@@ -1725,7 +1744,7 @@ namespace UE::USDInfoCache::Private
 		OutState = PrimState;
 	}
 
-	void CheckForGeometryCache(const pxr::UsdPrim& UsdPrim, FUsdSchemaTranslationContext& Context, FUsdInfoCache::FUsdInfoCacheImpl& Impl)
+	void CheckForGeometryCache(const pxr::UsdPrim& UsdPrim, FUsdSchemaTranslationContext& Context, FUsdInfoCacheImpl& Impl)
 	{
 		TRACE_CPUPROFILER_EVENT_SCOPE(CheckForGeometryCache);
 
@@ -2095,8 +2114,8 @@ void FUsdInfoCache::RebuildCacheForSubtree(const UE::FUsdPrim& Prim, FUsdSchemaT
 			ImplPtr->InfoMap.Reset();
 		}
 		{
-			FWriteScopeLock PointInstancerLock(ImplPtr->TempPointInstancerPathsLock);
-			ImplPtr->TempPointInstancerPaths.Reset();
+			FWriteScopeLock PointInstancerLock(ImplPtr->PointInstancerPathsLock);
+			ImplPtr->PointInstancerPaths.Reset();
 		}
 
 		// This should be the first step as all future functions will expect to find one entry per prim in the cache
@@ -2149,8 +2168,8 @@ void FUsdInfoCache::Clear()
 			ImplPtr->AssetToPrimPaths.Empty();
 		}
 		{
-			FWriteScopeLock PointInstancerLock(ImplPtr->TempPointInstancerPathsLock);
-			ImplPtr->TempPointInstancerPaths.Empty();
+			FWriteScopeLock PointInstancerLock(ImplPtr->PointInstancerPathsLock);
+			ImplPtr->PointInstancerPaths.Empty();
 		}
 
 		ResetTranslatedPrototypes();
@@ -2237,3 +2256,542 @@ TOptional<bool> FUsdInfoCache::CanXformableSubtreeBeCollapsed(const UE::FSdfPath
 
 	return {};
 }
+
+PRAGMA_ENABLE_DEPRECATION_WARNINGS
+
+UUsdInfoCache::UUsdInfoCache()
+{
+	Impl = MakePimpl<FUsdInfoCacheImpl>();
+}
+
+void UUsdInfoCache::CopyImpl(const UUsdInfoCache& Other)
+{
+	const bool bAlwaysMarkDirty = false;
+	Modify(bAlwaysMarkDirty);
+
+	*Impl = *Other.Impl;
+}
+
+void UUsdInfoCache::Serialize(FArchive& Ar)
+{
+	TRACE_CPUPROFILER_EVENT_SCOPE(UUsdInfoCache::Serialize);
+
+	if (FUsdInfoCacheImpl* ImplPtr = Impl.Get())
+	{
+		FWriteScopeLock ScopeLock(ImplPtr->InfoMapLock);
+		Ar << ImplPtr->InfoMap;
+	}
+}
+
+bool UUsdInfoCache::ContainsInfoAboutPrim(const UE::FSdfPath& Path) const
+{
+	if (FUsdInfoCacheImpl* ImplPtr = Impl.Get())
+	{
+		FReadScopeLock ScopeLock(ImplPtr->InfoMapLock);
+		return ImplPtr->InfoMap.Contains(Path);
+	}
+
+	return false;
+}
+
+TSet<UE::FSdfPath> UUsdInfoCache::GetKnownPrims() const
+{
+	TSet<UE::FSdfPath> Result;
+
+	if (FUsdInfoCacheImpl* ImplPtr = Impl.Get())
+	{
+		FReadScopeLock ScopeLock(ImplPtr->InfoMapLock);
+
+		ImplPtr->InfoMap.GetKeys(Result);
+		return Result;
+	}
+
+	return Result;
+}
+
+bool UUsdInfoCache::IsPathCollapsed(const UE::FSdfPath& Path, ECollapsingType CollapsingType) const
+{
+	if (FUsdInfoCacheImpl* ImplPtr = Impl.Get())
+	{
+		FReadScopeLock ScopeLock(ImplPtr->InfoMapLock);
+
+		if (const UE::UsdInfoCache::Private::FUsdPrimInfo* FoundInfo = ImplPtr->InfoMap.Find(Path))
+		{
+			FReadScopeLock PrimLock{ImplPtr->PrimLocks[FoundInfo->PrimLockIndex]};
+
+			// See comment on the AssetCollapsedRoot member
+			return !FoundInfo->AssetCollapsedRoot.IsSet() || (!FoundInfo->AssetCollapsedRoot->IsEmpty() && FoundInfo->AssetCollapsedRoot != Path);
+		}
+
+		// This should never happen: We should have cached the entire tree
+		ensureMsgf(false, TEXT("Prim path '%s' has not been cached!"), *Path.GetString());
+	}
+
+	return false;
+}
+
+bool UUsdInfoCache::DoesPathCollapseChildren(const UE::FSdfPath& Path, ECollapsingType CollapsingType) const
+{
+	if (FUsdInfoCacheImpl* ImplPtr = Impl.Get())
+	{
+		FReadScopeLock ScopeLock(ImplPtr->InfoMapLock);
+
+		if (const UE::UsdInfoCache::Private::FUsdPrimInfo* FoundInfo = ImplPtr->InfoMap.Find(Path))
+		{
+			FReadScopeLock PrimLock{ImplPtr->PrimLocks[FoundInfo->PrimLockIndex]};
+
+			// We store our own Path in there when we collapse children.
+			// Otherwise we hold the path of our collapse root, or empty (in case nothing is collapsed up to here)
+			return FoundInfo->AssetCollapsedRoot == Path;
+		}
+
+		// This should never happen: We should have cached the entire tree
+		ensureMsgf(false, TEXT("Prim path '%s' has not been cached!"), *Path.GetString());
+	}
+
+	return false;
+}
+
+UE::FSdfPath UUsdInfoCache::UnwindToNonCollapsedPath(const UE::FSdfPath& Path, ECollapsingType CollapsingType) const
+{
+	if (FUsdInfoCacheImpl* ImplPtr = Impl.Get())
+	{
+		FReadScopeLock ScopeLock(ImplPtr->InfoMapLock);
+
+		TOptional<UE::FSdfPath> CollapseRoot;
+		UE::UsdInfoCache::Private::FUsdPrimInfo* MainFoundInfo = ImplPtr->InfoMap.Find(Path);
+		if (MainFoundInfo)
+		{
+			FReadScopeLock PrimLock{ImplPtr->PrimLocks[MainFoundInfo->PrimLockIndex]};
+			CollapseRoot = MainFoundInfo->AssetCollapsedRoot;
+		}
+
+		// We never visited this prim before. We know it's collapsed, let's find our collapse root
+		if (!CollapseRoot.IsSet())
+		{
+			TArray<UE::UsdInfoCache::Private::FUsdPrimInfo*> InfosToUpdate = {MainFoundInfo};
+
+			UE::FSdfPath TraversalPath = Path.GetParentPath();
+			while (!TraversalPath.IsAbsoluteRootPath())
+			{
+				if (UE::UsdInfoCache::Private::FUsdPrimInfo* AncestorInfo = ImplPtr->InfoMap.Find(TraversalPath))
+				{
+					FReadScopeLock PrimLock{ImplPtr->PrimLocks[AncestorInfo->PrimLockIndex]};
+
+					// We found an ancestor that has this filled in: We're collapsed, so whatever is its
+					// collapse root is also our collapse root
+					if (AncestorInfo->AssetCollapsedRoot.IsSet())
+					{
+						// If our original Path doesn't have anything filled in, then we *must* be a child of a
+						// collapsed prim (i.e. something that has a non-empty path in its AssetCollapsedRoot)
+						ensure(!AncestorInfo->AssetCollapsedRoot->IsEmpty());
+
+						CollapseRoot = AncestorInfo->AssetCollapsedRoot;
+						break;
+					}
+					// Still nothing. Let's keep track of this info so that we can update it later
+					else
+					{
+						InfosToUpdate.Add(AncestorInfo);
+					}
+				}
+
+				TraversalPath = TraversalPath.GetParentPath();
+			}
+
+			// Fill in all visited infos with what we found on our ancestor
+			// Note: Here we modify the stored data, and yet we won't call Modify(). The idea is that while we are
+			// in fact filling in some fields, those are always derived from the other fields we already have, so there's
+			// no new information being set. If our transaction happens to wipe these new values, we can just compute
+			// them again on-demand
+			if (CollapseRoot.IsSet())
+			{
+				for (UE::UsdInfoCache::Private::FUsdPrimInfo* Info : InfosToUpdate)
+				{
+					FWriteScopeLock PrimLock{ImplPtr->PrimLocks[Info->PrimLockIndex]};
+					Info->AssetCollapsedRoot = CollapseRoot;
+				}
+			}
+		}
+
+		// We have visited this prim during the info cache build (or another Unwind, or just now within this function)
+		if (CollapseRoot.IsSet())
+		{
+			// An empty path here means that we are not collapsed at all
+			if (CollapseRoot->IsEmpty())
+			{
+				return Path;
+			}
+			// Otherwise we have our own path in there (in case we collapse children) or the path to the prim that collapsed us
+			else
+			{
+				return CollapseRoot.GetValue();
+			}
+		}
+
+		// This should never happen: We should have cached the entire tree
+		ensureAlwaysMsgf(false, TEXT("Prim path '%s' has not been cached!"), *Path.GetString());
+	}
+
+	return Path;
+}
+
+TSet<UE::FSdfPath> UUsdInfoCache::GetMainPrims(const UE::FSdfPath& AuxPrimPath) const
+{
+	if (FUsdInfoCacheImpl* ImplPtr = Impl.Get())
+	{
+		FReadScopeLock ScopeLock(ImplPtr->InfoMapLock);
+
+		if (const UE::UsdInfoCache::Private::FUsdPrimInfo* FoundInfo = ImplPtr->InfoMap.Find(AuxPrimPath))
+		{
+			FReadScopeLock PrimLock{ImplPtr->PrimLocks[FoundInfo->PrimLockIndex]};
+
+			TSet<UE::FSdfPath> Result = FoundInfo->MainPrims;
+			Result.Add(AuxPrimPath);
+			return Result;
+		}
+	}
+
+	return {AuxPrimPath};
+}
+
+TSet<UE::FSdfPath> UUsdInfoCache::GetAuxiliaryPrims(const UE::FSdfPath& MainPrimPath) const
+{
+	if (FUsdInfoCacheImpl* ImplPtr = Impl.Get())
+	{
+		FReadScopeLock ScopeLock(ImplPtr->InfoMapLock);
+
+		if (const UE::UsdInfoCache::Private::FUsdPrimInfo* FoundInfo = ImplPtr->InfoMap.Find(MainPrimPath))
+		{
+			FReadScopeLock PrimLock{ImplPtr->PrimLocks[FoundInfo->PrimLockIndex]};
+
+			TSet<UE::FSdfPath> Result = FoundInfo->AuxPrims;
+			Result.Add(MainPrimPath);
+			return Result;
+		}
+	}
+
+	return {MainPrimPath};
+}
+
+TSet<UE::FSdfPath> UUsdInfoCache::GetMaterialUsers(const UE::FSdfPath& Path) const
+{
+	if (FUsdInfoCacheImpl* ImplPtr = Impl.Get())
+	{
+		FReadScopeLock ScopeLock(ImplPtr->InfoMapLock);
+
+		if (const UE::UsdInfoCache::Private::FUsdPrimInfo* FoundInfo = ImplPtr->InfoMap.Find(Path))
+		{
+			FReadScopeLock PrimLock{ImplPtr->PrimLocks[FoundInfo->PrimLockIndex]};
+
+			return FoundInfo->MaterialUsers;
+		}
+	}
+
+	return {};
+}
+
+bool UUsdInfoCache::IsMaterialUsed(const UE::FSdfPath& Path) const
+{
+	if (FUsdInfoCacheImpl* ImplPtr = Impl.Get())
+	{
+		FReadScopeLock ScopeLock(ImplPtr->InfoMapLock);
+
+		if (const UE::UsdInfoCache::Private::FUsdPrimInfo* FoundInfo = ImplPtr->InfoMap.Find(Path))
+		{
+			FReadScopeLock PrimLock{ImplPtr->PrimLocks[FoundInfo->PrimLockIndex]};
+
+			return FoundInfo->MaterialUsers.Num() > 0;
+		}
+	}
+
+	return false;
+}
+
+bool UUsdInfoCache::IsPotentialGeometryCacheRoot(const UE::FSdfPath& Path) const
+{
+	if (FUsdInfoCacheImpl* ImplPtr = Impl.Get())
+	{
+		FReadScopeLock ScopeLock(ImplPtr->InfoMapLock);
+
+		if (const UE::UsdInfoCache::Private::FUsdPrimInfo* FoundInfo = ImplPtr->InfoMap.Find(Path))
+		{
+			FReadScopeLock PrimLock{ImplPtr->PrimLocks[FoundInfo->PrimLockIndex]};
+
+			return FoundInfo->GeometryCacheState == UE::UsdInfoCache::Private::EGeometryCachePrimState::ValidRoot;
+		}
+
+		// This should never happen: We should have cached the entire tree
+		ensureMsgf(false, TEXT("Prim path '%s' has not been cached!"), *Path.GetString());
+	}
+
+	return false;
+}
+
+void UUsdInfoCache::ResetTranslatedPrototypes()
+{
+	if (FUsdInfoCacheImpl* ImplPtr = Impl.Get())
+	{
+		FWriteScopeLock ScopeLock(ImplPtr->TranslatedPrototypesLock);
+		return ImplPtr->TranslatedPrototypes.Reset();
+	}
+}
+
+bool UUsdInfoCache::IsPrototypeTranslated(const UE::FSdfPath& PrototypePath)
+{
+	if (FUsdInfoCacheImpl* ImplPtr = Impl.Get())
+	{
+		FReadScopeLock ScopeLock(ImplPtr->TranslatedPrototypesLock);
+		return ImplPtr->TranslatedPrototypes.Contains(PrototypePath);
+	}
+
+	return false;
+}
+
+void UUsdInfoCache::MarkPrototypeAsTranslated(const UE::FSdfPath& PrototypePath)
+{
+	if (FUsdInfoCacheImpl* ImplPtr = Impl.Get())
+	{
+		FWriteScopeLock ScopeLock(ImplPtr->TranslatedPrototypesLock);
+		ImplPtr->TranslatedPrototypes.Add(PrototypePath);
+	}
+}
+
+TOptional<uint64> UUsdInfoCache::GetSubtreeVertexCount(const UE::FSdfPath& Path)
+{
+	if (FUsdInfoCacheImpl* ImplPtr = Impl.Get())
+	{
+		FReadScopeLock ScopeLock(ImplPtr->InfoMapLock);
+
+		if (const UE::UsdInfoCache::Private::FUsdPrimInfo* FoundInfo = ImplPtr->InfoMap.Find(Path))
+		{
+			FReadScopeLock PrimLock{ImplPtr->PrimLocks[FoundInfo->PrimLockIndex]};
+
+			return FoundInfo->ExpectedVertexCountForSubtree;
+		}
+
+		// This should never happen: We should have cached the entire tree
+		ensureMsgf(false, TEXT("Prim path '%s' has not been cached!"), *Path.GetString());
+	}
+
+	return {};
+}
+
+TOptional<uint64> UUsdInfoCache::GetSubtreeMaterialSlotCount(const UE::FSdfPath& Path)
+{
+	if (FUsdInfoCacheImpl* ImplPtr = Impl.Get())
+	{
+		FReadScopeLock ScopeLock(ImplPtr->InfoMapLock);
+
+		if (const UE::UsdInfoCache::Private::FUsdPrimInfo* FoundInfo = ImplPtr->InfoMap.Find(Path))
+		{
+			FReadScopeLock PrimLock{ImplPtr->PrimLocks[FoundInfo->PrimLockIndex]};
+
+			return FoundInfo->SubtreeMaterialSlots.Num();
+		}
+
+		// This should never happen: We should have cached the entire tree
+		ensureMsgf(false, TEXT("Prim path '%s' has not been cached!"), *Path.GetString());
+	}
+
+	return {};
+}
+
+TOptional<TArray<UsdUtils::FUsdPrimMaterialSlot>> UUsdInfoCache::GetSubtreeMaterialSlots(const UE::FSdfPath& Path)
+{
+	if (FUsdInfoCacheImpl* ImplPtr = Impl.Get())
+	{
+		FReadScopeLock ScopeLock(ImplPtr->InfoMapLock);
+
+		if (const UE::UsdInfoCache::Private::FUsdPrimInfo* FoundInfo = ImplPtr->InfoMap.Find(Path))
+		{
+			FReadScopeLock PrimLock{ImplPtr->PrimLocks[FoundInfo->PrimLockIndex]};
+
+			return FoundInfo->SubtreeMaterialSlots;
+		}
+
+		// This should never happen: We should have cached the entire tree
+		ensureMsgf(false, TEXT("Prim path '%s' has not been cached!"), *Path.GetString());
+	}
+
+	return {};
+}
+
+void UUsdInfoCache::RebuildCacheForSubtree(const UE::FUsdPrim& Prim, FUsdSchemaTranslationContext& Context)
+{
+#if USE_USD_SDK
+	TRACE_CPUPROFILER_EVENT_SCOPE(UUsdInfoCache::RebuildCacheForSubtree);
+
+	using namespace UE::USDInfoCache::Private;
+
+	FUsdInfoCacheImpl* ImplPtr = Impl.Get();
+	if (!ImplPtr)
+	{
+		return;
+	}
+
+	bool bAlwaysMarkDirty = false;
+	Modify(bAlwaysMarkDirty);
+
+	// We can't deallocate our info cache pointer with the Usd allocator
+	FScopedUnrealAllocs UEAllocs;
+
+	TGuardValue<bool> Guard{Context.bIsBuildingInfoCache, true};
+	{
+		pxr::UsdPrim UsdPrim{Prim};
+		if (!UsdPrim)
+		{
+			return;
+		}
+
+		ImplPtr->TempStage = UE::FUsdStageWeak{UsdPrim.GetStage()};
+
+		IUsdSchemasModule& UsdSchemasModule = FModuleManager::Get().LoadModuleChecked<IUsdSchemasModule>(TEXT("USDSchemas"));
+		FUsdSchemaTranslatorRegistry& Registry = UsdSchemasModule.GetTranslatorRegistry();
+
+		pxr::TfToken MaterialPurposeToken = pxr::UsdShadeTokens->allPurpose;
+		if (!Context.MaterialPurpose.IsNone())
+		{
+			MaterialPurposeToken = UnrealToUsd::ConvertToken(*Context.MaterialPurpose.ToString()).Get();
+		}
+
+		Clear();
+
+		// This should be the first step as all future functions will expect to find one entry per prim in the cache
+		RepopulateInfoMap(UsdPrim, *ImplPtr);
+
+		// Propagate vertex and material slot counts before we query CollapsesChildren because the Xformable
+		// translator needs to know when it would generate too large a static mesh
+		uint64 SubtreeVertexCount = 0;
+		TArray<UsdUtils::FUsdPrimMaterialSlot> SubtreeSlots;
+		const bool bPossibleInheritedBindings = false;
+		RecursivePropagateVertexAndMaterialSlotCounts(
+			UsdPrim,
+			Context,
+			MaterialPurposeToken,
+			*ImplPtr,
+			Registry,
+			SubtreeVertexCount,
+			SubtreeSlots,
+			bPossibleInheritedBindings
+		);
+
+		UpdateInfoForPointInstancers(Context, *ImplPtr);
+
+		CheckForGeometryCache(UsdPrim, Context, *ImplPtr);
+
+		RecursiveQueryCollapsesChildren(UsdPrim, Context, *ImplPtr, Registry);
+
+		RegisterInstanceableAuxPrims(Context, *ImplPtr);
+
+		CollectMaterialSlotCounts(*ImplPtr, Context.bMergeIdenticalMaterialSlots);
+	}
+#endif	  // USE_USD_SDK
+}
+
+void UUsdInfoCache::Clear()
+{
+	TRACE_CPUPROFILER_EVENT_SCOPE(UUsdInfoCache::Clear);
+
+	if (FUsdInfoCacheImpl* ImplPtr = Impl.Get())
+	{
+		{
+			TRACE_CPUPROFILER_EVENT_SCOPE(InfoCacheModify);
+			bool bAlwaysMarkDirty = false;
+			Modify(bAlwaysMarkDirty);
+		}
+		{
+			TRACE_CPUPROFILER_EVENT_SCOPE(InfoMapEmpty);
+			FWriteScopeLock ScopeLock(ImplPtr->InfoMapLock);
+			ImplPtr->InfoMap.Empty();
+		}
+		{
+			FWriteScopeLock ScopeLock(ImplPtr->PointInstancerPathsLock);
+			ImplPtr->PointInstancerPaths.Empty();
+		}
+	}
+}
+
+bool UUsdInfoCache::IsEmpty()
+{
+	if (FUsdInfoCacheImpl* ImplPtr = Impl.Get())
+	{
+		FReadScopeLock ScopeLock(ImplPtr->InfoMapLock);
+		return ImplPtr->InfoMap.IsEmpty();
+	}
+
+	return true;
+}
+
+TOptional<bool> UUsdInfoCache::CanXformableSubtreeBeCollapsed(const UE::FSdfPath& RootPath, FUsdSchemaTranslationContext& Context)
+{
+#if USE_USD_SDK
+	TRACE_CPUPROFILER_EVENT_SCOPE(UUsdInfoCache::CanSubtreeBeCollapsed);
+
+	// The only reason this function exists is that FUsdGeomXformableTranslator::CollapsesChildren() needs to check if all
+	// GeomXformable prims in its subtree return true for CanBeCollapsed().
+	//
+	// We don't want to compute this for the entire stage on the main info cache build, because it may not be needed.
+	// However, we definitely do not want each call to FUsdGeomXformableTranslator::CollapsesChildren() to traverse its entire
+	// subtree of prims calling CanBeCollapsed() on their own: That would be a massive waste since the output is going to
+	// be the same regardless of the caller.
+	//
+	// This is the awkward compromise where the first call to FUsdGeomXformableTranslator::CollapsesChildren() will traverse
+	// its entire subtree and fill this in, and subsequent calls can just use those results, or fill in additional subtrees, etc.
+
+	if (FUsdInfoCacheImpl* ImplPtr = Impl.Get())
+	{
+		FReadScopeLock ScopeLock(ImplPtr->InfoMapLock);
+
+		if (const UE::UsdInfoCache::Private::FUsdPrimInfo* FoundInfo = ImplPtr->InfoMap.Find(RootPath))
+		{
+			FReadScopeLock PrimLock{ImplPtr->PrimLocks[FoundInfo->PrimLockIndex]};
+			if (FoundInfo->bXformSubtreeCanBeCollapsed.IsSet())
+			{
+				return FoundInfo->bXformSubtreeCanBeCollapsed.GetValue();
+			}
+		}
+
+		TOptional<bool> bCanBeCollapsed;
+
+		// Fill in missing entries for CanBeCollapsed on-demand and compute the value for the prim at RootPath,
+		// if we can still access our stage
+		pxr::UsdStageWeakPtr UsdStageWeak = ImplPtr->TempStage;
+		pxr::UsdStageRefPtr Stage = UsdStageWeak;
+		if (Stage)
+		{
+			if (pxr::UsdPrim Prim = Stage->GetPrimAtPath(RootPath))
+			{
+				IUsdSchemasModule& UsdSchemasModule = FModuleManager::Get().LoadModuleChecked<IUsdSchemasModule>(TEXT("USDSchemas"));
+				FUsdSchemaTranslatorRegistry& Registry = UsdSchemasModule.GetTranslatorRegistry();
+
+				bool bAlwaysMarkDirty = false;
+				Modify(bAlwaysMarkDirty);
+
+				bCanBeCollapsed = UE::USDInfoCache::Private::RecursiveQueryCanBeCollapsed(Prim, Context, *ImplPtr, Registry);
+			}
+		}
+
+		// We can potentially still fail to find this here, in case our stage reference is broken (i.e. called outside of the
+		// main infocache build callstack).
+		//
+		// There shouldn't be any point in checking our FoundInfo again though: If we didn't return anything valid from
+		// our call to RecursiveQueryCanBeCollapsed, then we didn't put anything new on the InfoMap either
+		if (bCanBeCollapsed.IsSet())
+		{
+			return bCanBeCollapsed;
+		}
+
+		UE_LOG(
+			LogUsd,
+			Warning,
+			TEXT(
+				"Failed to find whether subtree '%s' can be collapsed or not. Note: This function is meant to be used only during the main UUsdInfoCache build!"
+			),
+			*RootPath.GetString()
+		);
+	}
+#endif	  // USE_USD_SDK
+
+	return {};
+}
diff --git a/Engine/Plugins/Importers/USDImporter/Source/USDSchemas/Private/USDLuxLightTranslator.cpp b/Engine/Plugins/Importers/USDImporter/Source/USDSchemas/Private/USDLuxLightTranslator.cpp
--- a/Engine/Plugins/Importers/USDImporter/Source/USDSchemas/Private/USDLuxLightTranslator.cpp
+++ b/Engine/Plugins/Importers/USDImporter/Source/USDSchemas/Private/USDLuxLightTranslator.cpp
@@ -32,7 +32,7 @@ void FUsdLuxLightTranslator::CreateAssets()
 {
 	TRACE_CPUPROFILER_EVENT_SCOPE(FUsdLuxLightTranslator::CreateAssets);
 
-	if (!Context->UsdAssetCache || !Context->InfoCache)
+	if (!Context->UsdAssetCache || !Context->PrimLinkCache)
 	{
 		return;
 	}
@@ -102,7 +102,7 @@ void FUsdLuxLightTranslator::CreateAssets()
 
 	if (Texture)
 	{
-		Context->InfoCache->LinkAssetToPrim(PrimPath, Texture);
+		Context->PrimLinkCache->LinkAssetToPrim(PrimPath, Texture);
 	}
 }
 
@@ -131,7 +131,7 @@ void FUsdLuxLightTranslator::UpdateComponents(USceneComponent* SceneComponent)
 
 	ULightComponentBase* LightComponent = Cast<ULightComponentBase>(SceneComponent);
 
-	if (!LightComponent)
+	if (!LightComponent || !Context->PrimLinkCache)
 	{
 		return;
 	}
@@ -179,7 +179,7 @@ void FUsdLuxLightTranslator::UpdateComponents(USceneComponent* SceneComponent)
 	{
 		SkyLightComponent->Modify();
 
-		if (UTextureCube* TextuxeCube = Context->InfoCache->GetSingleAssetForPrim<UTextureCube>(PrimPath))
+		if (UTextureCube* TextuxeCube = Context->PrimLinkCache->GetSingleAssetForPrim<UTextureCube>(PrimPath))
 		{
 			SkyLightComponent->Cubemap = TextuxeCube;
 			SkyLightComponent->SourceType = ESkyLightSourceType::SLS_SpecifiedCubemap;
diff --git a/Engine/Plugins/Importers/USDImporter/Source/USDSchemas/Private/USDMediaSpatialAudioTranslator.cpp b/Engine/Plugins/Importers/USDImporter/Source/USDSchemas/Private/USDMediaSpatialAudioTranslator.cpp
--- a/Engine/Plugins/Importers/USDImporter/Source/USDSchemas/Private/USDMediaSpatialAudioTranslator.cpp
+++ b/Engine/Plugins/Importers/USDImporter/Source/USDSchemas/Private/USDMediaSpatialAudioTranslator.cpp
@@ -173,7 +173,7 @@ void FUsdMediaSpatialAudioTranslator::CreateAssets()
 
 	using namespace UE::UsdMediaSpatialAudioTranslator::Private;
 
-	if (!Context->UsdAssetCache || !Context->InfoCache)
+	if (!Context->UsdAssetCache || !Context->PrimLinkCache)
 	{
 		return;
 	}
@@ -231,7 +231,7 @@ void FUsdMediaSpatialAudioTranslator::CreateAssets()
 
 	if (SoundWave)
 	{
-		Context->InfoCache->LinkAssetToPrim(PrimPath, SoundWave);
+		Context->PrimLinkCache->LinkAssetToPrim(PrimPath, SoundWave);
 
 		if (UUsdAssetUserData* UserData = UsdUnreal::ObjectUtils::GetOrCreateAssetUserData<UUsdAssetUserData>(SoundWave))
 		{
diff --git a/Engine/Plugins/Importers/USDImporter/Source/USDSchemas/Private/USDPrimLinkCache.cpp b/Engine/Plugins/Importers/USDImporter/Source/USDSchemas/Private/USDPrimLinkCache.cpp
new file mode 100644
--- /dev/null
+++ b/Engine/Plugins/Importers/USDImporter/Source/USDSchemas/Private/USDPrimLinkCache.cpp
@@ -0,0 +1,276 @@
+// Copyright Epic Games, Inc. All Rights Reserved.
+
+#include "USDPrimLinkCache.h"
+
+#include "USDLog.h"
+#include "UsdWrappers/SdfPath.h"
+
+struct UUsdPrimLinkCache::UUsdPrimLinkCacheImpl
+{
+	UUsdPrimLinkCacheImpl() = default;
+	UUsdPrimLinkCacheImpl(const UUsdPrimLinkCacheImpl& Other)
+		: UUsdPrimLinkCacheImpl()
+	{
+		FReadScopeLock OtherScopedPrimPathToAssetsLock(Other.PrimPathToAssetsLock);
+		FWriteScopeLock ThisScopedPrimPathToAssetsLock(PrimPathToAssetsLock);
+
+		PrimPathToAssets = Other.PrimPathToAssets;
+		AssetToPrimPaths = Other.AssetToPrimPaths;
+	}
+
+	// Information we may have about a subset of prims
+	TMap<UE::FSdfPath, TArray<TWeakObjectPtr<UObject>>> PrimPathToAssets;
+	TMap<TWeakObjectPtr<UObject>, TArray<UE::FSdfPath>> AssetToPrimPaths;
+	mutable FRWLock PrimPathToAssetsLock;
+};
+
+UUsdPrimLinkCache::UUsdPrimLinkCache()
+{
+	Impl = MakeUnique<UUsdPrimLinkCache::UUsdPrimLinkCacheImpl>();
+}
+
+// Boilerplate to support Pimpl in an UObject
+// See UniquePtr.h, TDefaultDelete
+UUsdPrimLinkCache::UUsdPrimLinkCache(FVTableHelper& Helper)
+	: Super(Helper)
+{
+}
+
+UUsdPrimLinkCache::~UUsdPrimLinkCache()
+{
+}
+
+void UUsdPrimLinkCache::Serialize(FArchive& Ar)
+{
+	TRACE_CPUPROFILER_EVENT_SCOPE(UUsdPrimLinkCache::Serialize);
+
+	if (UUsdPrimLinkCacheImpl* ImplPtr = Impl.Get())
+	{
+		FWriteScopeLock ScopeLock(ImplPtr->PrimPathToAssetsLock);
+		Ar << ImplPtr->PrimPathToAssets;
+		Ar << ImplPtr->AssetToPrimPaths;
+	}
+}
+
+bool UUsdPrimLinkCache::ContainsInfoAboutPrim(const UE::FSdfPath& Path) const
+{
+	if (UUsdPrimLinkCacheImpl* ImplPtr = Impl.Get())
+	{
+		FReadScopeLock ScopeLock(ImplPtr->PrimPathToAssetsLock);
+		return ImplPtr->PrimPathToAssets.Contains(Path);
+	}
+
+	return false;
+}
+
+TSet<UE::FSdfPath> UUsdPrimLinkCache::GetKnownPrims() const
+{
+	TSet<UE::FSdfPath> Result;
+
+	if (UUsdPrimLinkCacheImpl* ImplPtr = Impl.Get())
+	{
+		FReadScopeLock ScopeLock(ImplPtr->PrimPathToAssetsLock);
+		ImplPtr->PrimPathToAssets.GetKeys(Result);
+		return Result;
+	}
+
+	return Result;
+}
+
+void UUsdPrimLinkCache::LinkAssetToPrim(const UE::FSdfPath& Path, UObject* Asset)
+{
+	UUsdPrimLinkCacheImpl* ImplPtr = Impl.Get();
+	if (!ImplPtr)
+	{
+		return;
+	}
+
+	bool bAlwaysMarkDirty = false;
+	Modify(bAlwaysMarkDirty);
+
+	FWriteScopeLock ScopeLock(ImplPtr->PrimPathToAssetsLock);
+
+	UE_LOG(LogUsd, Verbose, TEXT("Linking asset '%s' to prim '%s'"), *Asset->GetPathName(), *Path.GetString());
+
+	ImplPtr->PrimPathToAssets.FindOrAdd(Path).AddUnique(Asset);
+	ImplPtr->AssetToPrimPaths.FindOrAdd(Asset).AddUnique(Path);
+}
+
+void UUsdPrimLinkCache::UnlinkAssetFromPrim(const UE::FSdfPath& Path, UObject* Asset)
+{
+	UUsdPrimLinkCacheImpl* ImplPtr = Impl.Get();
+	if (!ImplPtr)
+	{
+		return;
+	}
+
+	bool bAlwaysMarkDirty = false;
+	Modify(bAlwaysMarkDirty);
+
+	FWriteScopeLock ScopeLock(ImplPtr->PrimPathToAssetsLock);
+
+	UE_LOG(LogUsd, Verbose, TEXT("Unlinking asset '%s' to prim '%s'"), *Asset->GetPathName(), *Path.GetString());
+
+	if (TArray<TWeakObjectPtr<UObject>>* FoundAssetsForPrim = ImplPtr->PrimPathToAssets.Find(Path))
+	{
+		FoundAssetsForPrim->Remove(Asset);
+	}
+	if (TArray<UE::FSdfPath>* FoundPrimPathsForAsset = ImplPtr->AssetToPrimPaths.Find(Asset))
+	{
+		FoundPrimPathsForAsset->Remove(Path);
+	}
+}
+
+TArray<TWeakObjectPtr<UObject>> UUsdPrimLinkCache::RemoveAllAssetPrimLinks(const UE::FSdfPath& Path)
+{
+	UUsdPrimLinkCacheImpl* ImplPtr = Impl.Get();
+	if (!ImplPtr)
+	{
+		return {};
+	}
+
+	bool bAlwaysMarkDirty = false;
+	Modify(bAlwaysMarkDirty);
+
+	FWriteScopeLock ScopeLock(ImplPtr->PrimPathToAssetsLock);
+
+	UE_LOG(LogUsd, Verbose, TEXT("Removing asset prim links for path '%s'"), *Path.GetString());
+
+	TArray<TWeakObjectPtr<UObject>> Assets;
+	ImplPtr->PrimPathToAssets.RemoveAndCopyValue(Path, Assets);
+
+	for (const TWeakObjectPtr<UObject>& Asset : Assets)
+	{
+		if (TArray<UE::FSdfPath>* PrimPaths = ImplPtr->AssetToPrimPaths.Find(Asset))
+		{
+			PrimPaths->Remove(Path);
+		}
+	}
+
+	return Assets;
+}
+
+TArray<UE::FSdfPath> UUsdPrimLinkCache::RemoveAllAssetPrimLinks(const UObject* Asset)
+{
+	UUsdPrimLinkCacheImpl* ImplPtr = Impl.Get();
+	if (!ImplPtr)
+	{
+		return {};
+	}
+
+	bool bAlwaysMarkDirty = false;
+	Modify(bAlwaysMarkDirty);
+
+	FWriteScopeLock ScopeLock(ImplPtr->PrimPathToAssetsLock);
+
+	UE_LOG(LogUsd, Verbose, TEXT("Removing asset prim links for asset '%s'"), *Asset->GetPathName());
+
+	TArray<UE::FSdfPath> PrimPaths;
+	ImplPtr->AssetToPrimPaths.RemoveAndCopyValue(const_cast<UObject*>(Asset), PrimPaths);
+
+	for (const UE::FSdfPath& Path : PrimPaths)
+	{
+		if (TArray<TWeakObjectPtr<UObject>>* Assets = ImplPtr->PrimPathToAssets.Find(Path))
+		{
+			Assets->Remove(const_cast<UObject*>(Asset));
+		}
+	}
+
+	return PrimPaths;
+}
+
+void UUsdPrimLinkCache::RemoveAllAssetPrimLinks()
+{
+	UUsdPrimLinkCacheImpl* ImplPtr = Impl.Get();
+	if (!ImplPtr)
+	{
+		return;
+	}
+
+	bool bAlwaysMarkDirty = false;
+	Modify(bAlwaysMarkDirty);
+
+	FWriteScopeLock ScopeLock(ImplPtr->PrimPathToAssetsLock);
+
+	UE_LOG(LogUsd, Verbose, TEXT("Removing all asset prim links"));
+
+	ImplPtr->PrimPathToAssets.Empty();
+	ImplPtr->AssetToPrimPaths.Empty();
+}
+
+TArray<TWeakObjectPtr<UObject>> UUsdPrimLinkCache::GetAllAssetsForPrim(const UE::FSdfPath& Path) const
+{
+	UUsdPrimLinkCacheImpl* ImplPtr = Impl.Get();
+	if (!ImplPtr)
+	{
+		return {};
+	}
+	FReadScopeLock ScopeLock(ImplPtr->PrimPathToAssetsLock);
+
+	if (const TArray<TWeakObjectPtr<UObject>>* FoundAssets = ImplPtr->PrimPathToAssets.Find(Path))
+	{
+		return *FoundAssets;
+	}
+
+	return {};
+}
+
+TArray<UE::FSdfPath> UUsdPrimLinkCache::GetPrimsForAsset(const UObject* Asset) const
+{
+	if (!Asset)
+	{
+		return {};
+	}
+
+	UUsdPrimLinkCacheImpl* ImplPtr = Impl.Get();
+	if (!ImplPtr)
+	{
+		return {};
+	}
+	FReadScopeLock ScopeLock(ImplPtr->PrimPathToAssetsLock);
+
+	if (const TArray<UE::FSdfPath>* FoundPrims = ImplPtr->AssetToPrimPaths.Find(Asset))
+	{
+		return *FoundPrims;
+	}
+
+	return {};
+}
+
+TMap<UE::FSdfPath, TArray<TWeakObjectPtr<UObject>>> UUsdPrimLinkCache::GetAllAssetPrimLinks() const
+{
+	UUsdPrimLinkCacheImpl* ImplPtr = Impl.Get();
+	if (!ImplPtr)
+	{
+		return {};
+	}
+
+	return ImplPtr->PrimPathToAssets;
+}
+
+void UUsdPrimLinkCache::Clear()
+{
+	TRACE_CPUPROFILER_EVENT_SCOPE(AUsdStageActor::UnloadUsdStage);
+
+	if (UUsdPrimLinkCacheImpl* ImplPtr = Impl.Get())
+	{
+		bool bAlwaysMarkDirty = false;
+		Modify(bAlwaysMarkDirty);
+
+		TRACE_CPUPROFILER_EVENT_SCOPE(PrimPathToAssetsEmpty);
+		FWriteScopeLock ScopeLock(ImplPtr->PrimPathToAssetsLock);
+		ImplPtr->PrimPathToAssets.Empty();
+		ImplPtr->AssetToPrimPaths.Empty();
+	}
+}
+
+bool UUsdPrimLinkCache::IsEmpty()
+{
+	if (UUsdPrimLinkCacheImpl* ImplPtr = Impl.Get())
+	{
+		FReadScopeLock ScopeLock(ImplPtr->PrimPathToAssetsLock);
+		return ImplPtr->PrimPathToAssets.IsEmpty();
+	}
+
+	return true;
+}
diff --git a/Engine/Plugins/Importers/USDImporter/Source/USDSchemas/Private/USDSchemaTranslator.cpp b/Engine/Plugins/Importers/USDImporter/Source/USDSchemas/Private/USDSchemaTranslator.cpp
--- a/Engine/Plugins/Importers/USDImporter/Source/USDSchemas/Private/USDSchemaTranslator.cpp
+++ b/Engine/Plugins/Importers/USDImporter/Source/USDSchemas/Private/USDSchemaTranslator.cpp
@@ -234,9 +234,9 @@ void FUsdSchemaTranslationContext::CompleteTasks()
 	}
 
 	// These only make sense within an asset translation period
-	if (InfoCache)
+	if (UsdInfoCache)
 	{
-		InfoCache->ResetTranslatedPrototypes();
+		UsdInfoCache->ResetTranslatedPrototypes();
 	}
 }
 
@@ -247,7 +247,7 @@ bool FUsdSchemaTranslator::IsCollapsed(ECollapsingType CollapsingType) const
 
 	if (!Context->bIsBuildingInfoCache)
 	{
-		return Context->InfoCache->IsPathCollapsed(PrimPath, CollapsingType);
+		return Context->UsdInfoCache->IsPathCollapsed(PrimPath, CollapsingType);
 	}
 
 	// This is merely a fallback, and we should never need this
@@ -290,12 +290,12 @@ bool FUsdSchemaTranslator::ShouldSkipInstance() const
 	UE::FSdfPath PrototypePath = GetPrototypePrimPath();
 	if (!PrototypePath.IsEmpty())
 	{
-		if (Context->InfoCache->IsPrototypeTranslated(PrototypePath))
+		if (Context->UsdInfoCache->IsPrototypeTranslated(PrototypePath))
 		{
 			return true;
 		}
 
-		Context->InfoCache->MarkPrototypeAsTranslated(PrototypePath);
+		Context->UsdInfoCache->MarkPrototypeAsTranslated(PrototypePath);
 	}
 
 	return false;
diff --git a/Engine/Plugins/Importers/USDImporter/Source/USDSchemas/Private/USDShadeMaterialTranslator.cpp b/Engine/Plugins/Importers/USDImporter/Source/USDSchemas/Private/USDShadeMaterialTranslator.cpp
--- a/Engine/Plugins/Importers/USDImporter/Source/USDSchemas/Private/USDShadeMaterialTranslator.cpp
+++ b/Engine/Plugins/Importers/USDImporter/Source/USDSchemas/Private/USDShadeMaterialTranslator.cpp
@@ -40,7 +40,7 @@ namespace UE::UsdShadeTranslator::Private
 		TSet<UMaterialInterface*>& NewMaterials
 	)
 	{
-		if (!Context->UsdAssetCache || !Context->InfoCache)
+		if (!Context->UsdAssetCache || !Context->PrimLinkCache)
 		{
 			return;
 		}
@@ -136,7 +136,7 @@ namespace UE::UsdShadeTranslator::Private
 		  // Don't spell out MID directly, as at runtime we may be trying to upgrade a packaged MIC
 							if (UMaterialInstance* MI = Cast<UMaterialInstance>(UserMaterial))
 							{
-								TArray<UE::FSdfPath> PrimsForAsset = Context->InfoCache->GetPrimsForAsset(MI);
+								TArray<UE::FSdfPath> PrimsForAsset = Context->PrimLinkCache->GetPrimsForAsset(MI);
 								const FString Hash = Context->UsdAssetCache->GetHashForAsset(MI);
 
 								// For MID we can't swap the reference material, so we need to remove the old material from the cache,
@@ -180,7 +180,7 @@ namespace UE::UsdShadeTranslator::Private
 
 								for (const UE::FSdfPath& PrimPath : PrimsForAsset)
 								{
-									Context->InfoCache->LinkAssetToPrim(PrimPath, NewMID);
+									Context->PrimLinkCache->LinkAssetToPrim(PrimPath, NewMID);
 								}
 								NewMaterials.Add(NewMID);
 
@@ -282,9 +282,9 @@ void FUsdShadeMaterialTranslator::CreateAssets()
 		return;
 	}
 
-	if (Context->bTranslateOnlyUsedMaterials && Context->InfoCache)
+	if (Context->bTranslateOnlyUsedMaterials && Context->UsdInfoCache)
 	{
-		if (!Context->InfoCache->IsMaterialUsed(PrimPath))
+		if (!Context->UsdInfoCache->IsMaterialUsed(PrimPath))
 		{
 			UE_LOG(
 				LogUsd,
@@ -378,7 +378,7 @@ void FUsdShadeMaterialTranslator::CreateAssets()
 			);
 			if (!bSuccess)
 			{
-				UsdUnreal::TranslatorUtils::AbandonFailedAsset(MIC, Context->UsdAssetCache.Get(), Context->InfoCache.Get());
+				UsdUnreal::TranslatorUtils::AbandonFailedAsset(MIC, Context->UsdAssetCache.Get(), Context->PrimLinkCache.Get());
 				return;
 			}
 
@@ -486,7 +486,7 @@ void FUsdShadeMaterialTranslator::CreateAssets()
 			);
 			if (!bSuccess)
 			{
-				UsdUnreal::TranslatorUtils::AbandonFailedAsset(MI, Context->UsdAssetCache.Get(), Context->InfoCache.Get());
+				UsdUnreal::TranslatorUtils::AbandonFailedAsset(MI, Context->UsdAssetCache.Get(), Context->PrimLinkCache.Get());
 				return;
 			}
 
@@ -509,9 +509,9 @@ void FUsdShadeMaterialTranslator::CreateAssets()
 
 			// We must stash our material and textures *before* we call UpgradeMaterialsAndTexturesToVT, as that
 			// is what will actually swap our reference with a VT one if needed
-			if (Context->InfoCache)
+			if (Context->PrimLinkCache)
 			{
-				Context->InfoCache->LinkAssetToPrim(PrimPath, MI);
+				Context->PrimLinkCache->LinkAssetToPrim(PrimPath, MI);
 			}
 			for (UTexture* Texture : VTTextures.Union(NonVTTextures))
 			{
@@ -546,7 +546,7 @@ bool FUsdShadeMaterialTranslator::CanBeCollapsed(ECollapsingType CollapsingType)
 
 void FUsdShadeMaterialTranslator::PostImportMaterial(const FString& PrefixedMaterialHash, UMaterialInterface* ImportedMaterial)
 {
-	if (!ImportedMaterial || !Context->InfoCache || !Context->UsdAssetCache)
+	if (!ImportedMaterial || !Context->PrimLinkCache || !Context->UsdAssetCache)
 	{
 		return;
 	}
@@ -574,7 +574,7 @@ void FUsdShadeMaterialTranslator::PostImportMaterial(const FString& PrefixedMate
 
 	// Note that this needs to run even if we found this material in the asset cache already, otherwise we won't
 	// re-register the prim asset links when we reload a stage
-	Context->InfoCache->LinkAssetToPrim(PrimPath, ImportedMaterial);
+	Context->PrimLinkCache->LinkAssetToPrim(PrimPath, ImportedMaterial);
 
 	// Also link the textures to the same material prim. Our textures should all come from USDShadeConversion.cpp or
 	// MaterialX or MDL translators, so they should already be tracked by the same asset cache the material is tracked by.
@@ -591,7 +591,7 @@ void FUsdShadeMaterialTranslator::PostImportMaterial(const FString& PrefixedMate
 			if (bIsTrackedByCache)
 			{
 				Context->UsdAssetCache->TouchAssetPath(Texture);
-				Context->InfoCache->LinkAssetToPrim(PrimPath, Texture);
+				Context->PrimLinkCache->LinkAssetToPrim(PrimPath, Texture);
 
 				if (UUsdAssetUserData* TextureUserData = UsdUnreal::ObjectUtils::GetOrCreateAssetUserData(Texture))
 				{
@@ -606,7 +606,7 @@ TSet<UE::FSdfPath> FUsdShadeMaterialTranslator::CollectAuxiliaryPrims() const
 {
 	if (!Context->bIsBuildingInfoCache)
 	{
-		return Context->InfoCache->GetAuxiliaryPrims(PrimPath);
+		return Context->UsdInfoCache->GetAuxiliaryPrims(PrimPath);
 	}
 
 	TSet<UE::FSdfPath> Result;
diff --git a/Engine/Plugins/Importers/USDImporter/Source/USDSchemas/Private/USDSkelRootTranslator.cpp b/Engine/Plugins/Importers/USDImporter/Source/USDSchemas/Private/USDSkelRootTranslator.cpp
--- a/Engine/Plugins/Importers/USDImporter/Source/USDSchemas/Private/USDSkelRootTranslator.cpp
+++ b/Engine/Plugins/Importers/USDImporter/Source/USDSchemas/Private/USDSkelRootTranslator.cpp
@@ -72,7 +72,7 @@ namespace UsdSkelRootTranslatorImpl
 		TArray<UsdUtils::FUsdPrimMaterialAssignmentInfo>& LODIndexToMaterialInfo,
 		USkeletalMesh* SkeletalMesh,
 		UUsdAssetCache3& AssetCache,
-		FUsdInfoCache& InfoCache,
+		UUsdPrimLinkCache& PrimLinkCache,
 		float Time,
 		EObjectFlags Flags,
 		bool bSkeletalMeshHasMorphTargets,
@@ -104,7 +104,7 @@ namespace UsdSkelRootTranslatorImpl
 			UsdPrim,
 			LODIndexToMaterialInfo,
 			AssetCache,
-			InfoCache,
+			PrimLinkCache,
 			Flags,
 			bShareAssetsForIdenticalPrims
 		);
@@ -660,7 +660,7 @@ namespace UsdSkelRootTranslatorImpl
 		const TArray<UMaterialInterface*>& ExistingAssignments,
 		UMeshComponent& MeshComponent,
 		UUsdAssetCache3& AssetCache,
-		FUsdInfoCache& InfoCache,
+		UUsdPrimLinkCache& PrimLinkCache,
 		float Time,
 		EObjectFlags Flags,
 		bool bInterpretLODs,
@@ -811,7 +811,7 @@ namespace UsdSkelRootTranslatorImpl
 			ValidSkelRootPrim,
 			LODIndexToAssignments,
 			AssetCache,
-			InfoCache,
+			PrimLinkCache,
 			Flags,
 			bShareAssetsForIdenticalPrims
 		);
@@ -922,14 +922,14 @@ namespace UsdSkelRootTranslatorImpl
 		bool* bOutNeedsRecompile = nullptr	  // Whether this function wants the returned AnimBP to be recompiled
 	)
 	{
-		if (!Prim || !Context.InfoCache || Context.UsdAssetCache)
+		if (!Prim || !Context.PrimLinkCache || Context.UsdAssetCache)
 		{
 			return nullptr;
 		}
 
 		FString PrimName = UsdToUnreal::ConvertString(Prim.GetName());
 
-		USkeletalMesh* SkeletalMesh = Context.InfoCache->GetSingleAssetForPrim<USkeletalMesh>(UE::FSdfPath{Prim.GetPath()});
+		USkeletalMesh* SkeletalMesh = Context.PrimLinkCache->GetSingleAssetForPrim<USkeletalMesh>(UE::FSdfPath{Prim.GetPath()});
 		if (!SkeletalMesh)
 		{
 			return nullptr;
@@ -980,7 +980,7 @@ namespace UsdSkelRootTranslatorImpl
 			FSHAHash Hash;
 			FSHA1 SHA1;
 			// Each stage actor has a separate info cache that is assigned to the context
-			SHA1.Update(reinterpret_cast<const uint8*>(Context.InfoCache.Get()), sizeof(Context.InfoCache.Get()));
+			SHA1.Update(reinterpret_cast<const uint8*>(Context.UsdInfoCache.Get()), sizeof(Context.UsdInfoCache.Get()));
 			SHA1.UpdateWithString(*PrimPath, PrimPath.Len());
 			SHA1.Final();
 			SHA1.GetHash(&Hash.Hash[0]);
@@ -1350,7 +1350,7 @@ namespace UsdSkelRootTranslatorImpl
 			ESchemaTranslationLaunchPolicy::Sync,
 			[this]()
 			{
-				if (!Context->InfoCache || !Context->UsdAssetCache)
+				if (!Context->PrimLinkCache || !Context->UsdAssetCache)
 				{
 					return false;
 				}
@@ -1399,7 +1399,7 @@ namespace UsdSkelRootTranslatorImpl
 							LODIndexToMaterialInfo,
 							SkeletalMesh,
 							*Context->UsdAssetCache.Get(),
-							*Context->InfoCache.Get(),
+							*Context->PrimLinkCache.Get(),
 							Context->Time,
 							Context->ObjectFlags,
 							NewBlendShapes.Num() > 0,
@@ -1442,7 +1442,7 @@ namespace UsdSkelRootTranslatorImpl
 
 						if (PhysicsAsset)
 						{
-							Context->InfoCache->LinkAssetToPrim(PrimPath, PhysicsAsset);
+							Context->PrimLinkCache->LinkAssetToPrim(PrimPath, PhysicsAsset);
 						}
 					}
 					else
@@ -1451,7 +1451,7 @@ namespace UsdSkelRootTranslatorImpl
 						SkeletalMesh->SetPhysicsAsset(nullptr);
 					}
 
-					Context->InfoCache->LinkAssetToPrim(PrimPath, SkeletalMesh);
+					Context->PrimLinkCache->LinkAssetToPrim(PrimPath, SkeletalMesh);
 
 					// Track our Skeleton by the source skeleton prim path
 					if (USkeleton* Skeleton = SkeletalMesh->GetSkeleton())
@@ -1462,11 +1462,11 @@ namespace UsdSkelRootTranslatorImpl
 						// should be able to quickly query our skeleton again to find its path
 						std::vector<pxr::UsdSkelBinding> SkeletonBindings;
 						SkeletonCache.Get().ComputeSkelBindings(pxr::UsdSkelRoot(GetPrim()), &SkeletonBindings, pxr::UsdTraverseInstanceProxies());
-						if (SkeletonBindings.size() > 0 && Context->InfoCache)
+						if (SkeletonBindings.size() > 0)
 						{
 							pxr::UsdSkelBinding& SkeletonBinding = SkeletonBindings[0];
 							const pxr::UsdSkelSkeleton& UsdSkeleton = SkeletonBinding.GetSkeleton();
-							Context->InfoCache->LinkAssetToPrim(PrimPath, Skeleton);
+							Context->PrimLinkCache->LinkAssetToPrim(PrimPath, Skeleton);
 						}
 					}
 
@@ -1534,12 +1534,12 @@ namespace UsdSkelRootTranslatorImpl
 			ESchemaTranslationLaunchPolicy::Sync,
 			[this]()
 			{
-				if (!Context->bAllowParsingSkeletalAnimations || !Context->InfoCache)
+				if (!Context->bAllowParsingSkeletalAnimations || !Context->PrimLinkCache)
 				{
 					return false;
 				}
 
-				USkeletalMesh* SkeletalMesh = Context->InfoCache->GetSingleAssetForPrim<USkeletalMesh>(PrimPath);
+				USkeletalMesh* SkeletalMesh = Context->PrimLinkCache->GetSingleAssetForPrim<USkeletalMesh>(PrimPath);
 				if (!SkeletalMesh)
 				{
 					return false;
@@ -1646,9 +1646,9 @@ namespace UsdSkelRootTranslatorImpl
 							}
 						}
 
-						if (AnimSequence && Context->InfoCache)
+						if (AnimSequence && Context->PrimLinkCache)
 						{
-							Context->InfoCache->LinkAssetToPrim(PrimPath, AnimSequence);
+							Context->PrimLinkCache->LinkAssetToPrim(PrimPath, AnimSequence);
 						}
 
 						// For now we shouldn't try to parse the SkelAnimations from skeletal bindings other than the first one as we only
@@ -1699,12 +1699,12 @@ USceneComponent* FUsdSkelRootTranslator::CreateComponents()
 #if WITH_EDITOR
 	// Check if the prim has the GroomBinding schema and setup the component and assets necessary to bind the groom to the SkeletalMesh
 	if (SceneComponent && SceneComponent->IsA<USkeletalMeshComponent>() && UsdUtils::PrimHasSchema(GetPrim(), UnrealIdentifiers::GroomBindingAPI)
-		&& Context->UsdAssetCache && Context->InfoCache && Context->bAllowParsingGroomAssets)
+		&& Context->UsdAssetCache && Context->PrimLinkCache && Context->bAllowParsingGroomAssets)
 	{
 		UsdGroomTranslatorUtils::CreateGroomBindingAsset(
 			GetPrim(),
 			*Context->UsdAssetCache,
-			*Context->InfoCache,
+			*Context->PrimLinkCache,
 			Context->ObjectFlags,
 			Context->bShareAssetsForIdenticalPrims
 		);
@@ -1764,7 +1764,7 @@ void FUsdSkelRootTranslator::UpdateComponents(USceneComponent* SceneComponent)
 	UE::FUsdPrim SkelAnimPrim = UsdUtils::FindFirstAnimationSource(Prim);
 	if (SkelAnimPrim)
 	{
-		UAnimSequence* TargetAnimSequence = Context->InfoCache->GetSingleAssetForPrim<UAnimSequence>(Prim.GetPrimPath());
+		UAnimSequence* TargetAnimSequence = Context->PrimLinkCache->GetSingleAssetForPrim<UAnimSequence>(Prim.GetPrimPath());
 		if (TargetAnimSequence != SkeletalMeshComponent->AnimationData.AnimToPlay)
 		{
 			SkeletalMeshComponent->AnimationData.AnimToPlay = TargetAnimSequence;
@@ -1782,7 +1782,7 @@ void FUsdSkelRootTranslator::UpdateComponents(USceneComponent* SceneComponent)
 
 #if WITH_EDITOR
 	// Re-set the skeletal mesh if we created a new one (maybe the hash changed, a skinned UsdGeomMesh was hidden, etc.)
-	USkeletalMesh* TargetSkeletalMesh = Context->InfoCache->GetSingleAssetForPrim<USkeletalMesh>(PrimPath);
+	USkeletalMesh* TargetSkeletalMesh = Context->PrimLinkCache->GetSingleAssetForPrim<USkeletalMesh>(PrimPath);
 	if (SkeletalMeshComponent->GetSkeletalMeshAsset() != TargetSkeletalMesh)
 	{
 		SkeletalMeshComponent->SetSkeletalMesh(TargetSkeletalMesh);
@@ -1801,7 +1801,7 @@ void FUsdSkelRootTranslator::UpdateComponents(USceneComponent* SceneComponent)
 				ExistingAssignments,
 				*SkeletalMeshComponent,
 				*Context->UsdAssetCache.Get(),
-				*Context->InfoCache.Get(),
+				*Context->PrimLinkCache.Get(),
 				Context->Time,
 				Context->ObjectFlags,
 				Context->bAllowInterpretingLODs,
@@ -1894,9 +1894,9 @@ void FUsdSkelRootTranslator::UpdateComponents(USceneComponent* SceneComponent)
 	}
 
 	// If the prim has a GroomBinding schema, apply the target groom to its associated GroomComponent
-	if (UsdUtils::PrimHasSchema(Prim, UnrealIdentifiers::GroomBindingAPI) && Context->InfoCache)
+	if (UsdUtils::PrimHasSchema(Prim, UnrealIdentifiers::GroomBindingAPI) && Context->PrimLinkCache)
 	{
-		UsdGroomTranslatorUtils::SetGroomFromPrim(Prim, *Context->InfoCache, SceneComponent);
+		UsdGroomTranslatorUtils::SetGroomFromPrim(Prim, *Context->PrimLinkCache, SceneComponent);
 	}
 #endif	  // WITH_EDITOR
 }
@@ -1915,10 +1915,10 @@ TSet<UE::FSdfPath> FUsdSkelRootTranslator::CollectAuxiliaryPrims() const
 {
 	if (!Context->bIsBuildingInfoCache)
 	{
-		return Context->InfoCache->GetAuxiliaryPrims(PrimPath);
+		return Context->UsdInfoCache->GetAuxiliaryPrims(PrimPath);
 	}
 
-	if (!Context->InfoCache->DoesPathCollapseChildren(PrimPath, ECollapsingType::Assets))
+	if (!Context->UsdInfoCache->DoesPathCollapseChildren(PrimPath, ECollapsingType::Assets))
 	{
 		return {};
 	}
diff --git a/Engine/Plugins/Importers/USDImporter/Source/USDSchemas/Private/USDSkelSkeletonTranslator.cpp b/Engine/Plugins/Importers/USDImporter/Source/USDSchemas/Private/USDSkelSkeletonTranslator.cpp
--- a/Engine/Plugins/Importers/USDImporter/Source/USDSchemas/Private/USDSkelSkeletonTranslator.cpp
+++ b/Engine/Plugins/Importers/USDImporter/Source/USDSchemas/Private/USDSkelSkeletonTranslator.cpp
@@ -19,6 +19,7 @@
 #include "USDMemory.h"
 #include "USDObjectUtils.h"
 #include "USDPrimConversion.h"
+#include "USDPrimLinkCache.h"
 #include "USDSkeletalDataConversion.h"
 #include "USDTranslatorUtils.h"
 #include "USDTypesConversion.h"
@@ -82,7 +83,7 @@ namespace UsdSkelSkeletonTranslatorImpl
 		TArray<UsdUtils::FUsdPrimMaterialAssignmentInfo>& LODIndexToMaterialInfo,
 		USkeletalMesh* SkeletalMesh,
 		UUsdAssetCache3& AssetCache,
-		FUsdInfoCache& InfoCache,
+		UUsdPrimLinkCache& PrimLinkCache,
 		float Time,
 		EObjectFlags Flags,
 		bool bSkeletalMeshHasMorphTargets,
@@ -114,7 +115,7 @@ namespace UsdSkelSkeletonTranslatorImpl
 			UsdPrim,
 			LODIndexToMaterialInfo,
 			AssetCache,
-			InfoCache,
+			PrimLinkCache,
 			Flags,
 			bShareAssetsForIdenticalPrims
 		);
@@ -705,7 +706,7 @@ namespace UsdSkelSkeletonTranslatorImpl
 		const TArray<UMaterialInterface*>& ExistingAssignments,
 		UMeshComponent& MeshComponent,
 		UUsdAssetCache3& AssetCache,
-		FUsdInfoCache& InfoCache,
+		UUsdPrimLinkCache& PrimLinkCache,
 		float Time,
 		EObjectFlags Flags,
 		bool bInterpretLODs,
@@ -863,7 +864,7 @@ namespace UsdSkelSkeletonTranslatorImpl
 			ValidSkeletonPrim,
 			LODIndexToAssignments,
 			AssetCache,
-			InfoCache,
+			PrimLinkCache,
 			Flags,
 			bShareAssetsForIdenticalPrims
 		);
@@ -932,14 +933,14 @@ namespace UsdSkelSkeletonTranslatorImpl
 		bool* bOutNeedsRecompile = nullptr	  // Whether this function wants the returned AnimBP to be recompiled
 	)
 	{
-		if (!SkeletonPrim || !Context.InfoCache || !Context.UsdAssetCache)
+		if (!SkeletonPrim || !Context.PrimLinkCache || !Context.UsdAssetCache)
 		{
 			return nullptr;
 		}
 
 		FString PrimName = UsdToUnreal::ConvertString(SkeletonPrim.GetName());
 
-		USkeletalMesh* SkeletalMesh = Context.InfoCache->GetSingleAssetForPrim<USkeletalMesh>(UE::FSdfPath{SkeletonPrim.GetPath()});
+		USkeletalMesh* SkeletalMesh = Context.PrimLinkCache->GetSingleAssetForPrim<USkeletalMesh>(UE::FSdfPath{SkeletonPrim.GetPath()});
 		if (!SkeletalMesh)
 		{
 			return nullptr;
@@ -1002,7 +1003,7 @@ namespace UsdSkelSkeletonTranslatorImpl
 			FSHAHash Hash;
 			FSHA1 SHA1;
 			// Each stage actor has a separate info cache that is assigned to the context
-			SHA1.Update(reinterpret_cast<const uint8*>(Context.InfoCache.Get()), sizeof(Context.InfoCache.Get()));
+			SHA1.Update(reinterpret_cast<const uint8*>(Context.UsdInfoCache.Get()), sizeof(Context.UsdInfoCache.Get()));
 			SHA1.UpdateWithString(*PrimPath, PrimPath.Len());
 			SHA1.Final();
 			SHA1.GetHash(&Hash.Hash[0]);
@@ -1459,7 +1460,7 @@ namespace UsdSkelSkeletonTranslatorImpl
 			ESchemaTranslationLaunchPolicy::Sync,
 			[this]()
 			{
-				if (!Context->InfoCache || !Context->UsdAssetCache)
+				if (!Context->PrimLinkCache || !Context->UsdAssetCache)
 				{
 					return false;
 				}
@@ -1528,7 +1529,7 @@ namespace UsdSkelSkeletonTranslatorImpl
 							TEXT("Failed to create SkeletalMesh for prim '%s'"),
 							*UsdToUnreal::ConvertPath(ClosestParentSkelRoot.Get().GetPrim().GetPrimPath())
 						);
-						UsdUnreal::TranslatorUtils::AbandonFailedAsset(SkeletalMesh, Context->UsdAssetCache.Get(), Context->InfoCache.Get());
+						UsdUnreal::TranslatorUtils::AbandonFailedAsset(SkeletalMesh, Context->UsdAssetCache.Get(), Context->PrimLinkCache.Get());
 
 						if (bSkeletonIsNew)
 						{
@@ -1539,7 +1540,7 @@ namespace UsdSkelSkeletonTranslatorImpl
 								*SkeletonName.ToString(),
 								*UsdToUnreal::ConvertPath(ClosestParentSkelRoot.Get().GetPrim().GetPrimPath())
 							);
-							UsdUnreal::TranslatorUtils::AbandonFailedAsset(Skeleton, Context->UsdAssetCache.Get(), Context->InfoCache.Get());
+							UsdUnreal::TranslatorUtils::AbandonFailedAsset(Skeleton, Context->UsdAssetCache.Get(), Context->PrimLinkCache.Get());
 						}
 					}
 				}
@@ -1609,7 +1610,7 @@ namespace UsdSkelSkeletonTranslatorImpl
 							LODIndexToMaterialInfo,
 							SkeletalMesh,
 							*Context->UsdAssetCache.Get(),
-							*Context->InfoCache.Get(),
+							*Context->PrimLinkCache.Get(),
 							Context->Time,
 							Context->ObjectFlags,
 							NewBlendShapes.Num() > 0,
@@ -1666,7 +1667,7 @@ namespace UsdSkelSkeletonTranslatorImpl
 									UsdUnreal::TranslatorUtils::AbandonFailedAsset(
 										PhysicsAsset,
 										Context->UsdAssetCache.Get(),
-										Context->InfoCache.Get()
+										Context->PrimLinkCache.Get()
 									);
 								}
 							}
@@ -1675,7 +1676,7 @@ namespace UsdSkelSkeletonTranslatorImpl
 
 						if (PhysicsAsset)
 						{
-							Context->InfoCache->LinkAssetToPrim(SkeletonPrimPath, PhysicsAsset);
+							Context->PrimLinkCache->LinkAssetToPrim(SkeletonPrimPath, PhysicsAsset);
 							Context->UsdAssetCache->TouchAssetPath(PhysicsAsset);
 						}
 					}
@@ -1685,12 +1686,12 @@ namespace UsdSkelSkeletonTranslatorImpl
 						SkeletalMesh->SetPhysicsAsset(nullptr);
 					}
 
-					Context->InfoCache->LinkAssetToPrim(SkeletonPrimPath, SkeletalMesh);
+					Context->PrimLinkCache->LinkAssetToPrim(SkeletonPrimPath, SkeletalMesh);
 
 					// Track our Skeleton by the source skeleton prim path
 					if (USkeleton* Skeleton = SkeletalMesh->GetSkeleton())
 					{
-						Context->InfoCache->LinkAssetToPrim(SkeletonPrimPath, Skeleton);
+						Context->PrimLinkCache->LinkAssetToPrim(SkeletonPrimPath, Skeleton);
 						Context->UsdAssetCache->TouchAssetPath(Skeleton);
 					}
 
@@ -1764,12 +1765,12 @@ namespace UsdSkelSkeletonTranslatorImpl
 			ESchemaTranslationLaunchPolicy::Sync,
 			[this]()
 			{
-				if (!Context->bAllowParsingSkeletalAnimations || !Context->InfoCache)
+				if (!Context->bAllowParsingSkeletalAnimations || !Context->PrimLinkCache)
 				{
 					return false;
 				}
 
-				USkeletalMesh* SkeletalMesh = Context->InfoCache->GetSingleAssetForPrim<USkeletalMesh>(SkeletonPrimPath);
+				USkeletalMesh* SkeletalMesh = Context->PrimLinkCache->GetSingleAssetForPrim<USkeletalMesh>(SkeletonPrimPath);
 				if (!SkeletalMesh)
 				{
 					return false;
@@ -1863,7 +1864,7 @@ namespace UsdSkelSkeletonTranslatorImpl
 						|| (AnimSequence->GetDataModel()->GetNumBoneTracks() == 0 && AnimSequence->GetDataModel()->GetNumberOfFloatCurves() == 0))
 					{
 						UE_LOG(LogUsd, Warning, TEXT("Failed to create AnimSequence for prim '%s'"), *SkelAnimationPrimPath);
-						UsdUnreal::TranslatorUtils::AbandonFailedAsset(AnimSequence, Context->UsdAssetCache.Get(), Context->InfoCache.Get());
+						UsdUnreal::TranslatorUtils::AbandonFailedAsset(AnimSequence, Context->UsdAssetCache.Get(), Context->PrimLinkCache.Get());
 						AnimSequence = nullptr;
 					}
 				}
@@ -1908,9 +1909,9 @@ namespace UsdSkelSkeletonTranslatorImpl
 					}
 				}
 
-				if (AnimSequence && Context->InfoCache)
+				if (AnimSequence && Context->PrimLinkCache)
 				{
-					Context->InfoCache->LinkAssetToPrim(SkeletonPrimPath, AnimSequence);
+					Context->PrimLinkCache->LinkAssetToPrim(SkeletonPrimPath, AnimSequence);
 				}
 
 				return true;
@@ -1960,7 +1961,7 @@ USceneComponent* FUsdSkelSkeletonTranslator::CreateComponents()
 	}
 
 	// Check if the prim has the GroomBinding schema and setup the component and assets necessary to bind the groom to the SkeletalMesh
-	if (Context->bAllowParsingGroomAssets && Context->UsdAssetCache && Context->InfoCache)
+	if (Context->bAllowParsingGroomAssets && Context->UsdAssetCache && Context->PrimLinkCache)
 	{
 		UE::FUsdPrim PrimWithSchema;
 		if (UsdUtils::PrimHasSchema(SkeletonPrim, UnrealIdentifiers::GroomBindingAPI))
@@ -1978,7 +1979,7 @@ USceneComponent* FUsdSkelSkeletonTranslator::CreateComponents()
 			UsdGroomTranslatorUtils::CreateGroomBindingAsset(
 				PrimWithSchema,
 				*Context->UsdAssetCache,
-				*Context->InfoCache,
+				*Context->PrimLinkCache,
 				Context->ObjectFlags,
 				Context->bShareAssetsForIdenticalPrims
 			);
@@ -2004,7 +2005,7 @@ USceneComponent* FUsdSkelSkeletonTranslator::CreateComponents()
 void FUsdSkelSkeletonTranslator::UpdateComponents(USceneComponent* SceneComponent)
 {
 	USkeletalMeshComponent* SkeletalMeshComponent = Cast<USkeletalMeshComponent>(SceneComponent);
-	if (!SkeletalMeshComponent || !Context->InfoCache)
+	if (!SkeletalMeshComponent || !Context->PrimLinkCache)
 	{
 		return;
 	}
@@ -2085,7 +2086,7 @@ void FUsdSkelSkeletonTranslator::UpdateComponents(USceneComponent* SceneComponen
 
 	if (SkelAnimPrim)
 	{
-		UAnimSequence* TargetAnimSequence = Context->InfoCache->GetSingleAssetForPrim<UAnimSequence>(PrimPath);
+		UAnimSequence* TargetAnimSequence = Context->PrimLinkCache->GetSingleAssetForPrim<UAnimSequence>(PrimPath);
 		if (TargetAnimSequence != SkeletalMeshComponent->AnimationData.AnimToPlay)
 		{
 			SkeletalMeshComponent->AnimationData.AnimToPlay = TargetAnimSequence;
@@ -2103,7 +2104,7 @@ void FUsdSkelSkeletonTranslator::UpdateComponents(USceneComponent* SceneComponen
 
 #if WITH_EDITOR
 	// Re-set the skeletal mesh if we created a new one (maybe the hash changed, a skinned UsdGeomMesh was hidden, etc.)
-	USkeletalMesh* TargetSkeletalMesh = Context->InfoCache->GetSingleAssetForPrim<USkeletalMesh>(PrimPath);
+	USkeletalMesh* TargetSkeletalMesh = Context->PrimLinkCache->GetSingleAssetForPrim<USkeletalMesh>(PrimPath);
 	if (SkeletalMeshComponent->GetSkeletalMeshAsset() != TargetSkeletalMesh)
 	{
 		SkeletalMeshComponent->SetSkeletalMesh(TargetSkeletalMesh);
@@ -2121,8 +2122,8 @@ void FUsdSkelSkeletonTranslator::UpdateComponents(USceneComponent* SceneComponen
 				SkeletonBindingForPrim.Get(),
 				ExistingAssignments,
 				*SkeletalMeshComponent,
-				*Context->UsdAssetCache.Get(),
-				*Context->InfoCache.Get(),
+				*Context->UsdAssetCache,
+				*Context->PrimLinkCache,
 				Context->Time,
 				Context->ObjectFlags,
 				Context->bAllowInterpretingLODs,
@@ -2212,11 +2213,10 @@ void FUsdSkelSkeletonTranslator::UpdateComponents(USceneComponent* SceneComponen
 	}
 
 	// If the prim has a GroomBinding schema, apply the target groom to its associated GroomComponent
-	if (Context->InfoCache)
 	{
 		if (UsdUtils::PrimHasSchema(SkeletonPrim, UnrealIdentifiers::GroomBindingAPI))
 		{
-			UsdGroomTranslatorUtils::SetGroomFromPrim(SkeletonPrim, *Context->InfoCache, SceneComponent);
+			UsdGroomTranslatorUtils::SetGroomFromPrim(SkeletonPrim, *Context->PrimLinkCache, SceneComponent);
 		}
 		else if (UsdUtils::PrimHasSchema(ClosestSkelRootPrim, UnrealIdentifiers::GroomBindingAPI))
 		{
@@ -2230,7 +2230,7 @@ void FUsdSkelSkeletonTranslator::UpdateComponents(USceneComponent* SceneComponen
 					 "prims instead!"),
 				*ClosestSkelRootPrim.GetPrimPath().GetString()
 			);
-			UsdGroomTranslatorUtils::SetGroomFromPrim(ClosestSkelRootPrim, *Context->InfoCache, SceneComponent);
+			UsdGroomTranslatorUtils::SetGroomFromPrim(ClosestSkelRootPrim, *Context->PrimLinkCache, SceneComponent);
 		}
 	}
 #endif	  // WITH_EDITOR
@@ -2256,7 +2256,7 @@ TSet<UE::FSdfPath> FUsdSkelSkeletonTranslator::CollectAuxiliaryPrims() const
 {
 	if (!Context->bIsBuildingInfoCache)
 	{
-		return Context->InfoCache->GetAuxiliaryPrims(PrimPath);
+		return Context->UsdInfoCache->GetAuxiliaryPrims(PrimPath);
 	}
 
 	TSet<UE::FSdfPath> Result;
diff --git a/Engine/Plugins/Importers/USDImporter/Source/USDSchemas/Private/USDTranslatorUtils.cpp b/Engine/Plugins/Importers/USDImporter/Source/USDSchemas/Private/USDTranslatorUtils.cpp
--- a/Engine/Plugins/Importers/USDImporter/Source/USDSchemas/Private/USDTranslatorUtils.cpp
+++ b/Engine/Plugins/Importers/USDImporter/Source/USDSchemas/Private/USDTranslatorUtils.cpp
@@ -3,11 +3,11 @@
 #include "USDTranslatorUtils.h"
 
 #include "USDAssetCache3.h"
-#include "USDInfoCache.h"
+#include "USDPrimLinkCache.h"
 
 #include "AssetRegistry/AssetRegistryModule.h"
 
-void UsdUnreal::TranslatorUtils::AbandonFailedAsset(UObject* Asset, UUsdAssetCache3* AssetCache, FUsdInfoCache* InfoCache)
+void UsdUnreal::TranslatorUtils::AbandonFailedAsset(UObject* Asset, UUsdAssetCache3* AssetCache, UUsdPrimLinkCache* PrimLinkCache)
 {
 	if (!Asset)
 	{
@@ -32,8 +32,8 @@ void UsdUnreal::TranslatorUtils::AbandonFailedAsset(UObject* Asset, UUsdAssetCac
 		}
 	}
 
-	if (InfoCache)
+	if (PrimLinkCache)
 	{
-		InfoCache->RemoveAllAssetPrimLinks(Asset);
+		PrimLinkCache->RemoveAllAssetPrimLinks(Asset);
 	}
 }
diff --git a/Engine/Plugins/Importers/USDImporter/Source/USDSchemas/Private/USDTranslatorUtils.h b/Engine/Plugins/Importers/USDImporter/Source/USDSchemas/Private/USDTranslatorUtils.h
--- a/Engine/Plugins/Importers/USDImporter/Source/USDSchemas/Private/USDTranslatorUtils.h
+++ b/Engine/Plugins/Importers/USDImporter/Source/USDSchemas/Private/USDTranslatorUtils.h
@@ -6,10 +6,10 @@
 #include "UObject/ObjectMacros.h"
 
 class UUsdAssetCache3;
-class FUsdInfoCache;
+class UUsdPrimLinkCache;
 
 namespace UsdUnreal::TranslatorUtils
 {
 	// Properly deletes the asset and removes it from the asset and info caches, if provided
-	void AbandonFailedAsset(UObject* Asset, UUsdAssetCache3* AssetCache, FUsdInfoCache* InfoCache);
+	void AbandonFailedAsset(UObject* Asset, UUsdAssetCache3* AssetCache, UUsdPrimLinkCache* PrimLinkCache);
 }
diff --git a/Engine/Plugins/Importers/USDImporter/Source/USDSchemas/Private/USDVolVolumeTranslator.cpp b/Engine/Plugins/Importers/USDImporter/Source/USDSchemas/Private/USDVolVolumeTranslator.cpp
--- a/Engine/Plugins/Importers/USDImporter/Source/USDSchemas/Private/USDVolVolumeTranslator.cpp
+++ b/Engine/Plugins/Importers/USDImporter/Source/USDSchemas/Private/USDVolVolumeTranslator.cpp
@@ -933,7 +933,7 @@ void FUsdVolVolumeTranslator::CreateAssets()
 
 	using namespace UE::UsdVolVolumeTranslator::Private;
 
-	if (!Context->UsdAssetCache || !Context->InfoCache)
+	if (!Context->UsdAssetCache || !Context->PrimLinkCache)
 	{
 		return;
 	}
@@ -1070,7 +1070,7 @@ void FUsdVolVolumeTranslator::CreateAssets()
 
 		if (SparseVolumeTexture)
 		{
-			Context->InfoCache->LinkAssetToPrim(PrimPath, SparseVolumeTexture);
+			Context->PrimLinkCache->LinkAssetToPrim(PrimPath, SparseVolumeTexture);
 
 			if (UUsdSparseVolumeTextureAssetUserData* UserData = UsdUnreal::ObjectUtils::GetOrCreateAssetUserData<
 					UUsdSparseVolumeTextureAssetUserData>(SparseVolumeTexture))
@@ -1259,7 +1259,7 @@ void FUsdVolVolumeTranslator::CreateAssets()
 
 	if (MaterialInstance)
 	{
-		Context->InfoCache->LinkAssetToPrim(PrimPath, MaterialInstance);
+		Context->PrimLinkCache->LinkAssetToPrim(PrimPath, MaterialInstance);
 
 		if (UUsdAssetUserData* UserData = UsdUnreal::ObjectUtils::GetOrCreateAssetUserData<UUsdAssetUserData>(MaterialInstance))
 		{
@@ -1320,7 +1320,7 @@ void FUsdVolVolumeTranslator::UpdateComponents(USceneComponent* SceneComponent)
 		const int32 ElementIndex = 0;
 		UMaterialInterface* CurrentMaterial = VolumeComponent->GetMaterial(ElementIndex);
 
-		if (UMaterialInstance* MaterialForPrim = Context->InfoCache->GetSingleAssetForPrim<UMaterialInstance>(PrimPath))
+		if (UMaterialInstance* MaterialForPrim = Context->PrimLinkCache->GetSingleAssetForPrim<UMaterialInstance>(PrimPath))
 		{
 			if (MaterialForPrim != CurrentMaterial)
 			{
@@ -1431,7 +1431,7 @@ TSet<UE::FSdfPath> FUsdVolVolumeTranslator::CollectAuxiliaryPrims() const
 {
 	if (!Context->bIsBuildingInfoCache)
 	{
-		return Context->InfoCache->GetAuxiliaryPrims(PrimPath);
+		return Context->UsdInfoCache->GetAuxiliaryPrims(PrimPath);
 	}
 
 	TSet<UE::FSdfPath> Result;
diff --git a/Engine/Plugins/Importers/USDImporter/Source/USDSchemas/Public/USDInfoCache.h b/Engine/Plugins/Importers/USDImporter/Source/USDSchemas/Public/USDInfoCache.h
--- a/Engine/Plugins/Importers/USDImporter/Source/USDSchemas/Public/USDInfoCache.h
+++ b/Engine/Plugins/Importers/USDImporter/Source/USDSchemas/Public/USDInfoCache.h
@@ -4,10 +4,13 @@
 
 #include "Containers/Set.h"
 #include "CoreMinimal.h"
+#include "Templates/PimplPtr.h"
 #include "UObject/WeakObjectPtr.h"
 
 #include "UsdWrappers/SdfPath.h"
 
+#include "USDInfoCache.generated.h"
+
 struct FUsdSchemaTranslationContext;
 namespace UE
 {
@@ -26,14 +29,14 @@ namespace UsdUtils
 	struct FUsdPrimMaterialSlot;
 }
 
+struct FUsdInfoCacheImpl;
+
 /**
  * Caches information about a specific USD Stage
  */
-class USDSCHEMAS_API FUsdInfoCache
+class USDSCHEMAS_API UE_DEPRECATED(5.5, "Use UUsdInfoCache or UUsdPrimLinkCache instead") FUsdInfoCache
 {
 public:
-	struct FUsdInfoCacheImpl;
-
 	FUsdInfoCache();
 	FUsdInfoCache(const FUsdInfoCache& Other);
 	virtual ~FUsdInfoCache();
@@ -160,3 +163,92 @@ private:
 private:
 	TUniquePtr<FUsdInfoCacheImpl> Impl;
 };
+
+/**
+ * Caches information about a specific USD Stage
+ */
+UCLASS()
+class USDSCHEMAS_API UUsdInfoCache : public UObject
+{
+	GENERATED_BODY()
+
+public:
+	UUsdInfoCache();
+
+	// Copies the internal storage from 'Other'
+	void CopyImpl(const UUsdInfoCache& Other);
+
+	// Begin UObject interface
+	virtual void Serialize(FArchive& Ar) override;
+	// End UObject interface
+
+	// Returns whether we contain any info about prim at 'Path' at all
+	bool ContainsInfoAboutPrim(const UE::FSdfPath& Path) const;
+
+	// Returns a list of all prims we have generic info about
+	TSet<UE::FSdfPath> GetKnownPrims() const;
+
+	void RebuildCacheForSubtree(const UE::FUsdPrim& Prim, FUsdSchemaTranslationContext& Context);
+
+	void Clear();
+	bool IsEmpty();
+
+public:
+	bool IsPathCollapsed(const UE::FSdfPath& Path, ECollapsingType CollapsingType) const;
+	bool DoesPathCollapseChildren(const UE::FSdfPath& Path, ECollapsingType CollapsingType) const;
+
+	// Returns Path in case it represents an uncollapsed prim, or returns the path to the prim that collapsed it
+	UE::FSdfPath UnwindToNonCollapsedPath(const UE::FSdfPath& Path, ECollapsingType CollapsingType) const;
+
+public:
+	// Returns the paths to prims that, when translated into assets or components, also require reading the prim at
+	// 'Path'. e.g. providing the path to a Shader prim will return the paths to all Material prims for which the
+	// translation involves reading that particular Shader.
+	TSet<UE::FSdfPath> GetMainPrims(const UE::FSdfPath& AuxPrimPath) const;
+
+	// The inverse of the function above: Provide it with the path to a Material prim and it will return the set of
+	// paths to all Shader prims that need to be read to translate that Material prim into material assets
+	TSet<UE::FSdfPath> GetAuxiliaryPrims(const UE::FSdfPath& MainPrimPath) const;
+
+public:
+	// Returns the set of paths to all prims that have a material:binding relationship to the particular material at
+	// 'Path', if any.
+	// Returns a copy for thread safety.
+	TSet<UE::FSdfPath> GetMaterialUsers(const UE::FSdfPath& Path) const;
+	bool IsMaterialUsed(const UE::FSdfPath& Path) const;
+
+public:
+	// Provides the total vertex or material slots counts for each prim *and* its subtree.
+	// This is built inside RebuildCacheForSubtree, so it will factor in the used Context's bMergeIdenticalMaterialSlots.
+	// Note that these aren't affected by actual collapsing: A prim that doesn't collapse its children will still
+	// provide the total sum of vertex counts of its entire subtree when queried
+	TOptional<uint64> GetSubtreeVertexCount(const UE::FSdfPath& Path);
+	TOptional<uint64> GetSubtreeMaterialSlotCount(const UE::FSdfPath& Path);
+	TOptional<TArray<UsdUtils::FUsdPrimMaterialSlot>> GetSubtreeMaterialSlots(const UE::FSdfPath& Path);
+
+	// Returns true if Path could potentially be collapsed as a Geometry Cache asset
+	bool IsPotentialGeometryCacheRoot(const UE::FSdfPath& Path) const;
+
+public:
+	// Marks/checks if the provided path to a prototype prim is already being translated.
+	// This is used during scene translation with instanceables, so that the schema translators can early out
+	// in case they have been created to translate multiple instances of the same prototype
+	void ResetTranslatedPrototypes();
+	bool IsPrototypeTranslated(const UE::FSdfPath& PrototypePath);
+	void MarkPrototypeAsTranslated(const UE::FSdfPath& PrototypePath);
+
+private:
+	friend class FUsdGeomXformableTranslator;
+	// Returns true if every prim on the subtree below RootPath (including the RootPath prim itself) returns true for
+	// CanBeCollapsed(), according to their own schema translators.
+	//
+	// WARNING: This is intended for internal use, and exclusively during the actual info cache build process as it will
+	// need to query the prim/stage directly. Calling it after the info cache build may yield back an empty optional,
+	// meaning it is unknown at this point whether the prim CanBeCollapsed or not.
+	//
+	// In general, you shouldn't call this, but just use "IsPathCollapsed" or "DoesPathCollapseChildren" instead.
+	TOptional<bool> CanXformableSubtreeBeCollapsed(const UE::FSdfPath& RootPath, FUsdSchemaTranslationContext& Context);
+
+private:
+	TPimplPtr<FUsdInfoCacheImpl> Impl;
+};
diff --git a/Engine/Plugins/Importers/USDImporter/Source/USDSchemas/Public/USDPrimLinkCache.h b/Engine/Plugins/Importers/USDImporter/Source/USDSchemas/Public/USDPrimLinkCache.h
new file mode 100644
--- /dev/null
+++ b/Engine/Plugins/Importers/USDImporter/Source/USDSchemas/Public/USDPrimLinkCache.h
@@ -0,0 +1,95 @@
+// Copyright Epic Games, Inc. All Rights Reserved.
+
+#pragma once
+
+#include "Containers/Set.h"
+#include "CoreMinimal.h"
+#include "UObject/WeakObjectPtr.h"
+
+#include "UsdWrappers/SdfPath.h"
+
+#include "USDPrimLinkCache.generated.h"
+
+namespace UE
+{
+	class FSdfPath;
+}
+
+UCLASS()
+class USDSCHEMAS_API UUsdPrimLinkCache : public UObject
+{
+	GENERATED_BODY()
+
+public:
+	struct UUsdPrimLinkCacheImpl;
+
+	UUsdPrimLinkCache();
+	UUsdPrimLinkCache(FVTableHelper& Helper);
+	virtual ~UUsdPrimLinkCache() override;
+
+	// Begin UObject interface
+	virtual void Serialize(FArchive& Ar) override;
+	// End UObject interface
+
+	// Returns whether we contain any info about prim at 'Path' at all
+	bool ContainsInfoAboutPrim(const UE::FSdfPath& Path) const;
+
+	// Returns a list of all prims we have generic info about
+	TSet<UE::FSdfPath> GetKnownPrims() const;
+
+	void Clear();
+	bool IsEmpty();
+
+public:
+	void LinkAssetToPrim(const UE::FSdfPath& Path, UObject* Asset);
+	void UnlinkAssetFromPrim(const UE::FSdfPath& Path, UObject* Asset);
+
+	TArray<TWeakObjectPtr<UObject>> RemoveAllAssetPrimLinks(const UE::FSdfPath& Path);
+	TArray<UE::FSdfPath> RemoveAllAssetPrimLinks(const UObject* Asset);
+	void RemoveAllAssetPrimLinks();
+
+	TArray<TWeakObjectPtr<UObject>> GetAllAssetsForPrim(const UE::FSdfPath& Path) const;
+
+	template<typename T = UObject>
+	T* GetSingleAssetForPrim(const UE::FSdfPath& Path) const
+	{
+		TArray<TWeakObjectPtr<UObject>> Assets = GetAllAssetsForPrim(Path);
+
+		// Search back to front so that if we generate a new version of an asset type we prefer
+		// returning that
+		for (int32 Index = Assets.Num() - 1; Index >= 0; --Index)
+		{
+			if (T* CastAsset = Cast<T>(Assets[Index].Get()))
+			{
+				return CastAsset;
+			}
+		}
+
+		return nullptr;
+	}
+
+	template<typename T>
+	TArray<T*> GetAssetsForPrim(const UE::FSdfPath& Path) const
+	{
+		TArray<TWeakObjectPtr<UObject>> Assets = GetAllAssetsForPrim(Path);
+
+		TArray<T*> CastAssets;
+		CastAssets.Reserve(Assets.Num());
+
+		for (const TWeakObjectPtr<UObject>& Asset : Assets)
+		{
+			if (T* CastAsset = Cast<T>(Asset.Get()))
+			{
+				CastAssets.Add(CastAsset);
+			}
+		}
+
+		return CastAssets;
+	}
+
+	TArray<UE::FSdfPath> GetPrimsForAsset(const UObject* Asset) const;
+	TMap<UE::FSdfPath, TArray<TWeakObjectPtr<UObject>>> GetAllAssetPrimLinks() const;
+
+private:
+	TUniquePtr<UUsdPrimLinkCacheImpl> Impl;
+};
diff --git a/Engine/Plugins/Importers/USDImporter/Source/USDSchemas/Public/USDSchemaTranslator.h b/Engine/Plugins/Importers/USDImporter/Source/USDSchemas/Public/USDSchemaTranslator.h
--- a/Engine/Plugins/Importers/USDImporter/Source/USDSchemas/Public/USDSchemaTranslator.h
+++ b/Engine/Plugins/Importers/USDImporter/Source/USDSchemas/Public/USDSchemaTranslator.h
@@ -8,8 +8,10 @@
 #include "USDAssetCache2.h"
 #include "USDAssetCache3.h"
 #include "USDInfoCache.h"
+#include "USDInfoCache.h"
 #include "USDMemory.h"
 #include "USDMetadataImportOptions.h"
+#include "USDPrimLinkCache.h"
 #include "USDSkeletalDataConversion.h"
 
 #include "UsdWrappers/SdfPath.h"
@@ -234,12 +236,21 @@ struct USDSCHEMAS_API FUsdSchemaTranslationContext : public TSharedFromThis<FUsd
 	/** Where the translated assets will be stored */
 	TStrongObjectPtr<UUsdAssetCache3> UsdAssetCache;
 
+	PRAGMA_DISABLE_DEPRECATION_WARNINGS
 	/** Where the translated assets will be stored */
 	UE_DEPRECATED(5.5, "Use the 'UsdAssetCache' member instead, which is of the new UUsdAssetCache3 type")
 	TStrongObjectPtr<UUsdAssetCache2> AssetCache;
 
 	/** Caches various information about prims that are expensive to query */
+	UE_DEPRECATED(5.5, "Use the 'UsdInfoCache' member instead, which is of the new UUsdInfoCache type")
 	TSharedPtr<FUsdInfoCache> InfoCache;
+	PRAGMA_ENABLE_DEPRECATION_WARNINGS
+
+	/** Caches various information about prims that are expensive to query */
+	TStrongObjectPtr<UUsdInfoCache> UsdInfoCache;
+
+	/** Caches which assets were generated from which USD prims */
+	TStrongObjectPtr<UUsdPrimLinkCache> PrimLinkCache;
 
 	/** Bounding box cache used for the USD stage in case we have to spawn bounds components */
 	TSharedPtr<UE::FUsdGeomBBoxCache> BBoxCache;
diff --git a/Engine/Plugins/Importers/USDImporter/Source/USDStage/Private/USDLevelSequenceHelper.cpp b/Engine/Plugins/Importers/USDImporter/Source/USDStage/Private/USDLevelSequenceHelper.cpp
--- a/Engine/Plugins/Importers/USDImporter/Source/USDStage/Private/USDLevelSequenceHelper.cpp
+++ b/Engine/Plugins/Importers/USDImporter/Source/USDStage/Private/USDLevelSequenceHelper.cpp
@@ -6,13 +6,13 @@
 #include "USDAttributeUtils.h"
 #include "USDConversionUtils.h"
 #include "USDDrawModeComponent.h"
-#include "USDInfoCache.h"
 #include "USDIntegrationUtils.h"
 #include "USDLayerUtils.h"
 #include "USDListener.h"
 #include "USDLog.h"
 #include "USDObjectUtils.h"
 #include "USDPrimConversion.h"
+#include "USDPrimLinkCache.h"
 #include "USDPrimTwin.h"
 #include "USDProjectSettings.h"
 #include "USDSkeletalDataConversion.h"
@@ -718,7 +718,7 @@ public:
 
 	ULevelSequence* Init(const UE::FUsdStage& InUsdStage);
 	bool Serialize(FArchive& Ar);
-	void SetInfoCache(TSharedPtr<FUsdInfoCache> InfoCache);
+	void SetPrimLinkCache(UUsdPrimLinkCache* PrimLinkCache);
 	void SetBBoxCache(TSharedPtr<UE::FUsdGeomBBoxCache> InBBoxCache);
 	bool HasData() const;
 	void Clear();
@@ -948,9 +948,14 @@ private:
 	FUsdLevelSequenceHelper::FOnSkelAnimationBaked OnSkelAnimationBaked;
 
 	TWeakObjectPtr<AUsdStageActor> StageActor = nullptr;
-	TSharedPtr<FUsdInfoCache> InfoCache = nullptr;			  // We keep a pointer to this directly because we may be called via the USDStageImporter
-															  // directly, when we don't have an available actor
-	TSharedPtr<UE::FUsdGeomBBoxCache> BBoxCache = nullptr;	  // Same as for the info cache
+
+	// We keep a pointer to these directly because we may be called via the
+	// USDStageImporter directly, when we don't have an available actor.
+	// This has to be weak or else we get a circular reference, as this will hold on the PrimLinkCache,
+	// that has an Outer reference to the stage actor, that owns this
+	TWeakObjectPtr<UUsdPrimLinkCache> PrimLinkCache = nullptr;
+	TSharedPtr<UE::FUsdGeomBBoxCache> BBoxCache = nullptr;
+
 	EUsdRootMotionHandling RootMotionHandling = EUsdRootMotionHandling::NoAdditionalRootMotion;
 	FGuid StageActorBinding;
 
@@ -1019,6 +1024,8 @@ ULevelSequence* FUsdLevelSequenceHelperImpl::Init(const UE::FUsdStage& InUsdStag
 
 bool FUsdLevelSequenceHelperImpl::Serialize(FArchive& Ar)
 {
+	TRACE_CPUPROFILER_EVENT_SCOPE(FUsdLevelSequenceHelperImpl::Serialize);
+
 	Ar << MainLevelSequence;
 	Ar << LevelSequencesByIdentifier;
 	Ar << IdentifierByLevelSequence;
@@ -1039,9 +1046,10 @@ bool FUsdLevelSequenceHelperImpl::Serialize(FArchive& Ar)
 	return true;
 }
 
-void FUsdLevelSequenceHelperImpl::SetInfoCache(TSharedPtr<FUsdInfoCache> InInfoCache)
+void FUsdLevelSequenceHelperImpl::SetPrimLinkCache(UUsdPrimLinkCache* InPrimLinkCache)
 {
-	InfoCache = InInfoCache;
+	// PrimLinkCache.Reset(InPrimLinkCache);
+	PrimLinkCache = InPrimLinkCache;
 }
 
 void FUsdLevelSequenceHelperImpl::SetBBoxCache(TSharedPtr<UE::FUsdGeomBBoxCache> InBBoxCache)
@@ -1195,7 +1203,7 @@ void FUsdLevelSequenceHelperImpl::BindToUsdStageActor(AUsdStageActor* InStageAct
 	UnbindFromUsdStageActor();
 
 	StageActor = InStageActor;
-	SetInfoCache(InStageActor ? InStageActor->GetInfoCache() : nullptr);
+	SetPrimLinkCache(InStageActor ? InStageActor->PrimLinkCache : nullptr);
 	SetBBoxCache(InStageActor ? InStageActor->GetBBoxCache() : nullptr);
 	SetRootMotionHandling(InStageActor ? InStageActor->RootMotionHandling : EUsdRootMotionHandling::NoAdditionalRootMotion);
 
@@ -1244,7 +1252,7 @@ void FUsdLevelSequenceHelperImpl::UnbindFromUsdStageActor()
 		StageActor.Reset();
 	}
 
-	SetInfoCache(nullptr);
+	SetPrimLinkCache(nullptr);
 	SetRootMotionHandling(EUsdRootMotionHandling::NoAdditionalRootMotion);
 
 	OnStageEditTargetChangedHandle.Reset();
@@ -1390,9 +1398,10 @@ ULevelSequence* FUsdLevelSequenceHelperImpl::FindOrAddSequenceForLayer(
 		// due to references from the transaction buffer, so we would basically end up creating a identical new object on top of an existing one (the
 		// new object has the same address as the existing one). When importing we don't actually want to do this though, because we want these assets
 		// name to conflict so that we can publish/replace old assets if desired. The stage importer will make these names unique later if needed. We
-		// only get an InfoCache when importing (from UUsdStageImporter::ImportFromFile) or when BindToUsdStageActor is called, which also gives us a
-		// stage actor. So if we don't have an actor but have a cache, we're importing
-		const bool bIsImporting = StageActor.IsExplicitlyNull() && InfoCache;
+		// only get a PrimLinkCache when importing (from UUsdStageImporter::ImportFromFile) or when BindToUsdStageActor is called, which also gives us
+		// a stage actor. So if we don't have an actor but have a cache, we're importing
+
+		const bool bIsImporting = StageActor.IsExplicitlyNull() && PrimLinkCache.IsValid();
 		FName UniqueSequenceName = bIsImporting ? *UsdUnreal::ObjectUtils::SanitizeObjectName(FPaths::GetBaseFilename(SequenceDisplayName))
 												: MakeUniqueObjectName(
 													GetTransientPackage(),
@@ -2407,7 +2416,8 @@ void FUsdLevelSequenceHelperImpl::AddSkeletalTracks(const UUsdPrimTwin& PrimTwin
 		return;
 	}
 
-	if (!InfoCache)
+	UUsdPrimLinkCache* PrimLinkCachePtr = PrimLinkCache.Get();
+	if (!PrimLinkCachePtr)
 	{
 		return;
 	}
@@ -2426,7 +2436,7 @@ void FUsdLevelSequenceHelperImpl::AddSkeletalTracks(const UUsdPrimTwin& PrimTwin
 	// but we may belong to a FUsdStageImportContext, and so there's no AUsdStageActor at all to use.
 	// At this point it doesn't matter much though, because we shouldn't need to uncollapse a SkelAnimation prim path anyway
 	const UE::FSdfPath PrimPath = Prim.GetPrimPath();
-	UAnimSequence* Sequence = InfoCache->GetSingleAssetForPrim<UAnimSequence>(PrimPath);
+	UAnimSequence* Sequence = PrimLinkCachePtr->GetSingleAssetForPrim<UAnimSequence>(PrimPath);
 	if (!Sequence)
 	{
 		return;
@@ -2499,14 +2509,15 @@ void FUsdLevelSequenceHelperImpl::AddSkeletalTracks(const UUsdPrimTwin& PrimTwin
 void FUsdLevelSequenceHelperImpl::AddGeometryCacheTracks(const UUsdPrimTwin& PrimTwin, const UE::FUsdPrim& Prim)
 {
 	UGeometryCacheComponent* ComponentToBind = Cast<UGeometryCacheComponent>(PrimTwin.GetSceneComponent());
-	if (!ComponentToBind || !InfoCache)
+	UUsdPrimLinkCache* PrimLinkCachePtr = PrimLinkCache.Get();
+	if (!ComponentToBind || !PrimLinkCachePtr)
 	{
 		return;
 	}
 
 	// Fetch the geometry cache asset from the asset cache. If there's none, don't actually need to create track
 	const UE::FSdfPath PrimPath = Prim.GetPrimPath();
-	UGeometryCache* GeometryCache = InfoCache->GetSingleAssetForPrim<UGeometryCache>(PrimPath);
+	UGeometryCache* GeometryCache = PrimLinkCachePtr->GetSingleAssetForPrim<UGeometryCache>(PrimPath);
 	if (!GeometryCache)
 	{
 		return;
@@ -2577,7 +2588,8 @@ void FUsdLevelSequenceHelperImpl::AddGroomTracks(const UUsdPrimTwin& PrimTwin, c
 		return;
 	}
 
-	if (!InfoCache)
+	UUsdPrimLinkCache* PrimLinkCachePtr = PrimLinkCache.Get();
+	if (!PrimLinkCachePtr)
 	{
 		return;
 	}
@@ -2585,7 +2597,7 @@ void FUsdLevelSequenceHelperImpl::AddGroomTracks(const UUsdPrimTwin& PrimTwin, c
 	// Fetch the groom cache asset from the asset cache. If there's none, don't actually need to create track
 	const FString PrimPath = Prim.GetPrimPath().GetString();
 	const FString GroomCachePath = FString::Printf(TEXT("%s_strands_cache"), *PrimPath);
-	UGroomCache* GroomCache = InfoCache->GetSingleAssetForPrim<UGroomCache>(UE::FSdfPath{*GroomCachePath});
+	UGroomCache* GroomCache = PrimLinkCachePtr->GetSingleAssetForPrim<UGroomCache>(UE::FSdfPath{*GroomCachePath});
 	if (!GroomCache)
 	{
 		return;
@@ -2830,9 +2842,15 @@ void FUsdLevelSequenceHelperImpl::AddAudioTracks(const UUsdPrimTwin& PrimTwin, c
 		return;
 	}
 
+	UUsdPrimLinkCache* PrimLinkCachePtr = PrimLinkCache.Get();
+	if (!PrimLinkCachePtr)
+	{
+		return;
+	}
+
 	// Note: We pull the audio directly from the info cache here, and not the component:
 	// See big comment within FUsdMediaSpatialAudioTranslator::UpdateComponents
-	USoundBase* Sound = InfoCache->GetSingleAssetForPrim<USoundBase>(Prim.GetPrimPath());
+	USoundBase* Sound = PrimLinkCachePtr->GetSingleAssetForPrim<USoundBase>(Prim.GetPrimPath());
 	if (!Sound)
 	{
 		return;
@@ -3296,6 +3314,12 @@ void FUsdLevelSequenceHelperImpl::UpdateControlRigTracks(UUsdPrimTwin& PrimTwin)
 		return;
 	}
 
+	UUsdPrimLinkCache* PrimLinkCachePtr = PrimLinkCache.Get();
+	if (!PrimLinkCachePtr)
+	{
+		return;
+	}
+
 	UE::FSdfPath PrimPath{*PrimTwin.PrimPath};
 	UE::FUsdPrim SkeletonPrim{UsdStage.GetPrimAtPath(PrimPath)};
 	if (!SkeletonPrim)
@@ -3359,7 +3383,7 @@ void FUsdLevelSequenceHelperImpl::UpdateControlRigTracks(UUsdPrimTwin& PrimTwin)
 	// Fetch the UAnimSequence asset from the asset cache. Ideally we'd call AUsdStageActor::GetGeneratedAssets,
 	// but we may belong to a FUsdStageImportContext, and so there's no AUsdStageActor at all to use.
 	// At this point it doesn't matter much though, because we shouldn't need to uncollapse a SkelAnimation prim path anyway
-	UAnimSequence* AnimSequence = InfoCache->GetSingleAssetForPrim<UAnimSequence>(PrimPath);
+	UAnimSequence* AnimSequence = PrimLinkCachePtr->GetSingleAssetForPrim<UAnimSequence>(PrimPath);
 
 	UE::FUsdEditContext EditContext{UsdStage, SkelAnimationLayer};
 	FString Identifier = SkelAnimationLayer.GetIdentifier();
@@ -4873,7 +4897,7 @@ public:
 	{
 		return false;
 	}
-	void SetInfoCache(TSharedPtr<FUsdInfoCache> InfoCache){};
+	void SetPrimLinkCache(UUsdPrimLinkCache* PrimLinkCache){};
 	void SetBBoxCache(TSharedPtr<UE::FUsdGeomBBoxCache> InBBoxCache){};
 	bool HasData() const
 	{
@@ -4982,11 +5006,17 @@ void FUsdLevelSequenceHelper::OnStageActorRenamed()
 	}
 }
 
+PRAGMA_DISABLE_DEPRECATION_WARNINGS
 void FUsdLevelSequenceHelper::SetInfoCache(TSharedPtr<FUsdInfoCache> InInfoCache)
+{
+}
+PRAGMA_ENABLE_DEPRECATION_WARNINGS
+
+void FUsdLevelSequenceHelper::SetPrimLinkCache(UUsdPrimLinkCache* PrimLinkCache)
 {
 	if (UsdSequencerImpl.IsValid())
 	{
-		UsdSequencerImpl->SetInfoCache(InInfoCache);
+		UsdSequencerImpl->SetPrimLinkCache(PrimLinkCache);
 	}
 }
 
diff --git a/Engine/Plugins/Importers/USDImporter/Source/USDStage/Private/USDStageActor.cpp b/Engine/Plugins/Importers/USDImporter/Source/USDStage/Private/USDStageActor.cpp
--- a/Engine/Plugins/Importers/USDImporter/Source/USDStage/Private/USDStageActor.cpp
+++ b/Engine/Plugins/Importers/USDImporter/Source/USDStage/Private/USDStageActor.cpp
@@ -15,6 +15,7 @@
 #include "USDGeomMeshConversion.h"
 #include "USDGeomXformableTranslator.h"
 #include "USDInfoCache.h"
+#include "USDInfoCache.h"
 #include "USDIntegrationUtils.h"
 #include "USDLayerUtils.h"
 #include "USDLightConversion.h"
@@ -22,6 +23,7 @@
 #include "USDLog.h"
 #include "USDObjectUtils.h"
 #include "USDPrimConversion.h"
+#include "USDPrimLinkCache.h"
 #include "USDPrimTwin.h"
 #include "USDProjectSettings.h"
 #include "USDSchemasModule.h"
@@ -165,7 +167,8 @@ struct FUsdStageActorImpl
 		TranslationContext->SubdivisionLevel = StageActor->SubdivisionLevel;
 		TranslationContext->MetadataOptions = StageActor->MetadataOptions;
 		TranslationContext->BlendShapesByPath = &StageActor->BlendShapesByPath;
-		TranslationContext->InfoCache = StageActor->InfoCache;
+		TranslationContext->UsdInfoCache.Reset(StageActor->UsdInfoCache);
+		TranslationContext->PrimLinkCache.Reset(StageActor->PrimLinkCache);
 		TranslationContext->BBoxCache = StageActor->BBoxCache;
 		TranslationContext->bTranslateOnlyUsedMaterials = GTranslateOnlyUsedMaterialsWhenOpeningStage;
 
@@ -1266,7 +1269,7 @@ void AUsdStageActor::IsolateLayer(const UE::FSdfLayer& Layer)
 void AUsdStageActor::OnUsdObjectsChanged(const UsdUtils::FObjectChangesByPath& InfoChanges, const UsdUtils::FObjectChangesByPath& ResyncChanges)
 {
 #if USE_USD_SDK
-	if (!IsListeningToUsdNotices() || !InfoCache.IsValid())
+	if (!IsListeningToUsdNotices() || !UsdInfoCache)
 	{
 		return;
 	}
@@ -1532,14 +1535,20 @@ void AUsdStageActor::HandleAccumulatedNotices()
 
 	FScopedUsdMessageLog ScopedMessageLog;
 
-	TSharedPtr<FUsdInfoCache> OldInfoCache = InfoCache;
-	if (InfoCache.IsValid() && bHasResync)
+	TStrongObjectPtr<UUsdInfoCache> OldInfoCache{UsdInfoCache};
+	if (UsdInfoCache && bHasResync)
 	{
 		// Take a copy of the info cache here: We want to keep the old one for this function call as it helps us cleanup the old assets and components
 		// For now we don't need to do this for info changes, only for resync changes
 		{
 			TRACE_CPUPROFILER_EVENT_SCOPE(CopyingOldInfoCache);
-			InfoCache = MakeShared<FUsdInfoCache>(*OldInfoCache);
+			UObject* OwnerObject = this;
+			UsdInfoCache = NewObject<UUsdInfoCache>(OwnerObject, NAME_None, UsdInfoCache->GetFlags());
+
+			// Use this instead of DuplicateObject as this is a single Serialize() call (due to the Modify()), while the DuplicateObject means 3
+			// Serialize() calls
+			UsdInfoCache->CopyImpl(*OldInfoCache);
+
 		}
 
 		// TODO: Selective rebuild of only the required parts of the cache.
@@ -1549,11 +1558,11 @@ void AUsdStageActor::HandleAccumulatedNotices()
 		// subtree at all, or sibling subtrees. Note that whenver we do a selective rebuild of the info cache we'll need to be very careful when
 		// to update certain info cache maps: For example whenever we delete a prim we need to make sure it's removed from MaterialUsers, etc.
 		TSharedRef<FUsdSchemaTranslationContext> TranslationContext = FUsdStageActorImpl::CreateUsdSchemaTranslationContext(this, TEXT("/"));
-		InfoCache->RebuildCacheForSubtree(Stage.GetPseudoRoot(), TranslationContext.Get());
+		UsdInfoCache->RebuildCacheForSubtree(Stage.GetPseudoRoot(), TranslationContext.Get());
 
 		// Need to update the LevelSequenceHelper with our new cache too, as it will need the new cache to find any
 		// assets we end up creating in this update (e.g. USoundWave assets)
-		LevelSequenceHelper.SetInfoCache(InfoCache);
+		LevelSequenceHelper.SetPrimLinkCache(PrimLinkCache);
 	}
 
 	if (BBoxCache.IsValid() && bHasResync)
@@ -1649,7 +1658,7 @@ void AUsdStageActor::HandleAccumulatedNotices()
 		// both X and Y will be on the notice change list, even though X doesn't exist on the stage anymore.
 		// It's easy to just ignore those here by doing this, but we could also pay attention to the flags
 		// on the notice and try to guess if a rename took place, if needed
-		if (InfoCache->ContainsInfoAboutPrim(PrimPath))
+		if (UsdInfoCache->ContainsInfoAboutPrim(PrimPath))
 		{
 			// We always want the unwound path here. We'll take care to only ever register main prims that are
 			// themselves uncollapsed or collapsed roots, but there's nothing stopping the user from manually
@@ -1657,7 +1666,7 @@ void AUsdStageActor::HandleAccumulatedNotices()
 			// parent prim of a point instancer prototype). If we retrieved a main prim for a prim like that,
 			// we'd only get that prim itself, and assume we need to spawn assets/components for it (which we
 			// really don't if it's collapsed)
-			UE::FSdfPath UnwoundPath = InfoCache->UnwindToNonCollapsedPath(PrimPath, ECollapsingType::Assets);
+			UE::FSdfPath UnwoundPath = UsdInfoCache->UnwindToNonCollapsedPath(PrimPath, ECollapsingType::Assets);
 			OutPrimsToUpdate.FindOrAdd(UnwoundPath) |= bIsResync;
 		}
 
@@ -1715,7 +1724,7 @@ void AUsdStageActor::HandleAccumulatedNotices()
 			}
 		}
 
-		TSet<UE::FSdfPath> NewMainPrims = InfoCache->GetMainPrims(PrimPath);
+		TSet<UE::FSdfPath> NewMainPrims = UsdInfoCache->GetMainPrims(PrimPath);
 		TSet<UE::FSdfPath> OldMainPrims = OldInfoCache->GetMainPrims(PrimPath);
 		OutPrimsToUpdate.Reserve(OutPrimsToUpdate.Num() + NewMainPrims.Num() + OldMainPrims.Num());
 		for (const UE::FSdfPath& NewPrimPath : NewMainPrims)
@@ -1802,7 +1811,7 @@ void AUsdStageActor::HandleAccumulatedNotices()
 
 			// It's OK to not have info about a prim if it's an old prim that only exists on the old info cache.
 			// If the new info cache has info about this prim then it must exist on the stage right now
-			if (!InfoCache->ContainsInfoAboutPrim(PrimPath))
+			if (!UsdInfoCache->ContainsInfoAboutPrim(PrimPath))
 			{
 				ensure(!PrimToUpdate);
 				continue;
@@ -1833,10 +1842,10 @@ void AUsdStageActor::HandleAccumulatedNotices()
 			// Note that even after UE-157644 this is still useful: Material prims are not marked as a dependency
 			// of Mesh prims, otherwise we'd have to regenerate the StaticMesh itself when material info changed.
 			// Ideally we'd just set a new material override on the component, which is what this does
-			const bool bPrimGeneratesMaterials = InfoCache->GetSingleAssetForPrim<UMaterialInterface>(PrimPath) != nullptr;
+			const bool bPrimGeneratesMaterials = PrimLinkCache->GetSingleAssetForPrim<UMaterialInterface>(PrimPath) != nullptr;
 			if (bThisPrimLoadedAssets && bPrimGeneratesMaterials)
 			{
-				MaterialUserPrims.Append(InfoCache->GetMaterialUsers(PrimPath));
+				MaterialUserPrims.Append(UsdInfoCache->GetMaterialUsers(PrimPath));
 			}
 		}
 	}
@@ -1907,7 +1916,7 @@ void AUsdStageActor::HandleAccumulatedNotices()
 			const UE::FSdfPath PrimPath = PrimChangedInfo.Key;
 			const bool bIsResync = PrimChangedInfo.Value;
 
-			if (!InfoCache->ContainsInfoAboutPrim(PrimPath))
+			if (!UsdInfoCache->ContainsInfoAboutPrim(PrimPath))
 			{
 				continue;
 			}
@@ -2898,7 +2907,7 @@ USceneComponent* AUsdStageActor::GetGeneratedComponent(const FString& PrimPath)
 		return nullptr;
 	}
 
-	// We can't query our InfoCache with invalid paths, as we're using ensures to track when we miss the cache (which shouldn't ever happen)
+	// We can't query our UsdInfoCache with invalid paths, as we're using ensures to track when we miss the cache (which shouldn't ever happen)
 	UE::FSdfPath UsdPath{*PrimPath};
 	if (!CurrentStage.GetPrimAtPath(UsdPath))
 	{
@@ -2906,9 +2915,9 @@ USceneComponent* AUsdStageActor::GetGeneratedComponent(const FString& PrimPath)
 	}
 
 	FString UncollapsedPath = PrimPath;
-	if (InfoCache.IsValid())
+	if (UsdInfoCache)
 	{
-		UncollapsedPath = InfoCache->UnwindToNonCollapsedPath(UsdPath, ECollapsingType::Components).GetString();
+		UncollapsedPath = UsdInfoCache->UnwindToNonCollapsedPath(UsdPath, ECollapsingType::Components).GetString();
 	}
 
 	if (UUsdPrimTwin* UsdPrimTwin = GetRootPrimTwin()->Find(UncollapsedPath))
@@ -2927,24 +2936,24 @@ TArray<UObject*> AUsdStageActor::GetGeneratedAssets(const FString& PrimPath) con
 		return {};
 	}
 
-	// We can't query our InfoCache with invalid paths, as we're using ensures to track when we miss the cache (which shouldn't ever happen)
+	// We can't query our UsdInfoCache with invalid paths, as we're using ensures to track when we miss the cache (which shouldn't ever happen)
 	UE::FSdfPath UsdPath{*PrimPath};
 	if (!CurrentStage.GetPrimAtPath(UsdPath))
 	{
 		return {};
 	}
 
-	if (!InfoCache)
+	if (!PrimLinkCache || !UsdInfoCache)
 	{
 		return {};
 	}
 
 	// Prefer checking the prim directly, but also check its collapsed root if it is collapsed.
-	TArray<TWeakObjectPtr<UObject>> AssetsPtrs = InfoCache->GetAllAssetsForPrim(UsdPath);
-	if (AssetsPtrs.Num() == 0 && InfoCache->IsPathCollapsed(UsdPath, ECollapsingType::Assets))
+	TArray<TWeakObjectPtr<UObject>> AssetsPtrs = PrimLinkCache->GetAllAssetsForPrim(UsdPath);
+	if (AssetsPtrs.Num() == 0 && UsdInfoCache->IsPathCollapsed(UsdPath, ECollapsingType::Assets))
 	{
-		UsdPath = InfoCache->UnwindToNonCollapsedPath(UsdPath, ECollapsingType::Assets);
-		AssetsPtrs = InfoCache->GetAllAssetsForPrim(UsdPath);
+		UsdPath = UsdInfoCache->UnwindToNonCollapsedPath(UsdPath, ECollapsingType::Assets);
+		AssetsPtrs = PrimLinkCache->GetAllAssetsForPrim(UsdPath);
 	}
 
 	TArray<UObject*> Assets;
@@ -2978,9 +2987,9 @@ FString AUsdStageActor::GetSourcePrimPath(const UObject* Object) const
 			return UsdPrimTwin->PrimPath;
 		}
 	}
-	else if (InfoCache)
+	else if (PrimLinkCache)
 	{
-		const TArray<UE::FSdfPath> FoundPaths = InfoCache->GetPrimsForAsset(Object);
+		const TArray<UE::FSdfPath> FoundPaths = PrimLinkCache->GetPrimsForAsset(Object);
 		if (FoundPaths.Num() > 0)
 		{
 			return FoundPaths[0].GetString();
@@ -3196,9 +3205,8 @@ void AUsdStageActor::OnObjectsReplaced(const TMap<UObject*, UObject*>& ObjectRep
 			NewActor->OnPrimChanged = OnPrimChanged;
 
 			NewActor->AssetCache = AssetCache;
-
-			NewActor->InfoCache = InfoCache;
-			InfoCache = nullptr;
+			NewActor->UsdInfoCache = UsdInfoCache;
+			NewActor->PrimLinkCache = PrimLinkCache;
 
 			NewActor->BBoxCache = BBoxCache;
 			BBoxCache = nullptr;
@@ -3297,9 +3305,16 @@ void AUsdStageActor::LoadUsdStage(bool bOpenIfNeeded)
 
 	// Create Info and BBoxCache before calling ReloadAnimations as that is when the LevelSequenceHelper will also take
 	// a reference to them
-	if (!InfoCache.IsValid())
+	if (!UsdInfoCache)
+	{
+		UObject* Outer = this;
+		UsdInfoCache = NewObject<UUsdInfoCache>(Outer, NAME_None, EObjectFlags::RF_Transient | EObjectFlags::RF_Transactional);
+	}
+
+	if (!PrimLinkCache)
 	{
-		InfoCache = MakeShared<FUsdInfoCache>();
+		UObject* Outer = this;
+		PrimLinkCache = NewObject<UUsdPrimLinkCache>(Outer, NAME_None, EObjectFlags::RF_Transient | EObjectFlags::RF_Transactional);
 	}
 
 	RegenerateLevelSequence();
@@ -3307,8 +3322,8 @@ void AUsdStageActor::LoadUsdStage(bool bOpenIfNeeded)
 	TSharedRef<FUsdSchemaTranslationContext> TranslationContext = FUsdStageActorImpl::CreateUsdSchemaTranslationContext(this, RootTwin->PrimPath);
 
 	SlowTask.EnterProgressFrame(0.1f);
-	InfoCache->RemoveAllAssetPrimLinks();	 // These are usually not reset when rebuilding the cache so we must call manually
-	InfoCache->RebuildCacheForSubtree(StageToLoad.GetPseudoRoot(), TranslationContext.Get());
+	PrimLinkCache->RemoveAllAssetPrimLinks();
+	UsdInfoCache->RebuildCacheForSubtree(StageToLoad.GetPseudoRoot(), TranslationContext.Get());
 
 	SlowTask.EnterProgressFrame(0.7f);
 	const bool bLoadedOrAbandonedAssets = LoadAssets(*TranslationContext, StageToLoad.GetPseudoRoot());
@@ -3468,9 +3483,14 @@ void AUsdStageActor::UnloadUsdStage()
 		AssetCache->RequestDelayedAssetAutoCleanup();
 	}
 
-	if (InfoCache)
+	if (PrimLinkCache)
+	{
+		PrimLinkCache->Clear();
+	}
+
+	if (UsdInfoCache)
 	{
-		InfoCache->Clear();
+		UsdInfoCache->Clear();
 	}
 
 	if (BBoxCache)
@@ -3598,10 +3618,12 @@ void AUsdStageActor::RepopulateLevelSequence()
 	UpdatePrim(UE::FSdfPath::AbsoluteRootPath(), bIsResync, *TranslationContext);
 }
 
+PRAGMA_DISABLE_DEPRECATION_WARNINGS
 TSharedPtr<FUsdInfoCache> AUsdStageActor::GetInfoCache()
 {
 	return InfoCache;
 }
+PRAGMA_ENABLE_DEPRECATION_WARNINGS
 
 TSharedPtr<UE::FUsdGeomBBoxCache> AUsdStageActor::GetBBoxCache()
 {
@@ -3944,12 +3966,6 @@ void AUsdStageActor::Serialize(FArchive& Ar)
 
 	if ((Ar.GetPortFlags() & PPF_DuplicateForPIE) || Ar.IsTransacting())
 	{
-		if (!InfoCache.IsValid())
-		{
-			InfoCache = MakeShared<FUsdInfoCache>();
-		}
-
-		InfoCache->Serialize(Ar);
 		LevelSequenceHelper.Serialize(Ar);
 	}
 }
@@ -4818,7 +4834,7 @@ bool AUsdStageActor::UnloadAssets(const UE::FSdfPath& StartPrimPath, bool bForEn
 	// be tracked via the info cache asset prim links, so if at any time they resync the old assets will still be
 	// found below when iterating the prim links, and we will discard them either way.
 
-	if (!AssetCache || !InfoCache)
+	if (!AssetCache || !PrimLinkCache)
 	{
 		return false;
 	}
@@ -4828,7 +4844,7 @@ bool AUsdStageActor::UnloadAssets(const UE::FSdfPath& StartPrimPath, bool bForEn
 	TSet<UE::FSdfPath> PrimPathsToRemove;
 	if (bForEntireSubtree)
 	{
-		for (const TPair<UE::FSdfPath, TArray<TWeakObjectPtr<UObject>>>& PrimPathToAssetIt : InfoCache->GetAllAssetPrimLinks())
+		for (const TPair<UE::FSdfPath, TArray<TWeakObjectPtr<UObject>>>& PrimPathToAssetIt : PrimLinkCache->GetAllAssetPrimLinks())
 		{
 			const UE::FSdfPath& LinkPrimPath = PrimPathToAssetIt.Key;
 			if (LinkPrimPath.HasPrefix(StartPrimPath) || LinkPrimPath == StartPrimPath)
@@ -4844,13 +4860,13 @@ bool AUsdStageActor::UnloadAssets(const UE::FSdfPath& StartPrimPath, bool bForEn
 
 	for (const UE::FSdfPath& PrimPathToRemove : PrimPathsToRemove)
 	{
-		TArray<TWeakObjectPtr<UObject>> OldAssets = InfoCache->RemoveAllAssetPrimLinks(PrimPathToRemove);
+		TArray<TWeakObjectPtr<UObject>> OldAssets = PrimLinkCache->RemoveAllAssetPrimLinks(PrimPathToRemove);
 		for (const TWeakObjectPtr<UObject>& OldAsset : OldAssets)
 		{
 			// If there are any other prim paths linked to this asset that we *won't* be removing/reparsing
 			// in here, it means our stage actor as a whole is still "referencing" that asset
 			bool bAssetStillReferenced = false;
-			for (const UE::FSdfPath& LinkedPrim : InfoCache->GetPrimsForAsset(OldAsset.Get()))
+			for (const UE::FSdfPath& LinkedPrim : PrimLinkCache->GetPrimsForAsset(OldAsset.Get()))
 			{
 				if (!PrimPathsToRemove.Contains(LinkedPrim))
 				{
@@ -4880,18 +4896,17 @@ bool AUsdStageActor::LoadAsset(FUsdSchemaTranslationContext& TranslationContext,
 {
 	TRACE_CPUPROFILER_EVENT_SCOPE(AUsdStageActor::LoadAsset);
 
-	if (!AssetCache || !InfoCache)
+	if (!AssetCache)
 	{
 		return false;
 	}
 
 	int32 StartNumAssets = AssetCache->GetNumAssets();
 
-	// Suppress transaction while we're creating assets.
-	// c.f. the big comment on the analogous position within AUsdStageActor::LoadAssets
-	Modify();
 	AssetCache->Modify();
 	{
+		// Suppress transaction while we're creating assets.
+		// c.f. the big comment on the analogous position within AUsdStageActor::LoadAssets
 		TGuardValue<ITransaction*> SuppressTransaction{GUndo, nullptr};
 
 		IUsdSchemasModule& UsdSchemasModule = FModuleManager::Get().LoadModuleChecked<IUsdSchemasModule>(TEXT("USDSchemas"));
@@ -4914,7 +4929,7 @@ bool AUsdStageActor::LoadAssets(FUsdSchemaTranslationContext& TranslationContext
 {
 	TRACE_CPUPROFILER_EVENT_SCOPE(AUsdStageActor::LoadAssets);
 
-	if (!AssetCache || !InfoCache)
+	if (!AssetCache)
 	{
 		return false;
 	}
@@ -4955,20 +4970,18 @@ bool AUsdStageActor::LoadAssets(FUsdSchemaTranslationContext& TranslationContext
 		return false;
 	};
 
-	// Suppress current transaction, as we never want assets to be put into the transaction buffer. This because these
-	// will be exposed to the content browser now, so that "opening the stage" essentially acts as a full import.
-	// We don't want to rip these assets up when pressing undo after they've been created. The engine should act
-	// essentially as if they've always been there.
-	//
-	// Note that we have tried achieving that by just creating these assets without the RF_Transactional flag, but that is not enough:
-	// Some assets create subobjects that are transactional anyway (StaticMeshes), and some other assets have much
-	// more complicated logic that can even spawn some transient Worlds, actors and components, and can put them all
-	// into the transaction buffer (Skeletal assets), causing havoc if we try to make sense of object referencers
-	// when it's time to clean up the asset.
-	//
-	Modify();	 // Mainly to dirty the info cache
 	AssetCache->Modify();
 	{
+		// Suppress current transaction, as we never want assets to be put into the transaction buffer. This because these
+		// will be exposed to the content browser now, so that "opening the stage" essentially acts as a full import.
+		// We don't want to rip these assets up when pressing undo after they've been created. The engine should act
+		// essentially as if they've always been there.
+		//
+		// Note that we have tried achieving that by just creating these assets without the RF_Transactional flag, but that is not enough:
+		// Some assets create subobjects that are transactional anyway (StaticMeshes), and some other assets have much
+		// more complicated logic that can even spawn some transient Worlds, actors and components, and can put them all
+		// into the transaction buffer (Skeletal assets), causing havoc if we try to make sense of object referencers
+		// when it's time to clean up the asset.
 		TGuardValue<ITransaction*> SuppressTransaction{GUndo, nullptr};
 
 		// Load materials first since meshes are referencing them
diff --git a/Engine/Plugins/Importers/USDImporter/Source/USDStage/Public/USDLevelSequenceHelper.h b/Engine/Plugins/Importers/USDImporter/Source/USDStage/Public/USDLevelSequenceHelper.h
--- a/Engine/Plugins/Importers/USDImporter/Source/USDStage/Public/USDLevelSequenceHelper.h
+++ b/Engine/Plugins/Importers/USDImporter/Source/USDStage/Public/USDLevelSequenceHelper.h
@@ -10,6 +10,7 @@
 class AUsdStageActor;
 class FScopedBlockMonitoringChangesForTransaction;
 class FUsdInfoCache;
+class UUsdPrimLinkCache;
 class FUsdLevelSequenceHelperImpl;
 class ULevelSequence;
 class UUsdPrimTwin;
@@ -43,8 +44,13 @@ public:
 	/** Allows serialization for transaction support */
 	bool Serialize(FArchive& Ar);
 
+	PRAGMA_DISABLE_DEPRECATION_WARNINGS
 	/** Sets the asset cache to use when fetching assets and asset info required for the level sequence animation, like UAnimSequences */
+	UE_DEPRECATED(5.5, "Use SetPrimLinkCache instead")
 	void SetInfoCache(TSharedPtr<FUsdInfoCache> InInfoCache);
+	PRAGMA_ENABLE_DEPRECATION_WARNINGS
+
+	void SetPrimLinkCache(UUsdPrimLinkCache* PrimLinkCache);
 
 	/** Sets the BBoxCache to use when importing bound animations for prims. Needed for importing, where we don't have a stage actor to take the
 	 * BBoxCache from */
diff --git a/Engine/Plugins/Importers/USDImporter/Source/USDStage/Public/USDStageActor.h b/Engine/Plugins/Importers/USDImporter/Source/USDStage/Public/USDStageActor.h
--- a/Engine/Plugins/Importers/USDImporter/Source/USDStage/Public/USDStageActor.h
+++ b/Engine/Plugins/Importers/USDImporter/Source/USDStage/Public/USDStageActor.h
@@ -2,9 +2,11 @@
 
 #pragma once
 
+#include "USDInfoCache.h"
 #include "USDLevelSequenceHelper.h"
 #include "USDListener.h"
 #include "USDMetadataImportOptions.h"
+#include "USDPrimLinkCache.h"
 #include "USDSkeletalDataConversion.h"
 #include "UsdWrappers/SdfPath.h"
 #include "UsdWrappers/UsdStage.h"
@@ -55,6 +57,12 @@ public:
 	UPROPERTY(EditAnywhere, BlueprintReadOnly, Category = "USD")
 	TObjectPtr<UUsdAssetCache3> AssetCache;
 
+	UPROPERTY()
+	TObjectPtr<UUsdInfoCache> UsdInfoCache;
+
+	UPROPERTY()
+	TObjectPtr<UUsdPrimLinkCache> PrimLinkCache;
+
 	UPROPERTY()
 	TObjectPtr<UUsdAssetCache2> UsdAssetCache;
 
@@ -317,7 +325,10 @@ public:
 	 */
 	USDSTAGE_API void RepopulateLevelSequence();
 
+	PRAGMA_DISABLE_DEPRECATION_WARNINGS
+	UE_DEPRECATED(5.5, "The FUsdInfoCache struct is deprecated in favor of the UUsdInfoCache and UUsdPrimLinkCache, referenced via UPROPERTYs")
 	USDSTAGE_API TSharedPtr<FUsdInfoCache> GetInfoCache();
+	PRAGMA_ENABLE_DEPRECATION_WARNINGS
 	USDSTAGE_API TSharedPtr<UE::FUsdGeomBBoxCache> GetBBoxCache();
 	USDSTAGE_API TMap<FString, TMap<FString, int32>> GetMaterialToPrimvarToUVIndex();
 	USDSTAGE_API const UsdUtils::FBlendShapeMap& GetBlendShapeMap();
@@ -453,8 +464,10 @@ protected:
 	UsdUtils::FObjectChangesByPath AccumulatedInfoChanges;
 	UsdUtils::FObjectChangesByPath AccumulatedResyncChanges;
 
+	PRAGMA_DISABLE_DEPRECATION_WARNINGS
 	/** Caches various information about prims that are expensive to query */
 	TSharedPtr<FUsdInfoCache> InfoCache;
+	PRAGMA_ENABLE_DEPRECATION_WARNINGS
 
 	/** USD bounding box cache for the stage. Constructed on-demand */
 	TSharedPtr<UE::FUsdGeomBBoxCache> BBoxCache;
diff --git a/Engine/Plugins/Importers/USDImporter/Source/USDStageImporter/Private/USDStageImporter.cpp b/Engine/Plugins/Importers/USDImporter/Source/USDStageImporter/Private/USDStageImporter.cpp
--- a/Engine/Plugins/Importers/USDImporter/Source/USDStageImporter/Private/USDStageImporter.cpp
+++ b/Engine/Plugins/Importers/USDImporter/Source/USDStageImporter/Private/USDStageImporter.cpp
@@ -290,14 +290,9 @@ namespace UE::USDStageImporter::Private
 
 	void CacheCollapsingState(FUsdSchemaTranslationContext& TranslationContext)
 	{
-		if (!TranslationContext.InfoCache.IsValid())
-		{
-			TranslationContext.InfoCache = MakeShared<FUsdInfoCache>();
-		}
-
 		// It's better to always rebuild the info cache because our import options may have changed
 		// from the options used when the cache was first built, which could change collapsing states, etc.
-		TranslationContext.InfoCache->RebuildCacheForSubtree(TranslationContext.Stage.GetPseudoRoot(), TranslationContext);
+		TranslationContext.UsdInfoCache->RebuildCacheForSubtree(TranslationContext.Stage.GetPseudoRoot(), TranslationContext);
 	}
 
 	void ImportMaterials(FUsdStageImportContext& ImportContext, FUsdSchemaTranslationContext& TranslationContext)
@@ -681,11 +676,7 @@ namespace UE::USDStageImporter::Private
 			// rename our asset to its target package too
 			else
 			{
-				ensure(Asset->Rename(
-					*TargetAssetName,
-					Package,
-					REN_DontCreateRedirectors | REN_NonTransactional | REN_DoNotDirty
-				));
+				ensure(Asset->Rename(*TargetAssetName, Package, REN_DontCreateRedirectors | REN_NonTransactional | REN_DoNotDirty));
 				MovedAsset = Asset;
 			}
 
@@ -703,9 +694,7 @@ namespace UE::USDStageImporter::Private
 		else
 		{
 			// We can't dirty the package here. Read the comment around MarkPackageDirty, below
-			ensure(
-				Asset->Rename(*TargetAssetName, Package, REN_DontCreateRedirectors | REN_NonTransactional | REN_DoNotDirty)
-			);
+			ensure(Asset->Rename(*TargetAssetName, Package, REN_DontCreateRedirectors | REN_NonTransactional | REN_DoNotDirty));
 			MovedAsset = Asset;
 		}
 
@@ -1898,9 +1887,14 @@ void UUsdStageImporter::ImportFromFile(FUsdStageImportContext& ImportContext)
 	// content browser and can instead just be renamed by the stageimporter directly into the import location
 	TGuardValue<FString> AssetCacheTransientGuard{ImportContext.UsdAssetCache->AssetDirectory.Path, GetTransientPackage()->GetPathName()};
 
-	TSharedPtr<FUsdInfoCache> InfoCache = MakeShared<FUsdInfoCache>();
+	UObject* Outer = GetTransientPackage();
+	const FName Name = NAME_None;
+	const EObjectFlags Flags = EObjectFlags::RF_Transient;
+	TStrongObjectPtr<UUsdPrimLinkCache> PrimLinkCache{NewObject<UUsdPrimLinkCache>(Outer, Name, Flags)};
+	TStrongObjectPtr<UUsdInfoCache> UsdInfoCache{NewObject<UUsdInfoCache>(Outer, Name, Flags)};
+
 	ImportContext.UsdAssetCache->MarkAssetsAsStale();
-	ImportContext.LevelSequenceHelper.SetInfoCache(InfoCache);
+	ImportContext.LevelSequenceHelper.SetPrimLinkCache(PrimLinkCache.Get());
 	ImportContext.LevelSequenceHelper.Init(ImportContext.Stage);	// Must happen after the context gets an InfoCache!
 	ImportContext.LevelSequenceHelper.SetRootMotionHandling(ImportContext.ImportOptions->RootMotionHandling);
 
@@ -1960,7 +1954,8 @@ void UUsdStageImporter::ImportFromFile(FUsdStageImportContext& ImportContext)
 	TranslationContext->bAllowParsingSparseVolumeTextures = ImportContext.ImportOptions->bImportSparseVolumeTextures;
 	TranslationContext->bAllowParsingSounds = ImportContext.ImportOptions->bImportSounds;
 	TranslationContext->bTranslateOnlyUsedMaterials = ImportContext.ImportOptions->bImportOnlyUsedMaterials;
-	TranslationContext->InfoCache = InfoCache;
+	TranslationContext->UsdInfoCache = UsdInfoCache;
+	TranslationContext->PrimLinkCache = PrimLinkCache;
 	TranslationContext->BBoxCache = ImportContext.BBoxCache;
 	TranslationContext->BlendShapesByPath = &BlendShapesByPath;
 	TranslationContext->GroomInterpolationSettings = ImportContext.ImportOptions->GroomInterpolationSettings;
@@ -2106,9 +2101,14 @@ bool UUsdStageImporter::ReimportSingleAsset(
 	// content browser and can instead just be renamed by the stageimporter directly into the import location
 	TGuardValue<FString> AssetCacheTransientGuard{ImportContext.UsdAssetCache->AssetDirectory.Path, GetTransientPackage()->GetPathName()};
 
-	TSharedPtr<FUsdInfoCache> InfoCache = MakeShared<FUsdInfoCache>();
+	UObject* Outer = GetTransientPackage();
+	const FName Name = NAME_None;
+	const EObjectFlags Flags = EObjectFlags::RF_Transient;
+	TStrongObjectPtr<UUsdPrimLinkCache> PrimLinkCache{NewObject<UUsdPrimLinkCache>(Outer, Name, Flags)};
+	TStrongObjectPtr<UUsdInfoCache> UsdInfoCache{NewObject<UUsdInfoCache>(Outer, Name, Flags)};
+
 	ImportContext.UsdAssetCache->MarkAssetsAsStale();
-	ImportContext.LevelSequenceHelper.SetInfoCache(InfoCache);
+	ImportContext.LevelSequenceHelper.SetPrimLinkCache(PrimLinkCache.Get());
 	ImportContext.LevelSequenceHelper.Init(ImportContext.Stage);	// Must happen after the context gets an InfoCache!
 	ImportContext.LevelSequenceHelper.SetRootMotionHandling(ImportContext.ImportOptions->RootMotionHandling);
 
@@ -2166,7 +2166,8 @@ bool UUsdStageImporter::ReimportSingleAsset(
 	TranslationContext->bAllowParsingSparseVolumeTextures = ImportContext.ImportOptions->bImportSparseVolumeTextures;
 	TranslationContext->bAllowParsingSounds = ImportContext.ImportOptions->bImportSounds;
 	TranslationContext->bTranslateOnlyUsedMaterials = ImportContext.ImportOptions->bImportOnlyUsedMaterials;
-	TranslationContext->InfoCache = InfoCache;
+	TranslationContext->UsdInfoCache = UsdInfoCache;
+	TranslationContext->PrimLinkCache = PrimLinkCache;
 	TranslationContext->BBoxCache = ImportContext.BBoxCache;
 	TranslationContext->BlendShapesByPath = &BlendShapesByPath;
 	TranslationContext->GroomInterpolationSettings = ImportContext.ImportOptions->GroomInterpolationSettings;
diff --git a/Engine/Plugins/Importers/USDImporter/Source/USDStageImporter/Public/USDStageImportContext.h b/Engine/Plugins/Importers/USDImporter/Source/USDStageImporter/Public/USDStageImportContext.h
--- a/Engine/Plugins/Importers/USDImporter/Source/USDStageImporter/Public/USDStageImportContext.h
+++ b/Engine/Plugins/Importers/USDImporter/Source/USDStageImporter/Public/USDStageImportContext.h
@@ -66,9 +66,11 @@ struct USDSTAGEIMPORTER_API FUsdStageImportContext
 	UPROPERTY()
 	TObjectPtr<UUsdAssetCache2> AssetCache;
 
+	PRAGMA_DISABLE_DEPRECATION_WARNINGS
 	/** Caches various information about prims that are expensive to query */
 	UE_DEPRECATED(5.3, "The import process now always builds its own InfoCache, so this member is no longer used")
 	TSharedPtr<FUsdInfoCache> InfoCache;
+	PRAGMA_ENABLE_DEPRECATION_WARNINGS
 
 	/** Bounding box cache used for the USD stage in case we have to spawn bounds components */
 	TSharedPtr<UE::FUsdGeomBBoxCache> BBoxCache;
diff --git a/Engine/Plugins/Importers/USDImporter/Source/USDTests/Private/USDTestsBlueprintLibrary.cpp b/Engine/Plugins/Importers/USDImporter/Source/USDTests/Private/USDTestsBlueprintLibrary.cpp
--- a/Engine/Plugins/Importers/USDImporter/Source/USDTests/Private/USDTestsBlueprintLibrary.cpp
+++ b/Engine/Plugins/Importers/USDImporter/Source/USDTests/Private/USDTestsBlueprintLibrary.cpp
@@ -50,7 +50,7 @@ bool USDTestsBlueprintLibrary::RecompileBlueprintStageActor(AUsdStageActor* Blue
 		UE_LOG(LogUsd, Error, TEXT("Blueprint failed to compile (%s)"), *BP->GetName());
 		return false;
 	}
-	
+
 	UE_LOG(LogUsd, Error, TEXT("Blueprint is in an unexpected state after compiling (%s)"), *BP->GetName());
 #endif	  // WITH_EDITOR
 	return false;
@@ -92,7 +92,7 @@ int64 USDTestsBlueprintLibrary::GetSubtreeVertexCount(AUsdStageActor* StageActor
 {
 	if (StageActor)
 	{
-		if (TSharedPtr<FUsdInfoCache> Cache = StageActor->GetInfoCache())
+		if (UUsdInfoCache* Cache = StageActor->UsdInfoCache)
 		{
 			TOptional<uint64> Result = Cache->GetSubtreeVertexCount(UE::FSdfPath{*PrimPath});
 			if (Result.IsSet())
@@ -114,7 +114,7 @@ int64 USDTestsBlueprintLibrary::GetSubtreeMaterialSlotCount(AUsdStageActor* Stag
 {
 	if (StageActor)
 	{
-		if (TSharedPtr<FUsdInfoCache> Cache = StageActor->GetInfoCache())
+		if (UUsdInfoCache* Cache = StageActor->UsdInfoCache)
 		{
 			TOptional<uint64> Result = Cache->GetSubtreeMaterialSlotCount(UE::FSdfPath{*PrimPath});
 			if (Result.IsSet())
