USD: Prevent ExpandPrim from recursing to children on info changes.

diff --git a/Engine/Plugins/Importers/USDImporter/Source/USDStage/Private/USDStageActor.cpp b/Engine/Plugins/Importers/USDImporter/Source/USDStage/Private/USDStageActor.cpp
--- a/Engine/Plugins/Importers/USDImporter/Source/USDStage/Private/USDStageActor.cpp
+++ b/Engine/Plugins/Importers/USDImporter/Source/USDStage/Private/USDStageActor.cpp
@@ -2108,7 +2108,7 @@ UUsdPrimTwin* AUsdStageActor::ExpandPrim(
 				SchemaTranslator->UpdateComponents(TwinSceneComponent);
 			}
 
-			bExpandChildren = !SchemaTranslator->CollapsesChildren(ECollapsingType::Components);
+			bExpandChildren = bResync && !SchemaTranslator->CollapsesChildren(ECollapsingType::Components);
 		}
 	}
 
