USD: Added support to export CameraRig_Rail animation to USD.

diff --git a/Engine/Plugins/Importers/USDImporter/Source/USDExporter/Private/LevelSequenceExporterUSD.cpp b/Engine/Plugins/Importers/USDImporter/Source/USDExporter/Private/LevelSequenceExporterUSD.cpp
--- a/Engine/Plugins/Importers/USDImporter/Source/USDExporter/Private/LevelSequenceExporterUSD.cpp
+++ b/Engine/Plugins/Importers/USDImporter/Source/USDExporter/Private/LevelSequenceExporterUSD.cpp
@@ -25,6 +25,7 @@
 #include "AssetExportTask.h"
 #include "Bindings/MovieSceneReplaceableDirectorBlueprintBinding.h"
 #include "Bindings/MovieSceneSpawnableDirectorBlueprintBinding.h"
+#include "CameraRig_Rail.h"
 #include "Compilation/MovieSceneCompiledDataManager.h"
 #include "Components/SkeletalMeshComponent.h"
 #include "Editor.h"
@@ -1067,9 +1068,9 @@ namespace UE::LevelSequenceExporterUSD::Private
 									// Note that it's perfectly fine if this ends up being just "."
 									UE::FSdfPath RelativePrimPathInSourceStage = SdfPrimPathInSourceStage.MakeRelativePath(DefaultPrimPath);
 
-									FString PrimPathRelativeToParentStageActor = UE::FSdfPath{*ParentStageActorPrimPathOnExport}
-																					 .AppendPath(RelativePrimPathInSourceStage)
-																					 .GetString();
+									FString PrimPathRelativeToParentStageActor = UE::FSdfPath{ *ParentStageActorPrimPathOnExport }
+										.AppendPath(RelativePrimPathInSourceStage)
+										.GetString();
 
 									if (!PrimPathRelativeToParentStageActor.IsEmpty())
 									{
@@ -1097,56 +1098,90 @@ namespace UE::LevelSequenceExporterUSD::Private
 				}
 			}
 
-			FString SchemaName = UsdUtils::GetSchemaNameForComponent(*BoundComponent);
-			if (SchemaName.IsEmpty())
+			auto GetPrimForComponent = [&Context, &UsdStage](const USceneComponent& Component, FString* PrimPathStr = nullptr) -> UE::FUsdPrim
 			{
-				continue;
-			}
+				FString PrimPath = PrimPathStr ? *PrimPathStr : UsdUtils::GetPrimPathForObject(
+					&Component,
+					TEXT(""),
+					Context.ExportOptions && Context.ExportOptions->LevelExportOptions.bExportActorFolders
+				);
+				if (PrimPath.IsEmpty())
+				{
+					return {};
+				}
+
+				FString SchemaName = UsdUtils::GetSchemaNameForComponent(Component);
+				if (SchemaName.IsEmpty())
+				{
+					return {};
+				}
+
+				// We will define a prim here so that we can apply schemas and use the shortcut CreateXAttribute functions,
+				// and not have to worry about attribute names and types. Later on we will convert these prims back into just 'overs' though
+				return UsdStage.DefinePrim(UE::FSdfPath{ *PrimPath }, *SchemaName);
+			};
 
-			// We will define a prim here so that we can apply schemas and use the shortcut CreateXAttribute functions,
-			// and not have to worry about attribute names and types. Later on we will convert these prims back into just 'overs' though
-			UE::FUsdPrim Prim = UsdStage.DefinePrim(UE::FSdfPath{*PrimPath}, *SchemaName);
+			UE::FUsdPrim Prim = GetPrimForComponent(*BoundComponent, &PrimPath);
 			if (!Prim)
 			{
 				continue;
 			}
 
 			TFunction<void(UnrealToUsd::FComponentBaker&)> AddBaker = [&InOutComponentBakers, &BoundComponent](UnrealToUsd::FComponentBaker& Baker)
-			{
-				// If we made a baker and we don't have one of this type for this component yet, add its lambda to the array
-				FCombinedComponentBakers& ExistingBakers = InOutComponentBakers.FindOrAdd(BoundComponent);
-				if (Baker.BakerType != UnrealToUsd::EBakingType::None && !EnumHasAnyFlags(ExistingBakers.CombinedBakingType, Baker.BakerType))
 				{
-					ExistingBakers.Bakers.Add(Baker);
-					ExistingBakers.CombinedBakingType |= Baker.BakerType;
-				}
-			};
+					// If we made a baker and we don't have one of this type for this component yet, add its lambda to the array
+					FCombinedComponentBakers& ExistingBakers = InOutComponentBakers.FindOrAdd(BoundComponent);
+					if (Baker.BakerType != UnrealToUsd::EBakingType::None && !EnumHasAnyFlags(ExistingBakers.CombinedBakingType, Baker.BakerType))
+					{
+						ExistingBakers.Bakers.Add(Baker);
+						ExistingBakers.CombinedBakingType |= Baker.BakerType;
+					}
+				};
 
 			TFunction<void(UnrealToUsd::FComponentBaker&)> GenerateSkeletalBaker =
 				[&UsdStage, &BoundComponent, &PrimPath](UnrealToUsd::FComponentBaker& InOutBaker)
-			{
-				if (USkeletalMeshComponent* SkeletalBoundComponent = Cast<USkeletalMeshComponent>(BoundComponent))
 				{
-					UE::FUsdPrim SkelAnimPrim = UsdStage.DefinePrim(UE::FSdfPath{*PrimPath}.AppendChild(TEXT("Anim")), TEXT("SkelAnimation"));
+					if (USkeletalMeshComponent* SkeletalBoundComponent = Cast<USkeletalMeshComponent>(BoundComponent))
+					{
+						UE::FUsdPrim SkelAnimPrim = UsdStage.DefinePrim(UE::FSdfPath{ *PrimPath }.AppendChild(TEXT("Anim")), TEXT("SkelAnimation"));
 
-					UE::FUsdPrim SkeletonPrim = UsdStage.DefinePrim(
-						UE::FSdfPath{*PrimPath}.AppendChild(UnrealIdentifiers::ExportedSkeletonPrimName),
-						TEXT("Skeleton")
-					);
+						UE::FUsdPrim SkeletonPrim = UsdStage.DefinePrim(
+							UE::FSdfPath{ *PrimPath }.AppendChild(UnrealIdentifiers::ExportedSkeletonPrimName),
+							TEXT("Skeleton")
+						);
 
-					if (SkelAnimPrim && SkeletonPrim)
-					{
-						UnrealToUsd::CreateSkeletalAnimationBaker(SkeletonPrim, SkelAnimPrim, *SkeletalBoundComponent, InOutBaker);
+						if (SkelAnimPrim && SkeletonPrim)
+						{
+							UnrealToUsd::CreateSkeletalAnimationBaker(SkeletonPrim, SkelAnimPrim, *SkeletalBoundComponent, InOutBaker);
+						}
+						else
+						{
+							UE_LOG(LogUsd, Warning, TEXT("Failed to generate Skeleton or SkelAnimation prim when baking out SkelRoot '%s'"), *PrimPath);
+						}
 					}
-					else
+				};
+
+			bool bHasTransformBaker = false;
+			bool bHasSkeletalBaker = false;
+
+			if (ACameraRig_Rail* RailActor = Cast<ACameraRig_Rail>(BoundObject))
+			{
+				// In the case of a CameraRig_Rail, what we want to bake is the transform animation of its 
+				// RailMountComponent since that is where children camera will be attached to.
+				const static FString TransformPropertyPath = UnrealIdentifiers::TransformPropertyName.ToString();
+				USceneComponent* RailMountComponent = RailActor->GetDefaultAttachComponent();
+				UE::FUsdPrim RailMountPrim = GetPrimForComponent(*RailMountComponent);
+				if (RailMountPrim)
+				{
+					UnrealToUsd::FComponentBaker Baker;
+					if (UnrealToUsd::CreateComponentPropertyBaker(RailMountPrim, *RailMountComponent, TransformPropertyPath, Baker))
 					{
-						UE_LOG(LogUsd, Warning, TEXT("Failed to generate Skeleton or SkelAnimation prim when baking out SkelRoot '%s'"), *PrimPath);
+						AddBaker(Baker);
+						bHasTransformBaker = true;
 					}
 				}
-			};
+			}
 
-			bool bHasTransformBaker = false;
-			bool bHasSkeletalBaker = false;
 			TMap<FString, int32> AudioTracksPerPrim;
 			if (const FMovieSceneBinding* Binding = MovieScene->FindBinding(InstanceKey.Key))
 			{
diff --git a/Engine/Plugins/Runtime/USDCore/Source/USDUtilities/Private/USDPrimConversion.cpp b/Engine/Plugins/Runtime/USDCore/Source/USDUtilities/Private/USDPrimConversion.cpp
--- a/Engine/Plugins/Runtime/USDCore/Source/USDUtilities/Private/USDPrimConversion.cpp
+++ b/Engine/Plugins/Runtime/USDCore/Source/USDUtilities/Private/USDPrimConversion.cpp
@@ -4124,6 +4124,12 @@ bool UnrealToUsd::CreateComponentPropertyBaker(
 					return;
 				}
 
+				// Some setup (like CameraRig_Rail) require that the parent actor on which the component is attached be Tick'ed to update their positions
+				if (AActor* AttachParentActor = Component->GetAttachParentActor())
+				{
+					AttachParentActor->Tick(0.0f);
+				}
+
 				USceneComponent* OriginalAttachParent = OriginalAttachParentGetter();
 
 				// If we're attached to a socket our RelativeTransform will be relative to the socket, instead of the parent
