Fixed unreachable code error when compiling a shipping editor

#rnx

diff --git a/Engine/Plugins/Importers/USDImporter/Source/USDSchemas/Private/USDSchemaTranslator.cpp b/Engine/Plugins/Importers/USDImporter/Source/USDSchemas/Private/USDSchemaTranslator.cpp
--- a/Engine/Plugins/Importers/USDImporter/Source/USDSchemas/Private/USDSchemaTranslator.cpp
+++ b/Engine/Plugins/Importers/USDImporter/Source/USDSchemas/Private/USDSchemaTranslator.cpp
@@ -252,9 +252,9 @@ bool FUsdSchemaTranslator::IsCollapsed(ECollapsingType CollapsingType) const
 
 	// This is merely a fallback, and we should never need this
 	return CanBeCollapsed(CollapsingType);
-#endif	  // #if USE_USD_SDK
-
+#else	// #if USE_USD_SDK
 	return false;
+#endif	// #if USE_USD_SDK
 }
 
 UE::FSdfPath FUsdSchemaTranslator::GetPrototypePrimPath() const
