USD: Fixed streaming GeometryCache not updating frame in some situation.
Fixed USD.GeometryCache.ForceImport cvar not working.
Fixed GeometryCaches hashing to the same asset even with different timecodes.

diff --git a/Engine/Plugins/Importers/USDImporter/Source/USDSchemas/Private/USDGeometryCacheTranslator.cpp b/Engine/Plugins/Importers/USDImporter/Source/USDSchemas/Private/USDGeometryCacheTranslator.cpp
--- a/Engine/Plugins/Importers/USDImporter/Source/USDSchemas/Private/USDGeometryCacheTranslator.cpp
+++ b/Engine/Plugins/Importers/USDImporter/Source/USDSchemas/Private/USDGeometryCacheTranslator.cpp
@@ -70,7 +70,7 @@ static FAutoConsoleVariableRef CVarEnableSubdiv(
 static bool GForceImport = true;
 static FAutoConsoleVariableRef CVarForceImport(
 	TEXT("USD.GeometryCache.ForceImport"),
-	GEnableSubdiv,
+	GForceImport,
 	TEXT(
 		"Whether to immediately cache all geometry cache frames into the asset when the stage opens, acting like a fully imported Geometry Cache asset"
 	)
@@ -416,24 +416,19 @@ namespace UsdGeometryCacheTranslatorImpl
 	)
 	{
 		FString RootPrimPath = RootPrim.GetPrimPath().GetString();
+		FReadMeshDataArgs Args(GetReadMeshDataArgs(Context, RootPrimPath));
 
 		// Compute the asset hash from the merged mesh description
 		FSHA1 SHA1;
 		FSHAHash MeshHash = FStaticMeshOperations::ComputeSHAHash(MeshDescription);
 		SHA1.Update(&MeshHash.Hash[0], sizeof(MeshHash.Hash));
 
-		const UE::FUsdStage& Stage = Context->Stage;
-		double FramesPerSecond = Stage.GetTimeCodesPerSecond();
-		if (FramesPerSecond == 0)
-		{
-			ensureMsgf(false, TEXT("Invalid USD GeometryCache FPS detected. Falling back to 1 FPS"));
-			FramesPerSecond = 1;
-		}
-
 		const bool bIsImporting = GForceImport || Context->bIsImporting;
 
 		// Frame rate must be taken into account as well since different frame rates must produce different sampling in the tracks
-		SHA1.Update(reinterpret_cast<uint8*>(&FramesPerSecond), sizeof(FramesPerSecond));
+		SHA1.Update(reinterpret_cast<uint8*>(&Args.FramesPerSecond), sizeof(Args.FramesPerSecond));
+		SHA1.Update(reinterpret_cast<uint8*>(&Args.StartFrame), sizeof(Args.StartFrame));
+		SHA1.Update(reinterpret_cast<uint8*>(&Args.EndFrame), sizeof(Args.EndFrame));
 
 		// Track type depends on if it's importing or not. Import needs to generate a persistent asset with all the frames already sampled
 		SHA1.Update(reinterpret_cast<const uint8*>(&bIsImporting), sizeof(bIsImporting));
@@ -454,7 +449,6 @@ namespace UsdGeometryCacheTranslatorImpl
 
 		if (GeometryCache && bOutIsNew)
 		{
-			FReadMeshDataArgs Args(GetReadMeshDataArgs(Context, RootPrimPath));
 			if (!bIsImporting)
 			{
 				// StartOffsetTime is the offset applied to the GeometryCache section on the sequencer track, so not relevant when importing
diff --git a/Engine/Plugins/Runtime/GeometryCache/Source/GeometryCacheStreamer/Private/GeometryCacheStreamBase.cpp b/Engine/Plugins/Runtime/GeometryCache/Source/GeometryCacheStreamer/Private/GeometryCacheStreamBase.cpp
--- a/Engine/Plugins/Runtime/GeometryCache/Source/GeometryCacheStreamer/Private/GeometryCacheStreamBase.cpp
+++ b/Engine/Plugins/Runtime/GeometryCache/Source/GeometryCacheStreamer/Private/GeometryCacheStreamBase.cpp
@@ -432,6 +432,13 @@ void FGeometryCacheStreamBase::UpdateCurrentFrameIndex(int32 FrameIndex)
 {
 	if (FrameIndex != CurrentFrameIndex)
 	{
+		const int32 FrameDelta = FMath::Abs(FrameIndex - CurrentFrameIndex);
+		if (FrameDelta >= MaxCachedFrames && FramesNeeded.Num() == 0)
+		{
+			// When the time jump is greater than the number of cached frames and it is not currently streaming,
+			// force loading the requested FrameIndex so that it gets displayed when the proxy is updated.
+			LoadFrameData(FrameIndex);
+		}
 		CurrentFrameIndex = FrameIndex;
 		bCacheNeedsUpdate = true;
 	}
