USD: Partial info cache rebuild;

Remove dedicated "MaterialUsers" mechanism and just use main/aux prims;
Prevent info cache from ever being serialized for transactions and just do partial builds on undo/redo;
Remove info cache copy when handling prim notices;
Use prim indices for recursive traversal;
Refactor collapsing information storage to allow partial rebuild.

diff --git a/Engine/Plugins/Importers/USDImporter/Source/USDSchemas/Private/USDGeomXformableTranslator.cpp b/Engine/Plugins/Importers/USDImporter/Source/USDSchemas/Private/USDGeomXformableTranslator.cpp
--- a/Engine/Plugins/Importers/USDImporter/Source/USDSchemas/Private/USDGeomXformableTranslator.cpp
+++ b/Engine/Plugins/Importers/USDImporter/Source/USDSchemas/Private/USDGeomXformableTranslator.cpp
@@ -325,42 +325,6 @@ USceneComponent* FUsdGeomXformableTranslator::CreateComponents()
 	// right away (like on FUsdGeomPointInstancerTranslator::CreateComponents)
 	UpdateComponents(SceneComponent);
 
-	// Handle material overrides for collapsed meshes. This can happen if we have two separate subtrees that collapse
-	// the same: A single static mesh will be shared between them and one of the task chains will manage to put their
-	// material assignments on the mesh directly. To ensure the correct materials for the second subtree, we need to
-	// set overrides.
-	// Note how we don't have to handle geometry caches in here as they're handled by the geometry cache translator now
-	if (UStaticMeshComponent* StaticMeshComponent = Cast<UStaticMeshComponent>(SceneComponent))
-	{
-		if (Context->PrimLinkCache)
-		{
-			UE::FSdfPath PrototypePath = GetPrototypePrimPath();
-			if (UStaticMesh* StaticMesh = Context->PrimLinkCache->GetSingleAssetForPrim<UStaticMesh>(PrototypePath))
-			{
-				TArray<UMaterialInterface*> ExistingAssignments;
-				for (FStaticMaterial& StaticMaterial : StaticMesh->GetStaticMaterials())
-				{
-					ExistingAssignments.Add(StaticMaterial.MaterialInterface);
-				}
-
-				MeshTranslationImpl::SetMaterialOverrides(
-					GetPrim(),
-					ExistingAssignments,
-					*StaticMeshComponent,
-					*Context->UsdAssetCache,
-					*Context->UsdInfoCache,
-					*Context->PrimLinkCache,
-					Context->Time,
-					Context->ObjectFlags,
-					Context->bAllowInterpretingLODs,
-					Context->RenderContext,
-					Context->MaterialPurpose,
-					Context->bShareAssetsForIdenticalPrims
-				);
-			}
-		}
-	}
-
 	return SceneComponent;
 }
 
@@ -786,6 +750,42 @@ void FUsdGeomXformableTranslator::UpdateComponents(USceneComponent* SceneCompone
 		{
 			SceneComponent->RegisterComponent();
 		}
+
+		// Handle material overrides for all kinds of meshes, even collapsed meshes.
+		// This can happen if we have two separate subtrees that collapse the same: A single static mesh will be shared between
+		// them and one of the task chains will manage to put their material assignments on the mesh directly. To ensure the
+		// correct materials for the second subtree, we need to set overrides.
+		// Note how we don't have to handle geometry caches in here as they're handled by the geometry cache translator now
+		if (Context->bAllowRecomputingMaterialOverrides)
+		{
+			if (UStaticMeshComponent* StaticMeshComponent = Cast<UStaticMeshComponent>(SceneComponent))
+			{
+				UE::FSdfPath PrototypePath = GetPrototypePrimPath();
+				if (UStaticMesh* StaticMesh = Context->PrimLinkCache->GetSingleAssetForPrim<UStaticMesh>(PrototypePath))
+				{
+					TArray<UMaterialInterface*> ExistingAssignments;
+					for (FStaticMaterial& StaticMaterial : StaticMesh->GetStaticMaterials())
+					{
+						ExistingAssignments.Add(StaticMaterial.MaterialInterface);
+					}
+
+					MeshTranslationImpl::SetMaterialOverrides(
+						GetPrim(),
+						ExistingAssignments,
+						*StaticMeshComponent,
+						*Context->UsdAssetCache,
+						*Context->UsdInfoCache,
+						*Context->PrimLinkCache,
+						Context->Time,
+						Context->ObjectFlags,
+						Context->bAllowInterpretingLODs,
+						Context->RenderContext,
+						Context->MaterialPurpose,
+						Context->bShareAssetsForIdenticalPrims
+					);
+				}
+			}
+		}
 	}
 }
 
diff --git a/Engine/Plugins/Importers/USDImporter/Source/USDSchemas/Private/USDGeometryCacheTranslator.cpp b/Engine/Plugins/Importers/USDImporter/Source/USDSchemas/Private/USDGeometryCacheTranslator.cpp
--- a/Engine/Plugins/Importers/USDImporter/Source/USDSchemas/Private/USDGeometryCacheTranslator.cpp
+++ b/Engine/Plugins/Importers/USDImporter/Source/USDSchemas/Private/USDGeometryCacheTranslator.cpp
@@ -1260,7 +1260,7 @@ TSet<UE::FSdfPath> FUsdGeometryCacheTranslator::CollectAuxiliaryPrims() const
 bool FUsdGeometryCacheTranslator::IsPotentialGeometryCacheRoot() const
 {
 	// The logic to check for GeometryCache is completely in the UsdInfoCache
-	return Context->UsdInfoCache->IsPotentialGeometryCacheRoot(PrimPath);
+	return Context->UsdInfoCache->IsPotentialGeometryCacheRoot(GetPrim());
 }
 
 #endif	  // #if USE_USD_SDK
diff --git a/Engine/Plugins/Importers/USDImporter/Source/USDSchemas/Private/USDInfoCache.cpp b/Engine/Plugins/Importers/USDImporter/Source/USDSchemas/Private/USDInfoCache.cpp
--- a/Engine/Plugins/Importers/USDImporter/Source/USDSchemas/Private/USDInfoCache.cpp
+++ b/Engine/Plugins/Importers/USDImporter/Source/USDSchemas/Private/USDInfoCache.cpp
@@ -82,59 +82,77 @@ namespace UE::UsdInfoCache::Private
 
 	struct FUsdPrimInfo
 	{
+		UE::FSdfPath PrimPath;
+
+		uint64 ParentInfoIndex = static_cast<uint64>(INDEX_NONE);
+		TArray<uint64> ChildIndices;
+
 		int32 PrimLockIndex = INDEX_NONE;
 
-		// Points to our collapsing root
-		// - Optional is not set: Prim wasn't visited on cache build. It's collapsed, but we don't know the root (yet);
-		// - Optional contains empty path: Prim was visited on build, we know it doesn't collapse and it isn't collapsed;
-		// - Optional contains prim's own path: This prim is a collapse root itself: It CollapsesChildren();
-		// - Optional contain other prim's path: This prim is collapsed, and that other prim is our collapse root (it CollapsesChildren());
-		TOptional<UE::FSdfPath> AssetCollapsedRoot;
+		// If this is true, it means this prim can and wants to collapse its entire subtree.
+		// If false, it either doesn't collapse its subtree, or we haven't visited it yet (same result)
+		bool bCollapsesChildren = false;
 
 		// Whether this prim can be collapsed or not, according to its schema translator
 		// - Optional is not set: Prim wasn't visited yet, we don't know
 		// - Optional has value: Whether the prim can be collapsed or not
 		TOptional<bool> bXformSubtreeCanBeCollapsed;
 
-		uint64 ExpectedVertexCountForSubtree = 0;
+		// This is used as a "visited" marker for RecursivePropagateVertexAndMaterialSlotCounts
+		TOptional<uint64> ExpectedVertexCountForSubtree;
 		TArray<UsdUtils::FUsdPrimMaterialSlot> SubtreeMaterialSlots;
+		bool bSlotsWereMerged = false;
 
 		int32 GeometryCacheDepth = -1;
 		EGeometryCachePrimState GeometryCacheState = EGeometryCachePrimState::None;
 
-		// Paths to material prims to the mesh prims they are bound to in the scene, given the current settings for
-		// render context, material purpose, variant selections, etc.
-		TSet<UE::FSdfPath> MaterialUsers;
-
 		// Maps from prims, to all the prims that require also reading this prim to be translated into an asset.
 		// Mainly used to update these assets whenever the depencency prim is updated.
+		TSet<UE::FSdfPath> MaterialUserMainPrims;
 		TSet<UE::FSdfPath> MainPrims;
 		TSet<UE::FSdfPath> AuxPrims;
+		bool bInstanceAuxPrimsRegistered = false;
 
-		FUsdPrimInfo() = default;
-		FUsdPrimInfo(int InPrimLockIndex)
-			: PrimLockIndex(InPrimLockIndex)
+		void ResetCollectedInfo()
 		{
+			bCollapsesChildren = false;
+
+			bXformSubtreeCanBeCollapsed.Reset();
+
+			ExpectedVertexCountForSubtree.Reset();
+			SubtreeMaterialSlots.Reset();
+			bSlotsWereMerged = false;
+
+			GeometryCacheDepth = -1;
+			GeometryCacheState = EGeometryCachePrimState::None;
+
+			bInstanceAuxPrimsRegistered = false;
 		}
 	};
 }	 // namespace UE::UsdInfoCache::Private
 
 FArchive& operator<<(FArchive& Ar, UE::UsdInfoCache::Private::FUsdPrimInfo& Info)
 {
+	Ar << Info.PrimPath;
+
+	Ar << Info.ParentInfoIndex;
+	Ar << Info.ChildIndices;
+
 	Ar << Info.PrimLockIndex;
 
-	Ar << Info.AssetCollapsedRoot;
+	Ar << Info.bCollapsesChildren;
 
 	Ar << Info.ExpectedVertexCountForSubtree;
 	Ar << Info.SubtreeMaterialSlots;
+	Ar << Info.bSlotsWereMerged;
 
 	Ar << Info.GeometryCacheDepth;
 	Ar << Info.GeometryCacheState;
 
-	Ar << Info.MaterialUsers;
-
+	Ar << Info.MaterialUserMainPrims;
 	Ar << Info.MainPrims;
 	Ar << Info.AuxPrims;
+	Ar << Info.bInstanceAuxPrimsRegistered;
 
 	return Ar;
 }
@@ -150,42 +168,42 @@ struct FUsdInfoCacheImpl
 		PrimLocks = new FRWLock[GNumPerPrimLocks]();
 	}
 
-	FUsdInfoCacheImpl(const FUsdInfoCacheImpl& Other)
-		: FUsdInfoCacheImpl()
-	{
-		FReadScopeLock ScopedInfoMapLock(Other.InfoMapLock);
-
-		FWriteScopeLock ThisScopedInfoMapLock(InfoMapLock);
-
-		InfoMap = Other.InfoMap;
-
-		AllowedExtensionsForGeometryCacheSource = Other.AllowedExtensionsForGeometryCacheSource;
-	}
-
 	FUsdInfoCacheImpl& operator=(const FUsdInfoCacheImpl& Other)
 	{
 		FReadScopeLock OtherScopedInfoMapLock(Other.InfoMapLock);
-
 		FWriteScopeLock ThisScopedInfoMapLock(InfoMapLock);
-
 		InfoMap = Other.InfoMap;
+		PrimInfoArray = Other.PrimInfoArray;
+		StaleInfoIndices = Other.StaleInfoIndices;
+
 		AllowedExtensionsForGeometryCacheSource = Other.AllowedExtensionsForGeometryCacheSource;
-		TempStage = Other.TempStage;
 		return *this;
 	}
 
 	~FUsdInfoCacheImpl()
 	{
+		FScopedUnrealAllocs Allocs;
 		delete[] PrimLocks;
 	}
 
 	// Information we must have about all prims on the stage
-	TMap<UE::FSdfPath, UE::UsdInfoCache::Private::FUsdPrimInfo> InfoMap;
+	TArray<UE::UsdInfoCache::Private::FUsdPrimInfo> PrimInfoArray;
+	TArray<uint64> StaleInfoIndices;
+	TMap<UE::FSdfPath, uint64> InfoMap;
 	mutable FRWLock InfoMapLock;
 
 	// Temporarily used during the info cache build, as we need to do another pass on point instancers afterwards
-	TArray<FString> PointInstancerPaths;
-	mutable FRWLock PointInstancerPathsLock;
+	TArray<FString> TempPointInstancerPaths;
+	mutable FRWLock TempPointInstancerPathsLock;
+
+	TSet<UE::FSdfPath> TempUpdatedPrototypes;
+	mutable FRWLock TempUpdatedPrototypesLock;
+
+	TMap<UE::FSdfPath, TSet<UE::FSdfPath>> TempStaleMaterialUserMainPrims;
+	mutable FRWLock TempStaleMaterialUserMainPrimsLock;
+
+	TSet<uint64> TempPrimsToDisconnectAuxLinks;
+	mutable FRWLock TempPrimsToDisconnectAuxLinksLock;
 
 	// This is used to keep track of which prototypes are already being translated within this "translation session",
 	// so that the schema translators can early out if they're trying to translate multiple instances of the same
@@ -196,15 +214,204 @@ struct FUsdInfoCacheImpl
 	// Geometry cache can come from a reference or payload of these file types
 	TArray<FString> AllowedExtensionsForGeometryCacheSource;
 
-	// Valid only during the main info cache build
-	UE::FUsdStageWeak TempStage;
-
-	// Individual locks distributed across the FUsdPrimInfo in the InfoMap
+private:
+	// Individual locks distributed across the FUsdPrimInfo
 	FRWLock* PrimLocks = nullptr;
 
 public:
+	void ClearTransientInfo()
+	{
+		{
+			FWriteScopeLock Lock(TempPointInstancerPathsLock);
+			TempPointInstancerPaths.Empty();
+		}
+		{
+			FWriteScopeLock Lock(TempUpdatedPrototypesLock);
+			TempUpdatedPrototypes.Empty();
+		}
+		{
+			FWriteScopeLock Lock(TempStaleMaterialUserMainPrimsLock);
+			TempStaleMaterialUserMainPrims.Empty();
+		}
+		{
+			FWriteScopeLock Lock(TempPrimsToDisconnectAuxLinksLock);
+			TempPrimsToDisconnectAuxLinks.Empty();
+		}
+	}
+
+	// WARNING: Assumes that the info map is locked for reading
+	UE::UsdInfoCache::Private::FUsdPrimInfo* GetPrimInfo(const UE::FSdfPath& PrimPath) const
+	{
+		if (const uint64* Index = InfoMap.Find(PrimPath))
+		{
+			return const_cast<UE::UsdInfoCache::Private::FUsdPrimInfo*>(&PrimInfoArray[*Index]);
+		}
+
+		return nullptr;
+	}
+
+	[[nodiscard]] FReadScopeLock LockForReading(const UE::UsdInfoCache::Private::FUsdPrimInfo& Info) const
+	{
+		return FReadScopeLock{PrimLocks[Info.PrimLockIndex]};
+	}
+
+	[[nodiscard]] FWriteScopeLock LockForWriting(const UE::UsdInfoCache::Private::FUsdPrimInfo& Info) const
+	{
+		return FWriteScopeLock{PrimLocks[Info.PrimLockIndex]};
+	}
+
+	UE::UsdInfoCache::Private::FUsdPrimInfo& CreateNewInfo(const UE::FSdfPath& PrimPath, uint64& OutNewIndex)
+	{
+		using namespace UE::UsdInfoCache::Private;
+
+		if (StaleInfoIndices.Num() > 0)
+		{
+			OutNewIndex = StaleInfoIndices.Pop();
+
+			FUsdPrimInfo& Result = PrimInfoArray[OutNewIndex];
+			Result = FUsdPrimInfo{};	// Get rid of old data since we're reusing an entry
+			Result.PrimPath = PrimPath;
+			return Result;
+		}
+		else
+		{
+			OutNewIndex = PrimInfoArray.Num();
+
+			FUsdPrimInfo& Result = PrimInfoArray.Emplace_GetRef();
+			Result.PrimPath = PrimPath;
+			return Result;
+		}
+	}
+
+	// This invalidates all data collected for a particular prim, but retains an info entry for
+	// that prim, as well as parent/child indices
+	//
+	// WARNING: This assumes the InfoMap is locked for at least reading.
+	// WARNING: Not thread safe (only called during the partial cleanup, which is single-threaded)
+	void ResetPrimInfoEntry(uint64 Index)
+	{
+		using namespace UE::UsdInfoCache::Private;
+
+		FUsdPrimInfo& Info = PrimInfoArray[Index];
+
+		// We're already reset, early out
+		if (!Info.ExpectedVertexCountForSubtree.IsSet())
+		{
+			return;
+		}
+
+		Info.ResetCollectedInfo();
+
+		// Remember to wipe our main/aux links later after we've done traversing them
+		TempPrimsToDisconnectAuxLinks.Add(Index);
+
+		// Propagate to main prims
+		for (const UE::FSdfPath& MainPath : Info.MainPrims)
+		{
+			if (uint64* MainPrimIndex = InfoMap.Find(MainPath))
+			{
+				ResetPrimInfoEntry(*MainPrimIndex);
+			}
+		}
+
+		// Propagate to ancestors
+		if (Info.ParentInfoIndex != INDEX_NONE)
+		{
+			ResetPrimInfoEntry(Info.ParentInfoIndex);
+		}
+	}
+
+	// Fully discards all data and all parent/child indices collected for a prim at InfoIndex.
+	// In practice this will just mark those entries/indices as stale so that we can reuse them later.
+	//
+	// WARNING: This assumes the InfoMap is locked for writing.
+	// WARNING: Not thread safe (only called during the partial cleanup, which is single-threaded)
+	void RemovePrimInfoSubtree(uint64 InfoIndex)
+	{
+		using namespace UE::UsdInfoCache::Private;
+
+		// We do this so that we recursively step through all main links as well.
+		// The main use case for this is point instancers with prototypes outside their own hierarchy. If we
+		// edit those prototypes, they may end up with a different vertex count or material slot count, so
+		// we need to update those counts for the point instancer's info struct (and so its ancestors) too.
+		// It may seem slow to do this recursively as we traverse into children, but these Main links usually
+		// point to ancestors anyway, and ResetPrimInfoEntry is quick to early out when we step into a visited
+		// info struct
+		ResetPrimInfoEntry(InfoIndex);
+
+		FUsdPrimInfo& Info = PrimInfoArray[InfoIndex];
+
+		for (uint64 ChildIndex : Info.ChildIndices)
+		{
+			FUsdPrimInfo& ChildInfo = PrimInfoArray[ChildIndex];
+			RemovePrimInfoSubtree(ChildIndex);
+		}
+
+		InfoMap.Remove(Info.PrimPath);
+		StaleInfoIndices.Push(InfoIndex);
+	};
+
+	// We have a separate function for this (not directly called from ResetPrimInfoEntry or other functions)
+	// because we want to unregister these aux links for a subtree only after we have have traversed the subtree
+	// resetting everything we need. Otherwise we may end up removing these aux links before we've had a chance
+	// to traverse them
+	//
+	// WARNING: This assumes the InfoMap is locked for at least reading.
+	// WARNING: Not thread safe (only called during the partial cleanup, which is single-threaded)
+	void DisconnectResetPrimAuxLinks()
+	{
+		using namespace UE::UsdInfoCache::Private;
+
+		for (uint64 Index : TempPrimsToDisconnectAuxLinks)
+		{
+			FUsdPrimInfo& Info = PrimInfoArray[Index];
+
+			// Disconnect our Aux links
+			for (const UE::FSdfPath& AuxPath : Info.AuxPrims)
+			{
+				if (FUsdPrimInfo* AuxPrimInfo = GetPrimInfo(AuxPath))
+				{
+					AuxPrimInfo->MainPrims.Remove(Info.PrimPath);
+				}
+			}
+			Info.AuxPrims.Reset();
+
+			// Disconnect our Main links
+			for (const UE::FSdfPath& MainPath : Info.MainPrims)
+			{
+				if (FUsdPrimInfo* MainPrimInfo = GetPrimInfo(MainPath))
+				{
+					MainPrimInfo->AuxPrims.Remove(Info.PrimPath);
+				}
+			}
+			Info.MainPrims.Reset();
+
+			// Disconnect our MaterialUserMain links
+			for (const UE::FSdfPath& MainPath : Info.MaterialUserMainPrims)
+			{
+				if (FUsdPrimInfo* MainPrimInfo = GetPrimInfo(MainPath))
+				{
+					MainPrimInfo->AuxPrims.Remove(Info.PrimPath);
+				}
+			}
+			if (Info.MaterialUserMainPrims.Num() > 0)
+			{
+				// Stash these before we reset them
+				//
+				// This because we can't recompute these easily if all we're resyncing is the material: The material users are just
+				// random prims in the stage. If we're resyncing just the material and the users are unmodified, then they really
+				// have the exact same material binding anyway, so we can use this data to quickly restore material users to these prim
+				// info structs, if they happen to match the same prim path of the resynced material prim
+				TempStaleMaterialUserMainPrims.FindOrAdd(Info.PrimPath).Append(Info.MaterialUserMainPrims);
+				Info.MaterialUserMainPrims.Reset();
+			}
+		}
+	}
+
 	void RegisterAuxiliaryPrims(const UE::FSdfPath& MainPrimPath, const TSet<UE::FSdfPath>& AuxPrimPaths)
 	{
+		using namespace UE::UsdInfoCache::Private;
+
 		TRACE_CPUPROFILER_EVENT_SCOPE(RegisterAuxiliaryPrims);
 
 		if (AuxPrimPaths.Num() == 0)
@@ -214,23 +421,105 @@ public:
 
 		FReadScopeLock ScopeLock{InfoMapLock};
 
-		if (UE::UsdInfoCache::Private::FUsdPrimInfo* MainPrim = InfoMap.Find(MainPrimPath))
+		if (FUsdPrimInfo* MainPrim = GetPrimInfo(MainPrimPath))
 		{
-			FWriteScopeLock PrimLock{PrimLocks[MainPrim->PrimLockIndex]};
+			FWriteScopeLock Lock = LockForWriting(*MainPrim);
 			MainPrim->AuxPrims.Append(AuxPrimPaths);
 		}
 
 		for (const UE::FSdfPath& AuxPrimPath : AuxPrimPaths)
 		{
-			if (UE::UsdInfoCache::Private::FUsdPrimInfo* AuxPrim = InfoMap.Find(AuxPrimPath))
+			if (FUsdPrimInfo* AuxPrim = GetPrimInfo(AuxPrimPath))
 			{
-				FWriteScopeLock PrimLock{PrimLocks[AuxPrim->PrimLockIndex]};
+				FWriteScopeLock Lock = LockForWriting(*AuxPrim);
 				AuxPrim->MainPrims.Add(MainPrimPath);
 			}
 
 			UE_LOG(LogUsd, Verbose, TEXT("Registering main prim '%s' and aux prim '%s'"), *MainPrimPath.GetString(), *AuxPrimPath.GetString());
 		}
 	}
+
+	void RegisterMaterialUserPrims(const UE::FSdfPath& MaterialPath, const TSet<UE::FSdfPath>& UserPaths)
+	{
+		using namespace UE::UsdInfoCache::Private;
+
+		TRACE_CPUPROFILER_EVENT_SCOPE(RegisterMaterialUserPrims);
+
+		if (UserPaths.Num() == 0)
+		{
+			return;
+		}
+
+		FReadScopeLock ScopeLock{InfoMapLock};
+
+		if (FUsdPrimInfo* Material = GetPrimInfo(MaterialPath))
+		{
+			FWriteScopeLock Lock = LockForWriting(*Material);
+			Material->MaterialUserMainPrims.Append(UserPaths);
+		}
+
+		for (const UE::FSdfPath& UserPath : UserPaths)
+		{
+			if (FUsdPrimInfo* User = GetPrimInfo(UserPath))
+			{
+				FWriteScopeLock Lock = LockForWriting(*User);
+				User->AuxPrims.Add(MaterialPath);
+			}
+
+			UE_LOG(LogUsd, Verbose, TEXT("Registering user '%s' of material prim '%s'"), *UserPath.GetString(), *MaterialPath.GetString());
+		}
+	}
+
+	void TryRestoreMaterialUserLinks(UE::UsdInfoCache::Private::FUsdPrimInfo& MaterialInfo)
+	{
+		using namespace UE::UsdInfoCache::Private;
+
+		FReadScopeLock ScopeLock{TempStaleMaterialUserMainPrimsLock};
+
+		const TSet<UE::FSdfPath>* OldMaterialUsers = nullptr;
+		UE::FSdfPath MaterialPrimPath;
+		{
+			FWriteScopeLock PrimLock = LockForWriting(MaterialInfo);
+
+			OldMaterialUsers = TempStaleMaterialUserMainPrims.Find(MaterialInfo.PrimPath);
+			if (!OldMaterialUsers || OldMaterialUsers->Num() == 0)
+			{
+				// We don't have any old users registered for this material prim path, just return
+				return;
+			}
+			MaterialPrimPath = MaterialInfo.PrimPath;
+
+			MaterialInfo.MaterialUserMainPrims.Append(*OldMaterialUsers);
+		}
+
+		for (const UE::FSdfPath& OldUserPath : *OldMaterialUsers)
+		{
+			if (FUsdPrimInfo* OldUser = GetPrimInfo(OldUserPath))
+			{
+				FWriteScopeLock Lock = LockForWriting(*OldUser);
+				OldUser->AuxPrims.Add(MaterialPrimPath);
+			}
+		}
+	}
+
+#if USE_USD_SDK
+	bool IsPotentialGeometryCacheRootInner(UE::UsdInfoCache::Private::FUsdPrimInfo& Info, const pxr::UsdPrim& Prim)
+	{
+		using namespace UE::UsdInfoCache::Private;
+
+		FWriteScopeLock PrimLock = LockForWriting(Info);
+
+		// When importing we fill all those in during the info cache initial build. If this is None still, it means
+		// we're in the default geometry cache workflow for opening the stage, where geometry caches are generated
+		// directly for single animated Mesh prims (so no collapsing of whole subtrees into geometry caches). We can
+		// then find out if our prim is animated on-demand
+		if (Info.GeometryCacheState == EGeometryCachePrimState::None)
+		{
+			Info.GeometryCacheState = UsdUtils::IsAnimatedMesh(Prim) ? EGeometryCachePrimState::ValidRoot : EGeometryCachePrimState::Uncollapsible;
+		}
+		return Info.GeometryCacheState == EGeometryCachePrimState::ValidRoot;
+	}
+#endif	  // USE_USD_SDK
 };
 
 FUsdInfoCache::FUsdInfoCache()
@@ -251,13 +540,12 @@ bool FUsdInfoCache::Serialize(FArchive& Ar)
 {
 	if (FUsdInfoCacheImpl* ImplPtr = Impl.Get())
 	{
+		TRACE_CPUPROFILER_EVENT_SCOPE(FUsdInfoCache::Serialize);
 		{
 			FWriteScopeLock ScopeLock(ImplPtr->InfoMapLock);
 			Ar << ImplPtr->InfoMap;
-		}
-		{
-			FWriteScopeLock ScopeLock(ImplPtr->TranslatedPrototypesLock);
-			Ar << ImplPtr->TranslatedPrototypes;
+			Ar << ImplPtr->PrimInfoArray;
+			Ar << ImplPtr->StaleInfoIndices;
 		}
 	}
 
@@ -275,6 +563,33 @@ bool FUsdInfoCache::ContainsInfoAboutPrim(const UE::FSdfPath& Path) const
 	return false;
 }
 
+TArray<UE::FSdfPath> FUsdInfoCache::GetChildren(const UE::FSdfPath& ParentPath) const
+{
+	using namespace UE::UsdInfoCache::Private;
+
+	TArray<UE::FSdfPath> Result;
+
+	if (FUsdInfoCacheImpl* ImplPtr = Impl.Get())
+	{
+		FReadScopeLock ScopeLock(ImplPtr->InfoMapLock);
+		if (const FUsdPrimInfo* FoundInfo = ImplPtr->GetPrimInfo(ParentPath))
+		{
+			FReadScopeLock PrimLock = ImplPtr->LockForReading(*FoundInfo);
+
+			Result.Reserve(FoundInfo->ChildIndices.Num());
+			for (uint64 ChildIndex : FoundInfo->ChildIndices)
+			{
+				const FUsdPrimInfo& ChildInfo = ImplPtr->PrimInfoArray[ChildIndex];
+				FReadScopeLock ChildLock = ImplPtr->LockForReading(ChildInfo);
+				Result.Add(ChildInfo.PrimPath);
+			}
+		}
+	}
+
+	return Result;
+}
+
+PRAGMA_DISABLE_DEPRECATION_WARNINGS
 TSet<UE::FSdfPath> FUsdInfoCache::GetKnownPrims() const
 {
 	TSet<UE::FSdfPath> Result;
@@ -289,19 +604,39 @@ TSet<UE::FSdfPath> FUsdInfoCache::GetKnownPrims() const
 
 	return Result;
 }
+PRAGMA_ENABLE_DEPRECATION_WARNINGS
 
 bool FUsdInfoCache::IsPathCollapsed(const UE::FSdfPath& Path, ECollapsingType CollapsingType) const
 {
+	using namespace UE::UsdInfoCache::Private;
+
 	if (FUsdInfoCacheImpl* ImplPtr = Impl.Get())
 	{
 		FReadScopeLock ScopeLock(ImplPtr->InfoMapLock);
 
-		if (const UE::UsdInfoCache::Private::FUsdPrimInfo* FoundInfo = ImplPtr->InfoMap.Find(Path))
+		if (uint64* PrimIndex = ImplPtr->InfoMap.Find(Path))
 		{
-			FReadScopeLock PrimLock{ImplPtr->PrimLocks[FoundInfo->PrimLockIndex]};
+			uint64 IterIndex = INDEX_NONE;
+			{
+				// We're only collapsed if a parent collapses us
+				const FUsdPrimInfo& Info = ImplPtr->PrimInfoArray[*PrimIndex];
+				FReadScopeLock PrimLock = ImplPtr->LockForReading(Info);
+				IterIndex = Info.ParentInfoIndex;
+			}
 
-			// See comment on the AssetCollapsedRoot member
-			return !FoundInfo->AssetCollapsedRoot.IsSet() || (!FoundInfo->AssetCollapsedRoot->IsEmpty() && FoundInfo->AssetCollapsedRoot != Path);
+			while (IterIndex != INDEX_NONE)
+			{
+				const FUsdPrimInfo& Info = ImplPtr->PrimInfoArray[IterIndex];
+				FReadScopeLock PrimLock = ImplPtr->LockForReading(Info);
+				if (Info.bCollapsesChildren)
+				{
+					return true;
+				}
+
+				IterIndex = Info.ParentInfoIndex;
+			}
+
+			return false;
 		}
 
 		// This should never happen: We should have cached the entire tree
@@ -313,17 +648,44 @@ bool FUsdInfoCache::IsPathCollapsed(const UE::FSdfPath& Path, ECollapsingType Co
 
 bool FUsdInfoCache::DoesPathCollapseChildren(const UE::FSdfPath& Path, ECollapsingType CollapsingType) const
 {
+	using namespace UE::UsdInfoCache::Private;
+
 	if (FUsdInfoCacheImpl* ImplPtr = Impl.Get())
 	{
 		FReadScopeLock ScopeLock(ImplPtr->InfoMapLock);
 
-		if (const UE::UsdInfoCache::Private::FUsdPrimInfo* FoundInfo = ImplPtr->InfoMap.Find(Path))
+		if (uint64* PrimIndex = ImplPtr->InfoMap.Find(Path))
 		{
-			FReadScopeLock PrimLock{ImplPtr->PrimLocks[FoundInfo->PrimLockIndex]};
+			uint64 IterIndex = INDEX_NONE;
+			{
+				const FUsdPrimInfo& Info = ImplPtr->PrimInfoArray[*PrimIndex];
+
+				FReadScopeLock PrimLock = ImplPtr->LockForReading(Info);
+				if (!Info.bCollapsesChildren)
+				{
+					// If this prim doesn't even want to collapse its children, we're done
+					return false;
+				}
 
-			// We store our own Path in there when we collapse children.
-			// Otherwise we hold the path of our collapse root, or empty (in case nothing is collapsed up to here)
-			return FoundInfo->AssetCollapsedRoot == Path;
+				IterIndex = Info.ParentInfoIndex;
+			}
+
+			// Even if this prim wants to collapse its children though, it could be that it's collapsed
+			// by a parent instead (collapsing is always done top-down)
+			while (IterIndex != INDEX_NONE)
+			{
+				const FUsdPrimInfo& AncestorInfo = ImplPtr->PrimInfoArray[IterIndex];
+
+				FReadScopeLock PrimLock = ImplPtr->LockForReading(AncestorInfo);
+				if (AncestorInfo.bCollapsesChildren)
+				{
+					return false;
+				}
+
+				IterIndex = AncestorInfo.ParentInfoIndex;
+			}
+
+			return true;
 		}
 
 		// This should never happen: We should have cached the entire tree
@@ -333,7 +695,7 @@ bool FUsdInfoCache::DoesPathCollapseChildren(const UE::FSdfPath& Path, ECollapsi
 	return false;
 }
 
-namespace UE::USDInfoCache::Private
+namespace UE::UsdInfoCache::Private
 {
 #if USE_USD_SDK
 	bool RecursiveQueryCanBeCollapsed(
@@ -343,16 +705,20 @@ namespace UE::USDInfoCache::Private
 		FUsdSchemaTranslatorRegistry& Registry
 	)
 	{
+		using namespace UE::UsdInfoCache::Private;
+
 		TRACE_CPUPROFILER_EVENT_SCOPE(UE::USDInfoCache::Private::RecursiveQueryCanBeCollapsed);
 
 		UE::FSdfPath UsdPrimPath = UE::FSdfPath{UsdPrim.GetPrimPath()};
 
+		FReadScopeLock ScopeLock{Impl.InfoMapLock};
+
 		// If we already have a value for our prim then we can just return it right now. We only fill these bCanBeCollapsed values
 		// through here, so if we know e.g. that UsdPrim can be collapsed, we know its entire subtree can too.
-		UE::UsdInfoCache::Private::FUsdPrimInfo* MainPrimInfo = Impl.InfoMap.Find(UsdPrimPath);
+		FUsdPrimInfo* MainPrimInfo = Impl.GetPrimInfo(UsdPrimPath);
 		if (MainPrimInfo)
 		{
-			FReadScopeLock PrimLock{Impl.PrimLocks[MainPrimInfo->PrimLockIndex]};
+			FReadScopeLock PrimLock = Impl.LockForReading(*MainPrimInfo);
 			if (MainPrimInfo->bXformSubtreeCanBeCollapsed.IsSet())
 			{
 				return MainPrimInfo->bXformSubtreeCanBeCollapsed.GetValue();
@@ -406,120 +772,109 @@ namespace UE::USDInfoCache::Private
 		}
 
 		// Record what we found about our main prim
+		if (MainPrimInfo)
 		{
-			FWriteScopeLock PrimLock{Impl.PrimLocks[MainPrimInfo->PrimLockIndex]};
+			FWriteScopeLock PrimLock = Impl.LockForWriting(*MainPrimInfo);
 			MainPrimInfo->bXformSubtreeCanBeCollapsed = bCanBeCollapsed;
 		}
 
 		// Before we return though, what we can do here is that if we know that we can't be collapsed ourselves,
 		// then none of our ancestors can either! So let's quickly paint upwards to make future queries faster
-		if (!bCanBeCollapsed)
+		if (!bCanBeCollapsed && MainPrimInfo)
 		{
-			UE::FSdfPath TraversalPath = UsdPrimPath.GetParentPath();
-			while (!TraversalPath.IsAbsoluteRootPath())
+			uint64 IterIndex = INDEX_NONE;
 			{
-				if (UE::UsdInfoCache::Private::FUsdPrimInfo* AncestorInfo = Impl.InfoMap.Find(TraversalPath))
-				{
-					FWriteScopeLock PrimLock{Impl.PrimLocks[AncestorInfo->PrimLockIndex]};
+				IterIndex = MainPrimInfo->ParentInfoIndex;
+			}
 
-					// We found something that was already filled out: Let's stop traversing here
-					if (AncestorInfo->bXformSubtreeCanBeCollapsed.IsSet())
-					{
-						// If we can't collapse ourselves then like we mentioned above none of our ancestors should
-						// be able to collapse either
-						ensure(AncestorInfo->bXformSubtreeCanBeCollapsed.GetValue() == false);
-						break;
-					}
-					else
-					{
-						AncestorInfo->bXformSubtreeCanBeCollapsed = false;
-					}
+			while (IterIndex != INDEX_NONE)
+			{
+				FUsdPrimInfo& AncestorInfo = Impl.PrimInfoArray[IterIndex];
+				FWriteScopeLock PrimLock = Impl.LockForWriting(*MainPrimInfo);
+
+				// We found something that was already filled out: Let's stop traversing here
+				if (AncestorInfo.bXformSubtreeCanBeCollapsed.IsSet())
+				{
+					// If we can't collapse ourselves then like we mentioned above none of our ancestors should
+					// be able to collapse either
+					ensure(AncestorInfo.bXformSubtreeCanBeCollapsed.GetValue() == false);
+					break;
+				}
+				else
+				{
+					AncestorInfo.bXformSubtreeCanBeCollapsed = false;
 				}
 
-				TraversalPath = TraversalPath.GetParentPath();
+				IterIndex = AncestorInfo.ParentInfoIndex;
 			}
 		}
 
 		return bCanBeCollapsed;
 	}
 #endif	  // USE_USD_SDK
-}	 // namespace UE::USDInfoCache::Private
+}	 // namespace UE::UsdInfoCache::Private
 
 UE::FSdfPath FUsdInfoCache::UnwindToNonCollapsedPath(const UE::FSdfPath& Path, ECollapsingType CollapsingType) const
 {
+	using namespace UE::UsdInfoCache::Private;
+
+	TRACE_CPUPROFILER_EVENT_SCOPE(FUsdInfoCache::UnwindToNonCollapsedPath);
+
 	if (FUsdInfoCacheImpl* ImplPtr = Impl.Get())
 	{
 		FReadScopeLock ScopeLock(ImplPtr->InfoMapLock);
 
-		TOptional<UE::FSdfPath> CollapseRoot;
-		UE::UsdInfoCache::Private::FUsdPrimInfo* MainFoundInfo = ImplPtr->InfoMap.Find(Path);
-		if (MainFoundInfo)
-		{
-			FReadScopeLock PrimLock{ImplPtr->PrimLocks[MainFoundInfo->PrimLockIndex]};
-			CollapseRoot = MainFoundInfo->AssetCollapsedRoot;
-		}
-
-		// We never visited this prim before. We know it's collapsed, let's find our collapse root
-		if (!CollapseRoot.IsSet())
+		if (uint64* PrimIndex = ImplPtr->InfoMap.Find(Path))
 		{
-			TArray<UE::UsdInfoCache::Private::FUsdPrimInfo*> InfosToUpdate = {MainFoundInfo};
-
-			UE::FSdfPath TraversalPath = Path.GetParentPath();
-			while (!TraversalPath.IsAbsoluteRootPath())
+			TFunction<TOptional<UE::FSdfPath>(uint64)> GetCollapseRootFromParent;
+			GetCollapseRootFromParent = [ImplPtr, &GetCollapseRootFromParent](uint64 Index) -> TOptional<UE::FSdfPath>
 			{
-				if (UE::UsdInfoCache::Private::FUsdPrimInfo* AncestorInfo = ImplPtr->InfoMap.Find(TraversalPath))
+				if (Index == INDEX_NONE)
 				{
-					FReadScopeLock PrimLock{ImplPtr->PrimLocks[AncestorInfo->PrimLockIndex]};
+					return {};
+				}
 
-					// We found an ancestor that has this filled in: We're collapsed, so whatever is its
-					// collapse root is also our collapse root
-					if (AncestorInfo->AssetCollapsedRoot.IsSet())
-					{
-						// If our original Path doesn't have anything filled in, then we *must* be a child of a
-						// collapsed prim (i.e. something that has a non-empty path in its AssetCollapsedRoot)
-						ensure(!AncestorInfo->AssetCollapsedRoot->IsEmpty());
+				const FUsdPrimInfo& Info = ImplPtr->PrimInfoArray[Index];
 
-						CollapseRoot = AncestorInfo->AssetCollapsedRoot;
-						break;
-					}
-					// Still nothing. Let's keep track of this info so that we can update it later
-					else
-					{
-						InfosToUpdate.Add(AncestorInfo);
-					}
+				uint64 ParentIndex = INDEX_NONE;
+				{
+					FReadScopeLock PrimLock = ImplPtr->LockForReading(Info);
+					ParentIndex = Info.ParentInfoIndex;
 				}
 
-				TraversalPath = TraversalPath.GetParentPath();
-			}
+				TOptional<UE::FSdfPath> CollapseRootPath = GetCollapseRootFromParent(ParentIndex);
+				if (CollapseRootPath.IsSet())
+				{
+					// Our parent says it is collapsed with this collapse root: That's going to be
+					// the collapse root for our children too
+					return CollapseRootPath;
+				}
 
-			// Fill in all visited infos with what we found on our ancestor
-			if (CollapseRoot.IsSet())
-			{
-				for (UE::UsdInfoCache::Private::FUsdPrimInfo* Info : InfosToUpdate)
+				FReadScopeLock PrimLock = ImplPtr->LockForReading(Info);
+				if (Info.bCollapsesChildren)
 				{
-					FWriteScopeLock PrimLock{ImplPtr->PrimLocks[Info->PrimLockIndex]};
-					Info->AssetCollapsedRoot = CollapseRoot;
+					// We are the collapse root, let's return this to our children
+					return Info.PrimPath;
 				}
-			}
-		}
 
-		// We have visited this prim during the info cache build (or another Unwind, or just now within this function)
-		if (CollapseRoot.IsSet())
-		{
-			// An empty path here means that we are not collapsed at all
-			if (CollapseRoot->IsEmpty())
-			{
-				return Path;
-			}
-			// Otherwise we have our own path in there (in case we collapse children) or the path to the prim that collapsed us
-			else
+				// Nothing collapses so far
+				return {};
+			};
+
+			TOptional<UE::FSdfPath> CollapseRoot = GetCollapseRootFromParent(*PrimIndex);
+			if (CollapseRoot.IsSet())
 			{
 				return CollapseRoot.GetValue();
 			}
+
+			// We're not being collapsed by anything, so we're already the "non collapsed path"
+			const FUsdPrimInfo& Info = ImplPtr->PrimInfoArray[*PrimIndex];
+			FReadScopeLock PrimLock = ImplPtr->LockForReading(Info);
+			return Info.PrimPath;
 		}
 
 		// This should never happen: We should have cached the entire tree
-		ensureAlwaysMsgf(false, TEXT("Prim path '%s' has not been cached!"), *Path.GetString());
+		ensureMsgf(false, TEXT("Prim path '%s' has not been cached!"), *Path.GetString());
 	}
 
 	return Path;
@@ -531,9 +886,9 @@ TSet<UE::FSdfPath> FUsdInfoCache::GetMainPrims(const UE::FSdfPath& AuxPrimPath)
 	{
 		FReadScopeLock ScopeLock(ImplPtr->InfoMapLock);
 
-		if (const UE::UsdInfoCache::Private::FUsdPrimInfo* FoundInfo = ImplPtr->InfoMap.Find(AuxPrimPath))
+		if (const UE::UsdInfoCache::Private::FUsdPrimInfo* FoundInfo = ImplPtr->GetPrimInfo(AuxPrimPath))
 		{
-			FReadScopeLock PrimLock{ImplPtr->PrimLocks[FoundInfo->PrimLockIndex]};
+			FReadScopeLock PrimLock = Impl->LockForReading(*FoundInfo);
 
 			TSet<UE::FSdfPath> Result = FoundInfo->MainPrims;
 			Result.Add(AuxPrimPath);
@@ -550,9 +905,9 @@ TSet<UE::FSdfPath> FUsdInfoCache::GetAuxiliaryPrims(const UE::FSdfPath& MainPrim
 	{
 		FReadScopeLock ScopeLock(ImplPtr->InfoMapLock);
 
-		if (const UE::UsdInfoCache::Private::FUsdPrimInfo* FoundInfo = ImplPtr->InfoMap.Find(MainPrimPath))
+		if (const UE::UsdInfoCache::Private::FUsdPrimInfo* FoundInfo = ImplPtr->GetPrimInfo(MainPrimPath))
 		{
-			FReadScopeLock PrimLock{ImplPtr->PrimLocks[FoundInfo->PrimLockIndex]};
+			FReadScopeLock PrimLock = Impl->LockForReading(*FoundInfo);
 
 			TSet<UE::FSdfPath> Result = FoundInfo->AuxPrims;
 			Result.Add(MainPrimPath);
@@ -569,11 +924,10 @@ TSet<UE::FSdfPath> FUsdInfoCache::GetMaterialUsers(const UE::FSdfPath& Path) con
 	{
 		FReadScopeLock ScopeLock(ImplPtr->InfoMapLock);
 
-		if (const UE::UsdInfoCache::Private::FUsdPrimInfo* FoundInfo = ImplPtr->InfoMap.Find(Path))
+		if (const UE::UsdInfoCache::Private::FUsdPrimInfo* FoundInfo = ImplPtr->GetPrimInfo(Path))
 		{
-			FReadScopeLock PrimLock{ImplPtr->PrimLocks[FoundInfo->PrimLockIndex]};
-
-			return FoundInfo->MaterialUsers;
+			FReadScopeLock PrimLock = Impl->LockForReading(*FoundInfo);
+			return FoundInfo->MaterialUserMainPrims;
 		}
 	}
 
@@ -586,15 +940,14 @@ bool FUsdInfoCache::IsMaterialUsed(const UE::FSdfPath& Path) const
 	{
 		FReadScopeLock ScopeLock(ImplPtr->InfoMapLock);
 
-		if (const UE::UsdInfoCache::Private::FUsdPrimInfo* FoundInfo = ImplPtr->InfoMap.Find(Path))
+		if (const UE::UsdInfoCache::Private::FUsdPrimInfo* FoundInfo = ImplPtr->GetPrimInfo(Path))
 		{
-			FReadScopeLock PrimLock{ImplPtr->PrimLocks[FoundInfo->PrimLockIndex]};
-
-			return FoundInfo->MaterialUsers.Num() > 0;
+			FReadScopeLock PrimLock = Impl->LockForReading(*FoundInfo);
+			return FoundInfo->MaterialUserMainPrims.Num() > 0;
 		}
 	}
 
-	return false;
+	return {};
 }
 
 namespace UE::USDInfoCache::Private
@@ -675,11 +1028,11 @@ namespace UE::USDInfoCache::Private
 
 						// If we're calling this for a point instancer we should have parsed the results for our
 						// prototype subtrees already
-						if (const UE::UsdInfoCache::Private::FUsdPrimInfo* FoundInfo = Impl.InfoMap.Find(UE::FSdfPath{PrototypePath}))
+						if (const UE::UsdInfoCache::Private::FUsdPrimInfo* FoundInfo = Impl.GetPrimInfo(UE::FSdfPath{PrototypePath}))
 						{
-							FReadScopeLock PrimLock{Impl.PrimLocks[FoundInfo->PrimLockIndex]};
+							FReadScopeLock PrimLock = Impl.LockForReading(*FoundInfo);
 
-							PrototypeVertexCounts[PrototypeIndex] = FoundInfo->ExpectedVertexCountForSubtree;
+							PrototypeVertexCounts[PrototypeIndex] = FoundInfo->ExpectedVertexCountForSubtree.GetValue();
 							OutMaterialSlots.Append(FoundInfo->SubtreeMaterialSlots);
 						}
 					}
@@ -700,25 +1053,145 @@ namespace UE::USDInfoCache::Private
 		}
 	}
 
-	void RepopulateInfoMap(const pxr::UsdPrim& UsdPrim, FUsdInfoCacheImpl& Impl)
+	void CleanUpInfoMapSubtree(const UE::FSdfPath& PrimPath, FUsdInfoCacheImpl& Impl)
+	{
+		using namespace UE::UsdInfoCache::Private;
+
+		TRACE_CPUPROFILER_EVENT_SCOPE(CleanUpInfoMapSubtree);
+
+		FWriteScopeLock ScopeLock(Impl.InfoMapLock);
+
+		const uint64* ExistingPrimIndex = Impl.InfoMap.Find(PrimPath);
+
+		uint64 ExistingSubtreeParentIndex = INDEX_NONE;
+		if (ExistingPrimIndex)
+		{
+			FUsdPrimInfo& Info = Impl.PrimInfoArray[*ExistingPrimIndex];
+			FReadScopeLock InfoLock = Impl.LockForReading(Info);
+			ExistingSubtreeParentIndex = Info.ParentInfoIndex;
+		}
+
+		// Invalidate ancestors
+		{
+			uint64 AncestorIndex = ExistingSubtreeParentIndex;
+
+			// Even if we don't have an existing subtree to remove (e.g. when resyncing and adding a brand new subtree)
+			// we still need to find where this new subtree would "attach" and clean up those ancestors too,
+			// as we may add new info structs later when repopulating that would invalidate their collected info
+			if (!ExistingPrimIndex)
+			{
+				UE::FSdfPath IterPath = PrimPath.GetParentPath();
+				while (!IterPath.IsEmpty())
+				{
+					if (const uint64* FoundAncestor = Impl.InfoMap.Find(IterPath))
+					{
+						AncestorIndex = *FoundAncestor;
+						break;
+					}
+					IterPath = IterPath.GetParentPath();
+				}
+			}
+
+			// Actually invalidate ancestors from the one we found all the way up to the root
+			if (AncestorIndex != INDEX_NONE)
+			{
+				Impl.ResetPrimInfoEntry(AncestorIndex);
+			}
+		}
+
+		// Remove the prim subtree, if it exists
+		if (ExistingPrimIndex)
+		{
+			// Disconnect our invalidated subtree from its parent info (PrimPath could actually point at a prim
+			// that has been fully removed from the stage)
+			if (ExistingSubtreeParentIndex != INDEX_NONE)
+			{
+				FUsdPrimInfo& ParentInfo = Impl.PrimInfoArray[ExistingSubtreeParentIndex];
+				FWriteScopeLock ParentLock = Impl.LockForWriting(ParentInfo);
+				ParentInfo.ChildIndices.Remove(*ExistingPrimIndex);
+			}
+
+			Impl.RemovePrimInfoSubtree(*ExistingPrimIndex);
+		}
+	}
+
+	void RepopulateInfoMapSubtree(const UE::FSdfPath& SubtreeRootPath, const UE::FUsdStage& Stage, FUsdInfoCacheImpl& Impl)
 	{
 		using namespace UE::UsdInfoCache::Private;
 
-		TRACE_CPUPROFILER_EVENT_SCOPE(RepopulateInfoMap);
+		TRACE_CPUPROFILER_EVENT_SCOPE(RepopulateInfoMapSubtree);
+
+		if (!Stage)
+		{
+			return;
+		}
+
+		UE::FUsdPrim SubtreeRootPrim = Stage.GetPrimAtPath(SubtreeRootPath);
+		if (!SubtreeRootPrim)
+		{
+			// It's possible to be called with paths to prims that don't exist on the stage, for example when handling the
+			// rebuild about removing a prim spec, where USD sends a resync notice for the path to the prim that was just removed.
+			// We still want the info cache build to do all the rest in those cases though (cleanup old entries, invalidate ancestors,
+			// etc.), so we handle ignoring this case only in here
+			return;
+		}
+
+		const bool bIsPartialBuild = SubtreeRootPath.IsAbsoluteRootPath();
+
+		TFunction<uint64(const pxr::UsdPrim&, uint64)> ConstructInfoForPrim;
+		ConstructInfoForPrim = [&Impl, &ConstructInfoForPrim, bIsPartialBuild](const pxr::UsdPrim& Prim, uint64 ParentIndex) -> uint64
+		{
+			// If we're affecting an instance, record that we need to visit its prototype later
+			if (bIsPartialBuild && Prim.IsInstance())
+			{
+				UE::FUsdPrim Prototype{Prim.GetPrototype()};
+				FWriteScopeLock ScopeLock(Impl.TempUpdatedPrototypesLock);
+				Impl.TempUpdatedPrototypes.Add(Prototype.GetPrimPath());
+			}
+
+			// Note: We're not locking the infos here at all as our access pattern will never touch the same info more than once anyway,
+			// and this function is single threaded and never calls in to any other thread-unsafe functions
+			uint64 NewIndex = INDEX_NONE;
+			FUsdPrimInfo& NewInfo = Impl.CreateNewInfo(UE::FSdfPath{Prim.GetPrimPath()}, NewIndex);
+			NewInfo.PrimLockIndex = static_cast<int32>(NewIndex % GNumPerPrimLocks);
+			NewInfo.ParentInfoIndex = ParentIndex;
+
+			Impl.InfoMap.Add(NewInfo.PrimPath, NewIndex);
+
+			// Note: I've tried a ParallelFor here, and it was slower than the single threaded version due to write lock
+			// contention on the InfoMap itself
+			TArray<uint64> ChildIndices;
+			pxr::UsdPrimSiblingRange Children = Prim.GetFilteredChildren(pxr::UsdTraverseInstanceProxies(pxr::UsdPrimAllPrimsPredicate));
+			for (const pxr::UsdPrim& Child : Children)
+			{
+				const uint64 ChildIndex = ConstructInfoForPrim(Child, NewIndex);
+				ChildIndices.Add(ChildIndex);
+			}
+			// Have to find our NewInfo again as the recursive calls likely invalidated our NewInfo reference
+			Impl.PrimInfoArray[NewIndex].ChildIndices = MoveTemp(ChildIndices);
+
+			return NewIndex;
+		};
 
 		FWriteScopeLock ScopeLock(Impl.InfoMapLock);
 
-		uint64 PrimCounter = 0;
+		// Find parent
+		UE::FSdfPath ParentPrimPath = SubtreeRootPath.GetParentPath();
+		uint64* ParentIndexPtr = Impl.InfoMap.Find(ParentPrimPath);
+		uint64 ParentIndex = ParentIndexPtr ? *ParentIndexPtr : INDEX_NONE;
 
-		pxr::UsdPrimRange PrimRange{UsdPrim, pxr::UsdTraverseInstanceProxies(pxr::UsdPrimAllPrimsPredicate)};
-		for (pxr::UsdPrim Child : PrimRange)
+		// Create new subtree
+		uint64 SubtreeRootIndex = ConstructInfoForPrim(SubtreeRootPrim, ParentIndex);
+
+		// Connect the new subtree to its target parent
+		if (ParentIndex != INDEX_NONE)
 		{
-			Impl.InfoMap.Emplace(Child.GetPrimPath(), FUsdPrimInfo{static_cast<int32>(PrimCounter++ % GNumPerPrimLocks)});
+			Impl.PrimInfoArray[ParentIndex].ChildIndices.Add(SubtreeRootIndex);
 		}
 	}
 
 	void RecursivePropagateVertexAndMaterialSlotCounts(
-		const pxr::UsdPrim& UsdPrim,
+		uint64 PrimIndex,
 		FUsdSchemaTranslationContext& Context,
 		const pxr::TfToken& MaterialPurposeToken,
 		FUsdInfoCacheImpl& Impl,
@@ -730,20 +1203,36 @@ namespace UE::USDInfoCache::Private
 	{
 		TRACE_CPUPROFILER_EVENT_SCOPE(RecursivePropagateVertexAndMaterialSlotCounts);
 
+		FScopedUsdAllocs Allocs;
+
+		FReadScopeLock ScopeLock{Impl.InfoMapLock};
+		UE::UsdInfoCache::Private::FUsdPrimInfo& Info = Impl.PrimInfoArray[PrimIndex];
+
+		// Don't bother locking for reading PrimPath: We only ever write to this before this stage of the build
+		pxr::UsdPrim UsdPrim = Context.Stage.GetPrimAtPath(Info.PrimPath);
 		if (!UsdPrim)
 		{
 			return;
 		}
 
-		FScopedUsdAllocs Allocs;
+		// We already visited this subtree.
+		// Note that we don't need to check for MaterialUsers, material slots, or worry about point instancers here:
+		// Since we fill all of those for every rebuild, if ExpectedVertexCountForSubtree is filled so is everything else
+		{
+			FReadScopeLock PrimLock = Impl.LockForReading(Info);
+			if (Info.ExpectedVertexCountForSubtree.IsSet())
+			{
+				OutSubtreeVertexCount = Info.ExpectedVertexCountForSubtree.GetValue();
+				OutSubtreeSlots = Info.SubtreeMaterialSlots;
+				return;
+			}
+		}
 
-		pxr::SdfPath UsdPrimPath = UsdPrim.GetPrimPath();
-		UE::FSdfPath PrimPath{UsdPrimPath};
 		pxr::UsdStageRefPtr Stage = UsdPrim.GetStage();
-
-		TFunction<void(const UE::FSdfPath&, TSet<UE::FSdfPath>&)> TryAddMaterialUser = [&Stage](const UE::FSdfPath& Path, TSet<UE::FSdfPath>& Users)
+		TFunction<void(const UE::FSdfPath&, TSet<UE::FSdfPath>&)> TryAddMaterialUser =
+			[&Stage](const UE::FSdfPath& User, TSet<UE::FSdfPath>& Material)
 		{
-			pxr::UsdPrim UserPrim = Stage->GetPrimAtPath(Path);
+			pxr::UsdPrim UserPrim = Stage->GetPrimAtPath(User);
 
 			if (UserPrim.IsA<pxr::UsdGeomImageable>())
 			{
@@ -751,7 +1240,7 @@ namespace UE::USDInfoCache::Private
 				// literally *all prims* below an included prim path. That's fine and it really does mean that any Mesh prim
 				// in there could use the collection-based material binding, but nevertheless we don't want to register that
 				// e.g. Shader prims or SkelAnimation prims are "material users"
-				Users.Add(Path);
+				Material.Add(User);
 			}
 			else if (UserPrim.IsA<pxr::UsdGeomSubset>())
 			{
@@ -761,7 +1250,7 @@ namespace UE::USDInfoCache::Private
 				// UsdGeomSubset, since the GeomMesh translator doesn't collapse. We'll unwind this path later when fetching material
 				// users, so collapsed static meshes are handled OK, skeletal meshes are handled OK, we just need this one exception
 				// for handling uncollapsed static meshes, because by default Mesh prims don't "collapse" their child UsdGeomSubsets
-				Users.Add(Path.GetParentPath());
+				Material.Add(User.GetParentPath());
 			}
 		};
 
@@ -771,7 +1260,13 @@ namespace UE::USDInfoCache::Private
 		// list of all the prims that they should apply to when we call ComputeIncludedPaths
 		bool bPrimHasInheritableMaterialBindings = false;
 
-		// Register material users
+		// Try restoring old material users for this prim if it's a Material
+		if (UsdPrim.IsA<pxr::UsdShadeMaterial>())
+		{
+			Impl.TryRestoreMaterialUserLinks(Info);
+		}
+
+		// Register material users for other Material prims if this prim is a user
 		if (!UsdPrim.IsPseudoRoot())
 		{
 			TRACE_CPUPROFILER_EVENT_SCOPE(CheckingMaterialUsers);
@@ -820,7 +1315,7 @@ namespace UE::USDInfoCache::Private
 									Warning,
 									TEXT("Prim '%s' describes a collection-based material binding, but the relationship '%s' is invalid: It should "
 										 "contain exactly one Material path and one path to a collection relationship"),
-									*PrimPath.GetString(),
+									*Info.PrimPath.GetString(),
 									*UsdToUnreal::ConvertToken(Rel.GetName())
 								);
 								continue;
@@ -843,7 +1338,7 @@ namespace UE::USDInfoCache::Private
 									Warning,
 									TEXT("Failed to find collection at path '%s' when processing collection-based material bindings on prim '%s'"),
 									*UsdToUnreal::ConvertPath(CollectionPath->GetPrimPath()),
-									*PrimPath.GetString()
+									*Info.PrimPath.GetString()
 								);
 							}
 						}
@@ -856,7 +1351,7 @@ namespace UE::USDInfoCache::Private
 					bPrimHasInheritableMaterialBindings = true;
 
 					TSet<UE::FSdfPath>& MaterialUsers = NewMaterialUsers.FindOrAdd(UE::FSdfPath{ShadeMaterial.GetPrim().GetPath()});
-					TryAddMaterialUser(UE::FSdfPath{UsdPrimPath}, MaterialUsers);
+					TryAddMaterialUser(Info.PrimPath, MaterialUsers);
 				}
 			}
 			// Temporary fallback for prims that don't have the MaterialBindingAPI but do have the relationship.
@@ -875,37 +1370,22 @@ namespace UE::USDInfoCache::Private
 						bPrimHasInheritableMaterialBindings = true;
 
 						TSet<UE::FSdfPath>& MaterialUsers = NewMaterialUsers.FindOrAdd(UE::FSdfPath{TargetMaterialPrimPath});
-						TryAddMaterialUser(UE::FSdfPath{UsdPrimPath}, MaterialUsers);
+						TryAddMaterialUser(Info.PrimPath, MaterialUsers);
 					}
 				}
 			}
 
 			for (const TPair<UE::FSdfPath, TSet<UE::FSdfPath>>& NewMaterialToUsers : NewMaterialUsers)
 			{
-				FReadScopeLock ScopeLock(Impl.InfoMapLock);
-				if (UsdInfoCache::Private::FUsdPrimInfo* Info = Impl.InfoMap.Find(NewMaterialToUsers.Key))
-				{
-					FWriteScopeLock PrimLock{Impl.PrimLocks[Info->PrimLockIndex]};
-					Info->MaterialUsers.Append(NewMaterialToUsers.Value);
-				}
-			}
-		}
-
-		TArray<pxr::UsdPrim> Prims;
-		{
-			TRACE_CPUPROFILER_EVENT_SCOPE(CollectingChildren);
-			pxr::UsdPrimSiblingRange PrimChildren = UsdPrim.GetFilteredChildren(pxr::UsdTraverseInstanceProxies(pxr::UsdPrimAllPrimsPredicate));
-
-			for (pxr::UsdPrim Child : PrimChildren)
-			{
-				Prims.Emplace(Child);
+				Impl.RegisterMaterialUserPrims(NewMaterialToUsers.Key, NewMaterialToUsers.Value);
 			}
 		}
 
-		const uint32 NumChildren = Prims.Num();
+		const uint32 NumChildren = Info.ChildIndices.Num();
 
 		TArray<uint64> ChildSubtreeVertexCounts;
-		ChildSubtreeVertexCounts.SetNumUninitialized(NumChildren);
+		ChildSubtreeVertexCounts.SetNumZeroed(NumChildren);	   // Zero instead of Uninitialized here because if we run into e.g a Material prim our
+															   // code will mostly early out, and we don't want to return an uninitialized int
 
 		TArray<TArray<UsdUtils::FUsdPrimMaterialSlot>> ChildSubtreeMaterialSlots;
 		ChildSubtreeMaterialSlots.SetNum(NumChildren);
@@ -913,12 +1393,12 @@ namespace UE::USDInfoCache::Private
 		const int32 MinBatchSize = 1;
 		ParallelFor(
 			TEXT("RecursivePropagateVertexAndMaterialSlotCounts"),
-			Prims.Num(),
+			Info.ChildIndices.Num(),
 			MinBatchSize,
 			[&](int32 Index)
 			{
 				RecursivePropagateVertexAndMaterialSlotCounts(
-					Prims[Index],
+					Info.ChildIndices[Index],
 					Context,
 					MaterialPurposeToken,
 					Impl,
@@ -976,28 +1456,22 @@ namespace UE::USDInfoCache::Private
 		}
 
 		{
-			TRACE_CPUPROFILER_EVENT_SCOPE(StoringCounts);
-
-			FReadScopeLock ScopeLock(Impl.InfoMapLock);
-			if (UE::UsdInfoCache::Private::FUsdPrimInfo* Info = Impl.InfoMap.Find(UE::FSdfPath{UsdPrimPath}))
+			// For point instancers we can't guarantee we parsed the prototypes yet because they
+			// could technically be anywhere, so store them here for a later pass
+			if (bIsPointInstancer)
 			{
-				// For point instancers we can't guarantee we parsed the prototypes yet because they
-				// could technically be anywhere, so store them here for a later pass
-				if (bIsPointInstancer)
-				{
-					FWriteScopeLock PointInstancerLock(Impl.PointInstancerPathsLock);
-					Impl.PointInstancerPaths.Emplace(UE::FSdfPath{UsdPrimPath}.GetString());
-				}
-				// While we will compute the totals for any and all children normally, don't just append the regular
-				// traversal vertex count to the point instancer prim itself just yet, as that doesn't really represent
-				// what will happen. We'll later do another pass to handle point instancers where we'll properly instance
-				// stuff, and then we'll updadte all ancestors
-				else
-				{
-					FWriteScopeLock PrimLock{Impl.PrimLocks[Info->PrimLockIndex]};
-					Info->ExpectedVertexCountForSubtree = OutSubtreeVertexCount;
-					Info->SubtreeMaterialSlots.Append(OutSubtreeSlots);
-				}
+				FWriteScopeLock PointInstancerLock(Impl.TempPointInstancerPathsLock);
+				Impl.TempPointInstancerPaths.Emplace(Info.PrimPath.GetString());
+			}
+			// While we will compute the totals for any and all children normally, don't just append the regular
+			// traversal vertex count to the point instancer prim itself just yet, as that doesn't really represent
+			// what will happen. We'll later do another pass to handle point instancers where we'll properly instance
+			// stuff, and then we'll updadte all ancestors
+			else
+			{
+				FWriteScopeLock PrimLock = Impl.LockForWriting(Info);
+				Info.ExpectedVertexCountForSubtree = OutSubtreeVertexCount;
+				Info.SubtreeMaterialSlots.Append(OutSubtreeSlots);
 			}
 		}
 	}
@@ -1015,8 +1489,7 @@ namespace UE::USDInfoCache::Private
 	{
 		TRACE_CPUPROFILER_EVENT_SCOPE(UpdateInfoForPointInstancers);
 
-		pxr::UsdStageWeakPtr UsdStageWeak = Impl.TempStage;
-		pxr::UsdStageRefPtr Stage = UsdStageWeak;
+		pxr::UsdStageRefPtr Stage = pxr::UsdStageRefPtr{Context.Stage};
 		if (!Stage)
 		{
 			return;
@@ -1062,12 +1535,12 @@ namespace UE::USDInfoCache::Private
 			return LHS < RHS;
 		};
 		{
-			FWriteScopeLock PointInstancerLock{Impl.PointInstancerPathsLock};
-			Impl.PointInstancerPaths.Sort(SortFunction);
+			FWriteScopeLock PointInstancerLock{Impl.TempPointInstancerPathsLock};
+			Impl.TempPointInstancerPaths.Sort(SortFunction);
 		}
 
-		FReadScopeLock PointInstancerLock{Impl.PointInstancerPathsLock};
-		for (const FString& PointInstancerPath : Impl.PointInstancerPaths)
+		FReadScopeLock PointInstancerLock{Impl.TempPointInstancerPathsLock};
+		for (const FString& PointInstancerPath : Impl.TempPointInstancerPaths)
 		{
 			UE::FSdfPath UsdPointInstancerPath{*PointInstancerPath};
 
@@ -1079,10 +1552,10 @@ namespace UE::USDInfoCache::Private
 				GetPrimVertexCountAndSlots(PointInstancer, Context, Impl, PointInstancerVertexCount, PointInstancerMaterialSlots);
 
 				FReadScopeLock Lock{Impl.InfoMapLock};
-				if (UE::UsdInfoCache::Private::FUsdPrimInfo* Info = Impl.InfoMap.Find(UsdPointInstancerPath))
+				if (UE::UsdInfoCache::Private::FUsdPrimInfo* Info = Impl.GetPrimInfo(UsdPointInstancerPath))
 				{
 					{
-						FWriteScopeLock PrimLock{Impl.PrimLocks[Info->PrimLockIndex]};
+						FWriteScopeLock PrimLock = Impl.LockForWriting(*Info);
 						Info->ExpectedVertexCountForSubtree = PointInstancerVertexCount;
 						Info->SubtreeMaterialSlots.Append(PointInstancerMaterialSlots);
 					}
@@ -1102,10 +1575,10 @@ namespace UE::USDInfoCache::Private
 							break;
 						}
 
-						if (UE::UsdInfoCache::Private::FUsdPrimInfo* ParentInfo = Impl.InfoMap.Find(ParentPath))
+						if (UE::UsdInfoCache::Private::FUsdPrimInfo* ParentInfo = Impl.GetPrimInfo(ParentPath))
 						{
-							FWriteScopeLock ParentPrimLock{Impl.PrimLocks[ParentInfo->PrimLockIndex]};
-							ParentInfo->ExpectedVertexCountForSubtree += PointInstancerVertexCount;
+							FWriteScopeLock PrimLock = Impl.LockForWriting(*ParentInfo);
+							ParentInfo->ExpectedVertexCountForSubtree.GetValue() += PointInstancerVertexCount;
 							ParentInfo->SubtreeMaterialSlots.Append(PointInstancerMaterialSlots);
 						}
 
@@ -1124,51 +1597,70 @@ namespace UE::USDInfoCache::Private
 	}
 
 	/**
-	 * Condenses our collected material slots for all subtress (SubtreeMaterialSlots) into just material slot counts,
-	 * (OutMaterialSlotCounts), according to bMergeIdenticalSlots.
+	 * Removes duplicate material slots for the subtree below RecursiveCollectMaterialSlotCounts, in case
+	 * we're allowed to merge material slots.
 	 *
 	 * We do this after the main pass because then the main material slot collecting code on
 	 * the main recursive pass just adds them to arrays, and we're allowed to handle bMergeIdenticalSlots
 	 * only here.
 	 */
-	void CollectMaterialSlotCounts(FUsdInfoCacheImpl& Impl, bool bContextMergeIdenticalSlots)
+	void RecursiveCollectMaterialSlotCounts(
+		uint64 PrimIndex,
+		FUsdInfoCacheImpl& Impl,
+		const FUsdSchemaTranslationContext& Context,
+		bool bPrimIsCollapsedOrCollapseRoot
+	)
 	{
-		TRACE_CPUPROFILER_EVENT_SCOPE(CollectMaterialSlotCounts);
+#if USE_USD_SDK
+		using namespace UE::UsdInfoCache::Private;
+
+		TRACE_CPUPROFILER_EVENT_SCOPE(RecursiveCollectMaterialSlotCounts);
 
-		if (!bContextMergeIdenticalSlots)
+		if (!Context.bMergeIdenticalMaterialSlots || PrimIndex == INDEX_NONE)
 		{
 			return;
 		}
 
-		FReadScopeLock Lock{Impl.InfoMapLock};
-		for (TPair<UE::FSdfPath, UE::UsdInfoCache::Private::FUsdPrimInfo>& Pair : Impl.InfoMap)
+		FUsdPrimInfo& Info = Impl.PrimInfoArray[PrimIndex];
 		{
-			const UE::FSdfPath& PrimPath = Pair.Key;
-			UE::UsdInfoCache::Private::FUsdPrimInfo& Info = Pair.Value;
-
-			bool bCanMergeSlotsForThisPrim = false;
-
-			// Check whether we merge slots for this prim or not
+			FReadScopeLock PrimLock = Impl.LockForReading(Info);
+			if (Info.bSlotsWereMerged)
 			{
-				FReadScopeLock ReadPrimLock{Impl.PrimLocks[Info.PrimLockIndex]};
-
-				// We only merge slots in the context of collapsing
-				const TOptional<UE::FSdfPath>& CollapsedRoot = Info.AssetCollapsedRoot;
-				const bool bPrimIsCollapsedOrCollapseRoot = !CollapsedRoot.IsSet() || (!CollapsedRoot->IsEmpty() || PrimPath.IsAbsoluteRootPath());
-
-				const bool bPrimIsPotentialGeometryCacheRoot = Info.GeometryCacheState
-															   == UE::UsdInfoCache::Private::EGeometryCachePrimState::ValidRoot;
-				bCanMergeSlotsForThisPrim = bPrimIsCollapsedOrCollapseRoot && !bPrimIsPotentialGeometryCacheRoot;
+				// This info has already been processed (we do partial info cache builds now)
+				return;
 			}
+		}
 
-			// Actually update the slot count
-			// For now we only ever merge material slots when collapsing
-			if (bCanMergeSlotsForThisPrim)
+		// For now we only ever merge material slots when collapsing (and not if we're collapsing due to being a geometry cache)
+		bPrimIsCollapsedOrCollapseRoot |= Info.bCollapsesChildren;
+		bool bCanMergeSlotsForThisPrim = bPrimIsCollapsedOrCollapseRoot || Info.PrimPath.IsAbsoluteRootPath();
+		if (bCanMergeSlotsForThisPrim)
+		{
+			if (Impl.IsPotentialGeometryCacheRootInner(Info, Context.Stage.GetPrimAtPath(Info.PrimPath)))
 			{
-				FWriteScopeLock WritePrimLock{Impl.PrimLocks[Info.PrimLockIndex]};
-				Info.SubtreeMaterialSlots = TSet<UsdUtils::FUsdPrimMaterialSlot>{Info.SubtreeMaterialSlots}.Array();
+				bCanMergeSlotsForThisPrim = false;
 			}
 		}
+
+		// Actually update the slot count
+		if (bCanMergeSlotsForThisPrim)
+		{
+			FWriteScopeLock WritePrimLock = Impl.LockForWriting(Info);
+			Info.SubtreeMaterialSlots = TSet<UsdUtils::FUsdPrimMaterialSlot>{Info.SubtreeMaterialSlots}.Array();
+			Info.bSlotsWereMerged = true;
+		}
+
+		const int32 MinBatchSize = 1;
+		ParallelFor(
+			TEXT("RecursiveCollectMaterialSlotCounts"),
+			Info.ChildIndices.Num(),
+			MinBatchSize,
+			[&](int32 Index)
+			{
+				RecursiveCollectMaterialSlotCounts(Info.ChildIndices[Index], Impl, Context, bPrimIsCollapsedOrCollapseRoot);
+			}
+		);
+#endif	  // USE_USD_SDK
 	}
 
 	bool CanMeshSubtreeBeCollapsed(
@@ -1194,11 +1686,11 @@ namespace UE::USDInfoCache::Private
 		pxr::SdfPath UsdPrimPath = UsdPrim.GetPrimPath();
 
 		FReadScopeLock ScopeLock(Impl.InfoMapLock);
-		if (UE::UsdInfoCache::Private::FUsdPrimInfo* Info = Impl.InfoMap.Find(UE::FSdfPath{UsdPrimPath}))
+		if (UE::UsdInfoCache::Private::FUsdPrimInfo* Info = Impl.GetPrimInfo(UE::FSdfPath{UsdPrimPath}))
 		{
-			FReadScopeLock PrimLock{Impl.PrimLocks[Info->PrimLockIndex]};
+			FReadScopeLock PrimLock = Impl.LockForReading(*Info);
 
-			if (Info->ExpectedVertexCountForSubtree > GMaxNumVerticesCollapsedMesh)
+			if (Info->ExpectedVertexCountForSubtree.GetValue() > GMaxNumVerticesCollapsedMesh)
 			{
 				return false;
 			}
@@ -1208,31 +1700,43 @@ namespace UE::USDInfoCache::Private
 	}
 
 	void RecursiveQueryCollapsesChildren(
-		const pxr::UsdPrim& UsdPrim,
+		uint64 PrimIndex,
 		FUsdSchemaTranslationContext& Context,
 		FUsdInfoCacheImpl& Impl,
 		FUsdSchemaTranslatorRegistry& Registry
 	)
 	{
+#if USE_USD_SDK
+		using namespace UE::UsdInfoCache::Private;
+
 		TRACE_CPUPROFILER_EVENT_SCOPE(UE::USDInfoCache::Private::RecursiveQueryCollapsesChildren);
+
+		FReadScopeLock ScopeLock{Impl.InfoMapLock};
+		FUsdPrimInfo& Info = Impl.PrimInfoArray[PrimIndex];
+		{
+			FReadScopeLock PrimLock = Impl.LockForReading(Info);
+			if (Info.bCollapsesChildren)
+			{
+				// We invalidate this when the prim is updated, so if we're here it means we know this prim
+				// collapses children, and that it hasn't been updated, so we can return
+				return;
+			}
+		}
+
 		FScopedUsdAllocs Allocs;
 
-		pxr::SdfPath UsdPrimPath = UsdPrim.GetPrimPath();
-		UE::FSdfPath PrimPath{UsdPrimPath};
+		pxr::UsdPrim UsdPrim = Context.Stage.GetPrimAtPath(Info.PrimPath);
+		if (!UsdPrim)
+		{
+			return;
+		}
 
 		bool bCollapsesChildren = false;
 
 		TSharedPtr<FUsdSchemaTranslator> SchemaTranslator = Registry.CreateTranslatorForSchema(Context.AsShared(), UE::FUsdTyped(UsdPrim));
 		if (SchemaTranslator)
 		{
-			bool bIsPotentialGeometryCacheRoot = false;
-			{
-				FReadScopeLock ScopeLock{Impl.InfoMapLock};
-				if (const UE::UsdInfoCache::Private::FUsdPrimInfo* Info = Impl.InfoMap.Find(UE::FSdfPath{UsdPrimPath}))
-				{
-					bIsPotentialGeometryCacheRoot = Info->GeometryCacheState == UE::UsdInfoCache::Private::EGeometryCachePrimState::ValidRoot;
-				}
-			}
+			const bool bIsPotentialGeometryCacheRoot = Impl.IsPotentialGeometryCacheRootInner(Info, UsdPrim);
 
 			// The potential geometry cache root is checked first since the FUsdGeometryCacheTranslator::CollapsesChildren
 			// has no logic of its own
@@ -1244,52 +1748,38 @@ namespace UE::USDInfoCache::Private
 			}
 		}
 
+		if (bCollapsesChildren)
+		{
+			FWriteScopeLock PrimLock = Impl.LockForWriting(Info);
+			Info.bCollapsesChildren = true;
+		}
 		// We only need to visit our children if we don't collapse. We'll leave the AssetCollapsedRoot fields unset on
 		// the InfoMap, and whenever we query info about a particular prim will fill that in on-demand by just traveling
 		// upwards until we run into our collapse root
-		if (!bCollapsesChildren)
+		else
 		{
-			TArray<pxr::UsdPrim> Prims;
-			for (pxr::UsdPrim Child : UsdPrim.GetFilteredChildren(pxr::UsdTraverseInstanceProxies(pxr::UsdPrimAllPrimsPredicate)))
-			{
-				Prims.Emplace(Child);
-			}
-
 			const int32 MinBatchSize = 1;
 			ParallelFor(
 				TEXT("RecursiveQueryCollapsesChildren"),
-				Prims.Num(),
+				Info.ChildIndices.Num(),
 				MinBatchSize,
 				[&](int32 Index)
 				{
-					RecursiveQueryCollapsesChildren(Prims[Index], Context, Impl, Registry);
+					RecursiveQueryCollapsesChildren(Info.ChildIndices[Index], Context, Impl, Registry);
 				}
 			);
 		}
 
-		// Record our collapse root
-		{
-			FReadScopeLock ScopeLock(Impl.InfoMapLock);
-			if (UE::UsdInfoCache::Private::FUsdPrimInfo* Info = Impl.InfoMap.Find(UE::FSdfPath{UsdPrimPath}))
-			{
-				FWriteScopeLock PrimLock{Impl.PrimLocks[Info->PrimLockIndex]};
-				Info->AssetCollapsedRoot.Emplace(bCollapsesChildren ? UsdPrimPath : pxr::SdfPath::EmptyPath());
-			}
-		}
-
-		// This really should be a separate pass, but it does no harm here and we have so many passes already...
-		// This needs to happen after we set things into the FUsdPrimInfo for this prim right above this, as it may query
-		// whether this prim (or any of its children) collapse
-		//
-		// We only do this for uncollapsed prims or collapse roots. This because whenever the collapse root
-		// registers its auxiliary prims here, it will already account for all of the relevant child prims in the entire subtree,
-		// according to the translator type. The links between prims inside of a collapsed subtree aren't really useful,
-		// because if anything inside the collapsed subtree updates, we'll always just need to update from the collapsed
-		// root anyway
+		// We only do this for uncollapsed prims or collapse roots (since RecursiveQueryCollapsesChildren never steps into a
+		// collapsed prim). This because whenever the collapse root registers its auxiliary prims here, it will already account
+		// for all of the relevant child prims in the entire subtree, according to the translator type. The links between prims
+		// inside of a collapsed subtree aren't really useful, because if anything inside the collapsed subtree updates, we'll
+		// always just need to update from the collapsed root anyway
 		if (SchemaTranslator)
 		{
-			Impl.RegisterAuxiliaryPrims(PrimPath, SchemaTranslator->CollectAuxiliaryPrims());
+			Impl.RegisterAuxiliaryPrims(Info.PrimPath, SchemaTranslator->CollectAuxiliaryPrims());
 		}
+#endif	  // USE_USD_SDK
 	}
 
 	// Returns the paths to all prims on the same local layer stack, that are used as sources for composition
@@ -1331,13 +1821,12 @@ namespace UE::USDInfoCache::Private
 		return Result;
 	}
 
-	void RegisterInstanceableAuxPrims(FUsdSchemaTranslationContext& Context, FUsdInfoCacheImpl& Impl)
+	void RegisterInstanceableAuxPrims(FUsdSchemaTranslationContext& Context, bool bPartialRebuild, FUsdInfoCacheImpl& Impl)
 	{
 		TRACE_CPUPROFILER_EVENT_SCOPE(UE::USDInfoCache::Private::RegisterInstanceableAuxPrims);
 		FScopedUsdAllocs Allocs;
 
-		pxr::UsdStageWeakPtr UsdStageWeak = Impl.TempStage;
-		pxr::UsdStageRefPtr Stage = UsdStageWeak;
+		pxr::UsdStageRefPtr Stage = pxr::UsdStageRefPtr{Context.Stage};
 		if (!Stage)
 		{
 			return;
@@ -1351,7 +1840,7 @@ namespace UE::USDInfoCache::Private
 			TEXT("RegisterInstanceableAuxPrimsPrototypes"),
 			Prototypes.size(),
 			MinBatchSize,
-			[&Prototypes, &Context, &Impl, &Stage](int32 Index)
+			[&Prototypes, &Context, &Impl, &Stage, bPartialRebuild](int32 Index)
 			{
 				TRACE_CPUPROFILER_EVENT_SCOPE(RegisterInstanceableAuxPrims::Prototype);
 
@@ -1363,6 +1852,19 @@ namespace UE::USDInfoCache::Private
 					return;
 				}
 
+				// If we're on a partial rebuild and none of our instances got updated, then we don't
+				// have anything to do (prototypes can't be resynced by themselves)
+				if (bPartialRebuild)
+				{
+					UE::FSdfPath PrototypePath{Prototype.GetPrimPath()};
+
+					FWriteScopeLock ScopeLock(Impl.TempUpdatedPrototypesLock);
+					if (!Impl.TempUpdatedPrototypes.Contains(PrototypePath))
+					{
+						return;
+					}
+				}
+
 				std::vector<pxr::UsdPrim> Instances = Prototype.GetInstances();
 				if (Instances.size() < 1)
 				{
@@ -1397,9 +1899,17 @@ namespace UE::USDInfoCache::Private
 						UE::FSdfPath InstancePath{Instance.GetPrimPath()};
 						InstancePaths[Index] = InstancePath;
 
-						if (UE::UsdInfoCache::Private::FUsdPrimInfo* MainPrim = Impl.InfoMap.Find(InstancePath))
+						if (UE::UsdInfoCache::Private::FUsdPrimInfo* MainPrim = Impl.GetPrimInfo(InstancePath))
 						{
-							FWriteScopeLock PrimLock{Impl.PrimLocks[MainPrim->PrimLockIndex]};
+							FWriteScopeLock PrimLock = Impl.LockForWriting(*MainPrim);
+
+							if (MainPrim->bInstanceAuxPrimsRegistered)
+							{
+								// We already processed this particular instance on a previous info cache build
+								return;
+							}
+							MainPrim->bInstanceAuxPrimsRegistered = true;
+
 							MainPrim->AuxPrims.Append(SourcePaths);
 						}
 
@@ -1430,9 +1940,9 @@ namespace UE::USDInfoCache::Private
 				// Append all the instance paths in one go for these source paths
 				for (const UE::FSdfPath& AuxPrimPath : SourcePaths)
 				{
-					if (UE::UsdInfoCache::Private::FUsdPrimInfo* AuxPrim = Impl.InfoMap.Find(AuxPrimPath))
+					if (UE::UsdInfoCache::Private::FUsdPrimInfo* AuxPrim = Impl.GetPrimInfo(AuxPrimPath))
 					{
-						FWriteScopeLock PrimLock{Impl.PrimLocks[AuxPrim->PrimLockIndex]};
+						FWriteScopeLock PrimLock = Impl.LockForWriting(*AuxPrim);
 						AuxPrim->MainPrims.Append(InstancePaths);
 					}
 				}
@@ -1457,9 +1967,9 @@ namespace UE::USDInfoCache::Private
 		UE::FSdfPath PrimPath{UsdPrimPath};
 		{
 			FReadScopeLock ScopeLock(Impl.InfoMapLock);
-			if (UE::UsdInfoCache::Private::FUsdPrimInfo* Info = Impl.InfoMap.Find(UE::FSdfPath(UsdPrim.GetPrimPath())))
+			if (UE::UsdInfoCache::Private::FUsdPrimInfo* Info = Impl.GetPrimInfo(UE::FSdfPath(UsdPrim.GetPrimPath())))
 			{
-				FWriteScopeLock PrimLock{Impl.PrimLocks[Info->PrimLockIndex]};
+				FWriteScopeLock PrimLock = Impl.LockForWriting(*Info);
 
 				// A prim is considered a valid root if its subtree has no uncollapsible branch and a valid depth.
 				// A valid depth is positive, meaning it has an animated mesh, and doesn't exceed the limit.
@@ -1489,8 +1999,12 @@ namespace UE::USDInfoCache::Private
 			bool bIsCollapsible = false;
 			{
 				FReadScopeLock ScopeLock(Impl.InfoMapLock);
-				const UE::UsdInfoCache::Private::FUsdPrimInfo& Info = Impl.InfoMap.FindChecked(UE::FSdfPath(Child.GetPrimPath()));
-				bIsCollapsible = EnumHasAnyFlags(Info.GeometryCacheState, EGeometryCachePrimState::Collapsible);
+				if (const UE::UsdInfoCache::Private::FUsdPrimInfo* Info = Impl.GetPrimInfo(UE::FSdfPath(Child.GetPrimPath())))
+				{
+					FReadScopeLock Lock = Impl.LockForReading(*Info);
+
+					bIsCollapsible = EnumHasAnyFlags(Info->GeometryCacheState, EGeometryCachePrimState::Collapsible);
+				}
 			}
 
 			// A subtree is considered only if it has anything collapsible in the first place
@@ -1504,7 +2018,7 @@ namespace UE::USDInfoCache::Private
 	}
 
 	void RecursiveCheckForGeometryCache(
-		const pxr::UsdPrim& UsdPrim,
+		uint64 PrimIndex,
 		FUsdSchemaTranslationContext& Context,
 		FUsdInfoCacheImpl& Impl,
 		bool bIsInsideSkelRoot,
@@ -1521,66 +2035,37 @@ namespace UE::USDInfoCache::Private
 		// With this recursive check for geometry cache, we want to find branches with an animated mesh at the leaf and find the root where they can
 		// meet. This root prim will collapses the static and animated meshes under it into a single geometry cache.
 
-		pxr::SdfPath UsdPrimPath = UsdPrim.GetPrimPath();
-		UE::FSdfPath PrimPath{UsdPrimPath};
-
-		pxr::UsdPrimSiblingRange PrimChildren = UsdPrim.GetFilteredChildren(pxr::UsdTraverseInstanceProxies(pxr::UsdPrimAllPrimsPredicate));
+		FReadScopeLock ScopeLock{Impl.InfoMapLock};
+		UE::UsdInfoCache::Private::FUsdPrimInfo& Info = Impl.PrimInfoArray[PrimIndex];
 
-		TArray<pxr::UsdPrim> Prims;
-		for (pxr::UsdPrim Child : PrimChildren)
+		pxr::UsdPrim UsdPrim = Context.Stage.GetPrimAtPath(Info.PrimPath);
+		if (!UsdPrim)
 		{
-			Prims.Emplace(Child);
+			return;
 		}
+		bIsInsideSkelRoot |= UsdPrim.IsA<pxr::UsdSkelRoot>();
 
 		TArray<int32> Depths;
-		Depths.SetNum(Prims.Num());
+		Depths.SetNum(Info.ChildIndices.Num());
 
 		TArray<EGeometryCachePrimState> States;
-		States.SetNum(Prims.Num());
+		States.SetNum(Info.ChildIndices.Num());
 
 		const int32 MinBatchSize = 1;
 		ParallelFor(
 			TEXT("RecursiveCheckForGeometryCache"),
-			Prims.Num(),
+			Info.ChildIndices.Num(),
 			MinBatchSize,
-			[&Prims, &Context, &Impl, bIsInsideSkelRoot, &Depths, &States](int32 Index)
+			[&Info, &Context, &Impl, bIsInsideSkelRoot, &Depths, &States](int32 Index)
 			{
-				RecursiveCheckForGeometryCache(
-					Prims[Index],
-					Context,
-					Impl,
-					bIsInsideSkelRoot || Prims[Index].IsA<pxr::UsdSkelRoot>(),
-					Depths[Index],
-					States[Index]
-				);
+				RecursiveCheckForGeometryCache(Info.ChildIndices[Index], Context, Impl, bIsInsideSkelRoot, Depths[Index], States[Index]);
 			}
 		);
 
-		static IConsoleVariable* ForceImportCvar = IConsoleManager::Get().FindConsoleVariable(TEXT("USD.GeometryCache.ForceImport"));
-		const bool bIsImporting = Context.bIsImporting || (ForceImportCvar && ForceImportCvar->GetBool());
-
-		bool bIsAnimatedMesh = UsdUtils::IsAnimatedMesh(UsdPrim);
-		if (!bIsImporting)
-		{
-			FReadScopeLock ScopeLock(Impl.InfoMapLock);
-			if (UE::UsdInfoCache::Private::FUsdPrimInfo* Info = Impl.InfoMap.Find(UE::FSdfPath(UsdPrimPath)))
-			{
-				FWriteScopeLock PrimLock{Impl.PrimLocks[Info->PrimLockIndex]};
-
-				// When loading on the stage, the GeometryCache root can only be the animated mesh prim itself
-				// and there's no collapsing involved since each animated mesh will become a GeometryCache.
-				// The depth is irrelevant here.
-				Info->GeometryCacheDepth = -1;
-				Info->GeometryCacheState = bIsAnimatedMesh ? EGeometryCachePrimState::ValidRoot : EGeometryCachePrimState::Uncollapsible;
-			}
-
-			return;
-		}
-
 		// A geometry cache "branch" starts from an animated mesh prim for which we assign a depth of 0
 		// Other branches, without any animated mesh, we don't care about and will remain at -1
 		int32 Depth = -1;
-		if (bIsAnimatedMesh)
+		if (UsdUtils::IsAnimatedMesh(UsdPrim))
 		{
 			Depth = 0;
 		}
@@ -1693,14 +2178,9 @@ namespace UE::USDInfoCache::Private
 		}
 
 		{
-			FReadScopeLock ScopeLock(Impl.InfoMapLock);
-			if (UE::UsdInfoCache::Private::FUsdPrimInfo* Info = Impl.InfoMap.Find(UE::FSdfPath(UsdPrimPath)))
-			{
-				FWriteScopeLock PrimLock{Impl.PrimLocks[Info->PrimLockIndex]};
-
-				Info->GeometryCacheDepth = Depth;
-				Info->GeometryCacheState = PrimState;
-			}
+			FWriteScopeLock PrimLock = Impl.LockForWriting(Info);
+			Info.GeometryCacheDepth = Depth;
+			Info.GeometryCacheState = PrimState;
 		}
 
 		// We've encountered a potential root and the subtree has a geometry cache branch, so find its root
@@ -1713,7 +2193,7 @@ namespace UE::USDInfoCache::Private
 					Warning,
 					TEXT("Prim '%s' is potentially a geometry cache %d levels deep, which exceeds the limit of %d. "
 						 "This could affect its imported animation. The limit can be increased with the cvar USD.GeometryCache.MaxDepth if needed."),
-					*PrimPath.GetString(),
+					*Info.PrimPath.GetString(),
 					Depth,
 					GGeometryCacheMaxDepth
 				);
@@ -1726,7 +2206,7 @@ namespace UE::USDInfoCache::Private
 		OutState = PrimState;
 	}
 
-	void CheckForGeometryCache(const pxr::UsdPrim& UsdPrim, FUsdSchemaTranslationContext& Context, FUsdInfoCacheImpl& Impl)
+	void CheckForGeometryCache(FUsdSchemaTranslationContext& Context, FUsdInfoCacheImpl& Impl)
 	{
 		TRACE_CPUPROFILER_EVENT_SCOPE(CheckForGeometryCache);
 
@@ -1737,11 +2217,25 @@ namespace UE::USDInfoCache::Private
 			return;
 		}
 
+		static IConsoleVariable* ForceImportCvar = IConsoleManager::Get().FindConsoleVariable(TEXT("USD.GeometryCache.ForceImport"));
+		const bool bIsImporting = Context.bIsImporting || (ForceImportCvar && ForceImportCvar->GetBool());
+		if (!bIsImporting)
+		{
+			// We only collapse subtrees into a single geometry cache when "importing".
+			//
+			// If we're not importing, all we need to know is whether the prim is an animated mesh or not. We'll do that on-demand
+			// when the first call to IsPotentialGeometryCacheRoot() executes, because we don't want to spend time finding all animated
+			// meshes in the entire stage only to never actually need that information
+			return;
+		}
+
+		FScopedUsdAllocs UsdAllocs;
+		pxr::UsdPrim PseudoRoot = Context.Stage.GetPseudoRoot();
+
 		// If the stage doesn't contain any animated mesh prims, then don't bother doing a full check
 		bool bHasAnimatedMesh = false;
 		{
-			FScopedUsdAllocs UsdAllocs;
-			TArray<TUsdStore<pxr::UsdPrim>> ChildPrims = UsdUtils::GetAllPrimsOfType(UsdPrim, pxr::TfType::Find<pxr::UsdGeomMesh>());
+			TArray<TUsdStore<pxr::UsdPrim>> ChildPrims = UsdUtils::GetAllPrimsOfType(PseudoRoot, pxr::TfType::Find<pxr::UsdGeomMesh>());
 			for (const TUsdStore<pxr::UsdPrim>& ChildPrim : ChildPrims)
 			{
 				if (UsdUtils::IsAnimatedMesh(ChildPrim.Get()))
@@ -1751,17 +2245,15 @@ namespace UE::USDInfoCache::Private
 				}
 			}
 		}
-
 		if (!bHasAnimatedMesh)
 		{
 			return;
 		}
 
-		const bool bIsInsideSkelRoot = static_cast<bool>(UsdUtils::GetClosestParentSkelRoot(UsdPrim));
-
 		int32 Depth = -1;
 		EGeometryCachePrimState State = EGeometryCachePrimState::None;
-		RecursiveCheckForGeometryCache(UsdPrim, Context, Impl, bIsInsideSkelRoot, Depth, State);
+		const bool bIsInsideSkelRoot = false;
+		RecursiveCheckForGeometryCache(0, Context, Impl, bIsInsideSkelRoot, Depth, State);
 
 		// If we end up with a positive depth, it means the check found an animated mesh somewhere
 		// but no potential root before reaching the pseudoroot, so find one
@@ -1779,12 +2271,8 @@ namespace UE::USDInfoCache::Private
 				);
 			}
 
-			FScopedUsdAllocs UsdAllocs;
-
-			pxr::UsdPrimSiblingRange PrimChildren = UsdPrim.GetFilteredChildren(pxr::UsdTraverseInstanceProxies(pxr::UsdPrimAllPrimsPredicate));
-
 			// The pseudoroot itself cannot be a root for the geometry cache so start from its children
-			TArray<pxr::UsdPrim> Prims;
+			pxr::UsdPrimSiblingRange PrimChildren = PseudoRoot.GetFilteredChildren(pxr::UsdTraverseInstanceProxies(pxr::UsdPrimAllPrimsPredicate));
 			for (pxr::UsdPrim Child : PrimChildren)
 			{
 				FindValidGeometryCacheRoot(Child, Context, Impl, State);
@@ -1794,15 +2282,16 @@ namespace UE::USDInfoCache::Private
 #endif	  // USE_SD_SDK
 }	 // namespace UE::USDInfoCache::Private
 
+PRAGMA_DISABLE_DEPRECATION_WARNINGS
 bool FUsdInfoCache::IsPotentialGeometryCacheRoot(const UE::FSdfPath& Path) const
 {
 	if (FUsdInfoCacheImpl* ImplPtr = Impl.Get())
 	{
 		FReadScopeLock ScopeLock(ImplPtr->InfoMapLock);
 
-		if (const UE::UsdInfoCache::Private::FUsdPrimInfo* FoundInfo = ImplPtr->InfoMap.Find(Path))
+		if (const UE::UsdInfoCache::Private::FUsdPrimInfo* FoundInfo = ImplPtr->GetPrimInfo(Path))
 		{
-			FReadScopeLock PrimLock{ImplPtr->PrimLocks[FoundInfo->PrimLockIndex]};
+			FReadScopeLock PrimLock = ImplPtr->LockForReading(*FoundInfo);
 
 			return FoundInfo->GeometryCacheState == UE::UsdInfoCache::Private::EGeometryCachePrimState::ValidRoot;
 		}
@@ -1813,6 +2302,30 @@ bool FUsdInfoCache::IsPotentialGeometryCacheRoot(const UE::FSdfPath& Path) const
 
 	return false;
 }
+PRAGMA_ENABLE_DEPRECATION_WARNINGS
+
+bool FUsdInfoCache::IsPotentialGeometryCacheRoot(const UE::FUsdPrim& Prim) const
+{
+	using namespace UE::UsdInfoCache::Private;
+
+#if USE_USD_SDK
+	if (FUsdInfoCacheImpl* ImplPtr = Impl.Get())
+	{
+		FReadScopeLock ScopeLock(ImplPtr->InfoMapLock);
+
+		UE::FSdfPath PrimPath = Prim.GetPrimPath();
+		if (FUsdPrimInfo* FoundInfo = ImplPtr->GetPrimInfo(PrimPath))
+		{
+			return ImplPtr->IsPotentialGeometryCacheRootInner(*FoundInfo, Prim);
+		}
+
+		// This should never happen: We should have cached the entire tree
+		ensureMsgf(false, TEXT("Prim path '%s' has not been cached!"), *PrimPath.GetString());
+	}
+#endif	  // USE_USD_SDK
+
+	return false;
+}
 
 void FUsdInfoCache::ResetTranslatedPrototypes()
 {
@@ -1849,9 +2362,9 @@ TOptional<uint64> FUsdInfoCache::GetSubtreeVertexCount(const UE::FSdfPath& Path)
 	{
 		FReadScopeLock ScopeLock(ImplPtr->InfoMapLock);
 
-		if (const UE::UsdInfoCache::Private::FUsdPrimInfo* FoundInfo = ImplPtr->InfoMap.Find(Path))
+		if (const UE::UsdInfoCache::Private::FUsdPrimInfo* FoundInfo = ImplPtr->GetPrimInfo(Path))
 		{
-			FReadScopeLock PrimLock{ImplPtr->PrimLocks[FoundInfo->PrimLockIndex]};
+			FReadScopeLock PrimLock = ImplPtr->LockForReading(*FoundInfo);
 
 			return FoundInfo->ExpectedVertexCountForSubtree;
 		}
@@ -1869,9 +2382,9 @@ TOptional<uint64> FUsdInfoCache::GetSubtreeMaterialSlotCount(const UE::FSdfPath&
 	{
 		FReadScopeLock ScopeLock(ImplPtr->InfoMapLock);
 
-		if (const UE::UsdInfoCache::Private::FUsdPrimInfo* FoundInfo = ImplPtr->InfoMap.Find(Path))
+		if (const UE::UsdInfoCache::Private::FUsdPrimInfo* FoundInfo = ImplPtr->GetPrimInfo(Path))
 		{
-			FReadScopeLock PrimLock{ImplPtr->PrimLocks[FoundInfo->PrimLockIndex]};
+			FReadScopeLock PrimLock = ImplPtr->LockForReading(*FoundInfo);
 
 			return FoundInfo->SubtreeMaterialSlots.Num();
 		}
@@ -1889,9 +2402,9 @@ TOptional<TArray<UsdUtils::FUsdPrimMaterialSlot>> FUsdInfoCache::GetSubtreeMater
 	{
 		FReadScopeLock ScopeLock(ImplPtr->InfoMapLock);
 
-		if (const UE::UsdInfoCache::Private::FUsdPrimInfo* FoundInfo = ImplPtr->InfoMap.Find(Path))
+		if (const UE::UsdInfoCache::Private::FUsdPrimInfo* FoundInfo = ImplPtr->GetPrimInfo(Path))
 		{
-			FReadScopeLock PrimLock{ImplPtr->PrimLocks[FoundInfo->PrimLockIndex]};
+			FReadScopeLock PrimLock = ImplPtr->LockForReading(*FoundInfo);
 
 			return FoundInfo->SubtreeMaterialSlots;
 		}
@@ -1944,11 +2457,77 @@ PRAGMA_ENABLE_DEPRECATION_WARNINGS
 
 void FUsdInfoCache::RebuildCacheForSubtree(const UE::FUsdPrim& Prim, FUsdSchemaTranslationContext& Context)
 {
-#if USE_USD_SDK
-	TRACE_CPUPROFILER_EVENT_SCOPE(FUsdInfoCache::RebuildCacheForSubtree);
+	RebuildCacheForSubtrees({Prim.GetPrimPath()}, Context);
+}
+
+namespace UE::USDInfoCache::Private
+{
+	// Computes an efficient list of root prims to rebuild, so that we don't try rebuilding any one prim more than once
+	TSet<UE::FSdfPath> GetSubtreeRootsToRebuild(const TArray<UE::FSdfPath>& SubtreeRootPaths)
+	{
+		TRACE_CPUPROFILER_EVENT_SCOPE(GetSubtreeRootsToRebuild);
+
+		TSet<UE::FSdfPath> Prefixes;
+		for (const UE::FSdfPath& SubtreeRootPath : SubtreeRootPaths)
+		{
+			// If one of the paths is the pseudoroot there's not point doing anything else in here: It's a full rebuild
+			if (SubtreeRootPath.IsAbsoluteRootPath())
+			{
+				return {SubtreeRootPath};
+			}
+
+			const UE::FSdfPath* PrefixToRemove = nullptr;
+
+			bool bKeepNewRootPath = true;
+			for (const UE::FSdfPath& Prefix : Prefixes)
+			{
+				// We already have a path that is a prefix of SubtreeRootPath: We'll already resync
+				// SubtreeRootPath as part of resyncing Prefix anyway, so we don't need to use it
+				if (SubtreeRootPath.HasPrefix(Prefix))
+				{
+					bKeepNewRootPath = false;
+					break;
+				}
+
+				// This new subtree root path is an actual parent of another path we already saw.
+				// We immediately know we want this path, because if we have added Prefix to our set in the past,
+				// then we know we never had a prefix to it, and so we can't have a prefix to SubtreeRootPath either.
+				//
+				// Furthermore, we can get rid of this Prefix and just use SubtreeRootPath instead!
+				if (Prefix.HasPrefix(SubtreeRootPath))
+				{
+					PrefixToRemove = &Prefix;
+					break;
+				}
+			}
+			if (!bKeepNewRootPath)
+			{
+				continue;
+			}
+
+			if (PrefixToRemove)
+			{
+				Prefixes.Remove(*PrefixToRemove);
+			}
+			Prefixes.Add(SubtreeRootPath);
+		}
+
+		return Prefixes;
+	}
+}	 // namespace UE::USDInfoCache::Private
 
+void FUsdInfoCache::RebuildCacheForSubtrees(const TArray<UE::FSdfPath>& SubtreeRootPaths, FUsdSchemaTranslationContext& Context)
+{
+#if USE_USD_SDK
 	using namespace UE::USDInfoCache::Private;
 
+	TRACE_CPUPROFILER_EVENT_SCOPE(FUsdInfoCache::RebuildCacheForSubtrees);
+
+	if (SubtreeRootPaths.Num() == 0)
+	{
+		return;
+	}
+
 	FUsdInfoCacheImpl* ImplPtr = Impl.Get();
 	if (!ImplPtr)
 	{
@@ -1960,14 +2539,6 @@ void FUsdInfoCache::RebuildCacheForSubtree(const UE::FUsdPrim& Prim, FUsdSchemaT
 
 	TGuardValue<bool> Guard{Context.bIsBuildingInfoCache, true};
 	{
-		pxr::UsdPrim UsdPrim{Prim};
-		if (!UsdPrim)
-		{
-			return;
-		}
-
-		ImplPtr->TempStage = UE::FUsdStageWeak{UsdPrim.GetStage()};
-
 		IUsdSchemasModule& UsdSchemasModule = FModuleManager::Get().LoadModuleChecked<IUsdSchemasModule>(TEXT("USDSchemas"));
 		FUsdSchemaTranslatorRegistry& Registry = UsdSchemasModule.GetTranslatorRegistry();
 
@@ -1977,19 +2548,43 @@ void FUsdInfoCache::RebuildCacheForSubtree(const UE::FUsdPrim& Prim, FUsdSchemaT
 			MaterialPurposeToken = UnrealToUsd::ConvertToken(*Context.MaterialPurpose.ToString()).Get();
 		}
 
-		// We don't call FUsdInfoCache::Clear() here as we don't want to get rid of PrimPathToAssets because
-		// we're rebuilding the cache, as that info is also linked to the asset cache
+		TSet<UE::FSdfPath> ProcessedRootPaths = GetSubtreeRootsToRebuild(SubtreeRootPaths);
+		const bool bIsFullRebuild = (ProcessedRootPaths.Num() == 1 && ProcessedRootPaths.Contains(UE::FSdfPath::AbsoluteRootPath()));
+
+		// We need to rebuild everything, don't bother carefully cleaning things up
+		if (bIsFullRebuild)
 		{
-			FWriteScopeLock ScopeLock(ImplPtr->InfoMapLock);
-			ImplPtr->InfoMap.Reset();
+			Clear();
+
+			RepopulateInfoMapSubtree(UE::FSdfPath::AbsoluteRootPath(), Context.Stage, *ImplPtr);
 		}
+		// We need to rebuild just some subtrees of the existing cache
+		else
 		{
-			FWriteScopeLock PointInstancerLock(ImplPtr->PointInstancerPathsLock);
-			ImplPtr->PointInstancerPaths.Reset();
+			// Prepare for cleanup
+			ImplPtr->ClearTransientInfo();
+
+			// Reset/remove everything
+			for (const UE::FSdfPath& RootPath : ProcessedRootPaths)
+			{
+				CleanUpInfoMapSubtree(RootPath, *ImplPtr);
+			}
+			ImplPtr->DisconnectResetPrimAuxLinks();
+
+			// Rebuild the new info struct nodes and connect parent/child links
+			for (const UE::FSdfPath& RootPath : ProcessedRootPaths)
+			{
+				RepopulateInfoMapSubtree(RootPath, Context.Stage, *ImplPtr);
+			}
 		}
 
-		// This should be the first step as all future functions will expect to find one entry per prim in the cache
-		RepopulateInfoMap(UsdPrim, *ImplPtr);
+		if (!Context.Stage)
+		{
+			// Quit only here so that the AUsdStageActor can just blindly do a "full rebuild" when it wants to cleanup as well
+			return;
+		}
+		UE::FUsdPrim PseudoRoot = Context.Stage.GetPseudoRoot();
+		const uint64 PseudoRootIndex = 0;
 
 		// Propagate vertex and material slot counts before we query CollapsesChildren because the Xformable
 		// translator needs to know when it would generate too large a static mesh
@@ -1997,7 +2592,7 @@ void FUsdInfoCache::RebuildCacheForSubtree(const UE::FUsdPrim& Prim, FUsdSchemaT
 		TArray<UsdUtils::FUsdPrimMaterialSlot> SubtreeSlots;
 		const bool bPossibleInheritedBindings = false;
 		RecursivePropagateVertexAndMaterialSlotCounts(
-			UsdPrim,
+			PseudoRootIndex,
 			Context,
 			MaterialPurposeToken,
 			*ImplPtr,
@@ -2009,13 +2604,14 @@ void FUsdInfoCache::RebuildCacheForSubtree(const UE::FUsdPrim& Prim, FUsdSchemaT
 
 		UpdateInfoForPointInstancers(Context, *ImplPtr);
 
-		CheckForGeometryCache(UsdPrim, Context, *ImplPtr);
+		CheckForGeometryCache(Context, *ImplPtr);
 
-		RecursiveQueryCollapsesChildren(UsdPrim, Context, *ImplPtr, Registry);
+		RecursiveQueryCollapsesChildren(PseudoRootIndex, Context, *ImplPtr, Registry);
 
-		RegisterInstanceableAuxPrims(Context, *ImplPtr);
+		RegisterInstanceableAuxPrims(Context, bIsFullRebuild, *ImplPtr);
 
-		CollectMaterialSlotCounts(*ImplPtr, Context.bMergeIdenticalMaterialSlots);
+		const bool bIsPrimCollapsedOrCollapseRoot = false;
+		RecursiveCollectMaterialSlotCounts(PseudoRootIndex, *ImplPtr, Context, bIsPrimCollapsedOrCollapseRoot);
 	}
 #endif	  // USE_USD_SDK
 }
@@ -2029,12 +2625,12 @@ void FUsdInfoCache::Clear()
 		{
 			TRACE_CPUPROFILER_EVENT_SCOPE(InfoMapEmpty);
 			FWriteScopeLock ScopeLock(ImplPtr->InfoMapLock);
+			ImplPtr->PrimInfoArray.Empty();
+			ImplPtr->StaleInfoIndices.Empty();
 			ImplPtr->InfoMap.Empty();
 		}
-		{
-			FWriteScopeLock PointInstancerLock(ImplPtr->PointInstancerPathsLock);
-			ImplPtr->PointInstancerPaths.Empty();
-		}
+
+		ImplPtr->ClearTransientInfo();
 
 		ResetTranslatedPrototypes();
 	}
@@ -2054,6 +2650,8 @@ bool FUsdInfoCache::IsEmpty()
 TOptional<bool> FUsdInfoCache::CanXformableSubtreeBeCollapsed(const UE::FSdfPath& RootPath, FUsdSchemaTranslationContext& Context) const
 {
 #if USE_USD_SDK
+	using namespace UE::UsdInfoCache::Private;
+
 	TRACE_CPUPROFILER_EVENT_SCOPE(FUsdInfoCache::CanSubtreeBeCollapsed);
 
 	// The only reason this function exists is that FUsdGeomXformableTranslator::CollapsesChildren() needs to check if all
@@ -2071,9 +2669,9 @@ TOptional<bool> FUsdInfoCache::CanXformableSubtreeBeCollapsed(const UE::FSdfPath
 	{
 		FReadScopeLock ScopeLock(ImplPtr->InfoMapLock);
 
-		if (const UE::UsdInfoCache::Private::FUsdPrimInfo* FoundInfo = ImplPtr->InfoMap.Find(RootPath))
+		if (const FUsdPrimInfo* FoundInfo = ImplPtr->GetPrimInfo(RootPath))
 		{
-			FReadScopeLock PrimLock{ImplPtr->PrimLocks[FoundInfo->PrimLockIndex]};
+			FReadScopeLock PrimLock = ImplPtr->LockForReading(*FoundInfo);
 			if (FoundInfo->bXformSubtreeCanBeCollapsed.IsSet())
 			{
 				return FoundInfo->bXformSubtreeCanBeCollapsed.GetValue();
@@ -2084,8 +2682,7 @@ TOptional<bool> FUsdInfoCache::CanXformableSubtreeBeCollapsed(const UE::FSdfPath
 
 		// Fill in missing entries for CanBeCollapsed on-demand and compute the value for the prim at RootPath,
 		// if we can still access our stage
-		pxr::UsdStageWeakPtr UsdStageWeak = ImplPtr->TempStage;
-		pxr::UsdStageRefPtr Stage = UsdStageWeak;
+		pxr::UsdStageRefPtr Stage = pxr::UsdStageRefPtr{Context.Stage};
 		if (Stage)
 		{
 			if (pxr::UsdPrim Prim = Stage->GetPrimAtPath(RootPath))
@@ -2093,7 +2690,7 @@ TOptional<bool> FUsdInfoCache::CanXformableSubtreeBeCollapsed(const UE::FSdfPath
 				IUsdSchemasModule& UsdSchemasModule = FModuleManager::Get().LoadModuleChecked<IUsdSchemasModule>(TEXT("USDSchemas"));
 				FUsdSchemaTranslatorRegistry& Registry = UsdSchemasModule.GetTranslatorRegistry();
 
-				bCanBeCollapsed = UE::USDInfoCache::Private::RecursiveQueryCanBeCollapsed(Prim, Context, *ImplPtr, Registry);
+				bCanBeCollapsed = RecursiveQueryCanBeCollapsed(Prim, Context, *ImplPtr, Registry);
 			}
 		}
 
diff --git a/Engine/Plugins/Importers/USDImporter/Source/USDSchemas/Private/USDPrimLinkCache.cpp b/Engine/Plugins/Importers/USDImporter/Source/USDSchemas/Private/USDPrimLinkCache.cpp
--- a/Engine/Plugins/Importers/USDImporter/Source/USDSchemas/Private/USDPrimLinkCache.cpp
+++ b/Engine/Plugins/Importers/USDImporter/Source/USDSchemas/Private/USDPrimLinkCache.cpp
@@ -58,20 +58,6 @@ bool FUsdPrimLinkCache::ContainsInfoAboutPrim(const UE::FSdfPath& Path) const
 	return false;
 }
 
-TSet<UE::FSdfPath> FUsdPrimLinkCache::GetKnownPrims() const
-{
-	TSet<UE::FSdfPath> Result;
-
-	if (FUsdPrimLinkCacheImpl* ImplPtr = Impl.Get())
-	{
-		FReadScopeLock ScopeLock(ImplPtr->PrimPathToAssetsLock);
-		ImplPtr->PrimPathToAssets.GetKeys(Result);
-		return Result;
-	}
-
-	return Result;
-}
-
 void FUsdPrimLinkCache::LinkAssetToPrim(const UE::FSdfPath& Path, UObject* Asset)
 {
 	FUsdPrimLinkCacheImpl* ImplPtr = Impl.Get();
diff --git a/Engine/Plugins/Importers/USDImporter/Source/USDSchemas/Public/USDInfoCache.h b/Engine/Plugins/Importers/USDImporter/Source/USDSchemas/Public/USDInfoCache.h
--- a/Engine/Plugins/Importers/USDImporter/Source/USDSchemas/Public/USDInfoCache.h
+++ b/Engine/Plugins/Importers/USDImporter/Source/USDSchemas/Public/USDInfoCache.h
@@ -44,10 +44,15 @@ public:
 	// Returns whether we contain any info about prim at 'Path' at all
 	bool ContainsInfoAboutPrim(const UE::FSdfPath& Path) const;
 
+	// Retrieves the children of a prim from the cached information
+	TArray<UE::FSdfPath> GetChildren(const UE::FSdfPath& ParentPath) const;
+
 	// Returns a list of all prims we have generic info about
+	UE_DEPRECATED(5.5, "No longer used")
 	TSet<UE::FSdfPath> GetKnownPrims() const;
 
 	void RebuildCacheForSubtree(const UE::FUsdPrim& Prim, FUsdSchemaTranslationContext& Context);
+	void RebuildCacheForSubtrees(const TArray<UE::FSdfPath>& SubtreeRoots, FUsdSchemaTranslationContext& Context);
 
 	void Clear();
 	bool IsEmpty();
@@ -70,9 +75,6 @@ public:
 	TSet<UE::FSdfPath> GetAuxiliaryPrims(const UE::FSdfPath& MainPrimPath) const;
 
 public:
-	// Returns the set of paths to all prims that have a material:binding relationship to the particular material at
-	// 'Path', if any.
-	// Returns a copy for thread safety.
 	TSet<UE::FSdfPath> GetMaterialUsers(const UE::FSdfPath& Path) const;
 	bool IsMaterialUsed(const UE::FSdfPath& Path) const;
 
@@ -86,6 +88,7 @@ public:
 	TOptional<TArray<UsdUtils::FUsdPrimMaterialSlot>> GetSubtreeMaterialSlots(const UE::FSdfPath& Path);
 
 	// Returns true if Path could potentially be collapsed as a Geometry Cache asset
+	UE_DEPRECATED(5.5, "No longer used")
 	bool IsPotentialGeometryCacheRoot(const UE::FSdfPath& Path) const;
 
 public:
@@ -137,6 +140,8 @@ public:
 
 private:
 	friend class FUsdGeomXformableTranslator;
+	friend class FUsdGeometryCacheTranslator;
+
 	// Returns true if every prim on the subtree below RootPath (including the RootPath prim itself) returns true for
 	// CanBeCollapsed(), according to their own schema translators.
 	//
@@ -147,6 +152,10 @@ private:
 	// In general, you shouldn't call this, but just use "IsPathCollapsed" or "DoesPathCollapseChildren" instead.
 	TOptional<bool> CanXformableSubtreeBeCollapsed(const UE::FSdfPath& RootPath, FUsdSchemaTranslationContext& Context) const;
 
+	// Analogous to the function above, this overload of IsPotentialGeometryCacheRoot is meant for internal use, and exists because
+	// during the info cache build (in some contexts) we can fill in this geometry cache information on-demand, for better performance.
+	bool IsPotentialGeometryCacheRoot(const UE::FUsdPrim& Prim) const;
+
 private:
 	TUniquePtr<FUsdInfoCacheImpl> Impl;
 };
diff --git a/Engine/Plugins/Importers/USDImporter/Source/USDSchemas/Public/USDPrimLinkCache.h b/Engine/Plugins/Importers/USDImporter/Source/USDSchemas/Public/USDPrimLinkCache.h
--- a/Engine/Plugins/Importers/USDImporter/Source/USDSchemas/Public/USDPrimLinkCache.h
+++ b/Engine/Plugins/Importers/USDImporter/Source/USDSchemas/Public/USDPrimLinkCache.h
@@ -28,9 +28,6 @@ public:
 	// Returns whether we contain any info about prim at 'Path' at all
 	bool ContainsInfoAboutPrim(const UE::FSdfPath& Path) const;
 
-	// Returns a list of all prims we have generic info about
-	TSet<UE::FSdfPath> GetKnownPrims() const;
-
 	void Clear();
 	bool IsEmpty();
 
diff --git a/Engine/Plugins/Importers/USDImporter/Source/USDSchemas/Public/USDSchemaTranslator.h b/Engine/Plugins/Importers/USDImporter/Source/USDSchemas/Public/USDSchemaTranslator.h
--- a/Engine/Plugins/Importers/USDImporter/Source/USDSchemas/Public/USDSchemaTranslator.h
+++ b/Engine/Plugins/Importers/USDImporter/Source/USDSchemas/Public/USDSchemaTranslator.h
@@ -8,7 +8,6 @@
 #include "USDAssetCache2.h"
 #include "USDAssetCache3.h"
 #include "USDInfoCache.h"
-#include "USDInfoCache.h"
 #include "USDMemory.h"
 #include "USDMetadataImportOptions.h"
 #include "USDPrimLinkCache.h"
@@ -325,6 +324,12 @@ struct USDSCHEMAS_API FUsdSchemaTranslationContext : public TSharedFromThis<FUsd
 	/** Skip the import of materials that aren't being used by any prim on the stage */
 	bool bTranslateOnlyUsedMaterials = false;
 
+	/**
+	 * We set material overrides within UsdGeomXformableTranslator::UpdateComponents when this flag is set. Since that is a non-trivial
+	 * amount of computation, this flag can be disabled for situations where material overrides shouldn't change (e.g. animating components)
+	 */
+	bool bAllowRecomputingMaterialOverrides = true;
+
 	/** Groom group interpolation settings */
 	TArray<FHairGroupsInterpolation> GroomInterpolationSettings;
 
diff --git a/Engine/Plugins/Importers/USDImporter/Source/USDStage/Private/USDPrimTwin.cpp b/Engine/Plugins/Importers/USDImporter/Source/USDStage/Private/USDPrimTwin.cpp
--- a/Engine/Plugins/Importers/USDImporter/Source/USDStage/Private/USDPrimTwin.cpp
+++ b/Engine/Plugins/Importers/USDImporter/Source/USDStage/Private/USDPrimTwin.cpp
@@ -9,6 +9,8 @@
 
 UUsdPrimTwin& UUsdPrimTwin::AddChild(const FString& InPrimPath)
 {
+	TRACE_CPUPROFILER_EVENT_SCOPE(UUsdPrimTwin::AddChild);
+
 	FScopedUnrealAllocs UnrealAllocs;	 // Make sure the call to new is done with the UE allocator
 
 	FString Dummy;
diff --git a/Engine/Plugins/Importers/USDImporter/Source/USDStage/Private/USDStageActor.cpp b/Engine/Plugins/Importers/USDImporter/Source/USDStage/Private/USDStageActor.cpp
--- a/Engine/Plugins/Importers/USDImporter/Source/USDStage/Private/USDStageActor.cpp
+++ b/Engine/Plugins/Importers/USDImporter/Source/USDStage/Private/USDStageActor.cpp
@@ -15,7 +15,6 @@
 #include "USDGeomMeshConversion.h"
 #include "USDGeomXformableTranslator.h"
 #include "USDInfoCache.h"
-#include "USDInfoCache.h"
 #include "USDIntegrationUtils.h"
 #include "USDLayerUtils.h"
 #include "USDLightConversion.h"
@@ -949,8 +948,10 @@ AUsdStageActor::AUsdStageActor()
 	IUsdSchemasModule& UsdSchemasModule = FModuleManager::Get().LoadModuleChecked<IUsdSchemasModule>(TEXT("USDSchemas"));
 	RenderContext = UsdSchemasModule.GetRenderContextRegistry().GetUnrealRenderContext();
 
+	// This is marked as a subobject now or else instances of blueprints that derive the AUsdStageActor don't get new transactors created
+	// for them
 	const FName UniqueName = MakeUniqueObjectName(this, UUsdTransactor::StaticClass(), TEXT("Transactor"));
-	Transactor = NewObject<UUsdTransactor>(this, UniqueName, EObjectFlags::RF_Transactional);
+	Transactor = CreateDefaultSubobject<UUsdTransactor>(UniqueName);
 	Transactor->Initialize(this);
 
 	// We never want to be without a valid BBoxCache or else we'll silently fail to compute bounds for all
@@ -1535,43 +1536,19 @@ void AUsdStageActor::HandleAccumulatedNotices()
 
 	FScopedUsdMessageLog ScopedMessageLog;
 
-	TStrongObjectPtr<UUsdInfoCache> OldInfoCache{UsdInfoCache};
-	if (UsdInfoCache && bHasResync)
-	{
-		// Take a copy of the info cache here: We want to keep the old one for this function call as it helps us cleanup the old assets and components
-		// For now we don't need to do this for info changes, only for resync changes
-		{
-			TRACE_CPUPROFILER_EVENT_SCOPE(CopyingOldInfoCache);
-			UObject* OwnerObject = this;
-			UsdInfoCache = NewObject<UUsdInfoCache>(OwnerObject, NAME_None, UsdInfoCache->GetFlags());
-
-			// Use this instead of DuplicateObject as this is a single Serialize() call (due to the Modify()), while the DuplicateObject means 3
-			// Serialize() calls
-			UsdInfoCache->Modify();
-			UsdInfoCache->GetInner().CopyImpl(OldInfoCache->GetInner());
-
-		}
-
-		// TODO: Selective rebuild of only the required parts of the cache.
-		// If a prim changes from CanBeCollapsed to not (or vice-versa), that means its parent may change, and its grandparent, etc. so we'd need
-		// to check a large part of the tree (here we're just rebuilding the whole thing for now). However, if we know that only Prim resynced, we
-		// can traverse the tree root down and if we reach Prim and its collapsing state hasn't updated from before, we don't have to update its
-		// subtree at all, or sibling subtrees. Note that whenver we do a selective rebuild of the info cache we'll need to be very careful when
-		// to update certain info cache maps: For example whenever we delete a prim we need to make sure it's removed from MaterialUsers, etc.
-		TSharedRef<FUsdSchemaTranslationContext> TranslationContext = FUsdStageActorImpl::CreateUsdSchemaTranslationContext(this, TEXT("/"));
-		UsdInfoCache->GetInner().RebuildCacheForSubtree(Stage.GetPseudoRoot(), TranslationContext.Get());
-
-		// Need to update the LevelSequenceHelper with our new cache too, as it will need the new cache to find any
-		// assets we end up creating in this update (e.g. USoundWave assets)
-		LevelSequenceHelper.SetPrimLinkCache(PrimLinkCache);
-	}
-
 	if (BBoxCache.IsValid() && bHasResync)
 	{
 		BBoxCache->Clear();
 	}
 
-	TFunction<void(TMap<UE::FSdfPath, bool>&)> SortAndCleanPrimsToUpdate = [](TMap<UE::FSdfPath, bool>& InOutMap)
+	enum class EPrimUpdateType : uint8
+	{
+		MaterialBind,	 // "Weaker" than an info change: We just need to make sure we update material overrides for this prim's component
+		Info,			 // An attribute/metadata value changed. We may need to regenerate assets, but at most update components for that one prim
+		Resync,			 // Drastic change. We need to regenerate all assets and components for the prim's entire subtree
+	};
+
+	TFunction<void(TMap<UE::FSdfPath, EPrimUpdateType>&)> SortAndCleanPrimsToUpdate = [](TMap<UE::FSdfPath, EPrimUpdateType>& InOutMap)
 	{
 		TRACE_CPUPROFILER_EVENT_SCOPE(SortAndCleanPrimsToUpdate);
 
@@ -1586,14 +1563,14 @@ void AUsdStageActor::HandleAccumulatedNotices()
 		// Strip child paths for resyncs as processing a parent means we already process the children anyway.
 		// This is not the same for info changes: We may have an info change for a parent and child component
 		// in the same change block, and we really want to call UpdateComponents for both then
-		TMap<UE::FSdfPath, bool> CleanedPairs;
+		TMap<UE::FSdfPath, EPrimUpdateType> CleanedPairs;
 		CleanedPairs.Reserve(InOutMap.Num());
 		TSet<UE::FSdfPath> ResyncedPaths;
 		ResyncedPaths.Reserve(InOutMap.Num());
-		for (const TPair<UE::FSdfPath, bool>& Pair : InOutMap)
+		for (const TPair<UE::FSdfPath, EPrimUpdateType>& Pair : InOutMap)
 		{
 			const UE::FSdfPath& ThisPrim = Pair.Key;
-			bool bIsResync = Pair.Value;
+			bool bIsResync = Pair.Value == EPrimUpdateType::Resync;
 
 			bool bRemoveThisPath = false;
 
@@ -1631,16 +1608,15 @@ void AUsdStageActor::HandleAccumulatedNotices()
 		Swap(CleanedPairs, InOutMap);
 	};
 
-	TSet<UE::FSdfPath> KnownOldPrims;
-
 	// Traverses the info caches to find out which prims we need to update
-	TFunction<void(const UE::FSdfPath&, bool, TMap<UE::FSdfPath, bool>&, TSet<UE::FSdfPath>&)> RecursiveCollectPrimsToUpdate;
-	RecursiveCollectPrimsToUpdate =
-		[this,
-		 &OldInfoCache,
-		 &RecursiveCollectPrimsToUpdate,
-		 &KnownOldPrims,
-		 &Stage](const UE::FSdfPath& PrimPath, bool bIsResync, TMap<UE::FSdfPath, bool>& OutPrimsToUpdate, TSet<UE::FSdfPath>& InOutVisitedPaths)
+	TFunction<void(const UE::FSdfPath&, EPrimUpdateType, TMap<UE::FSdfPath, EPrimUpdateType>&, TMap<UE::FSdfPath, EPrimUpdateType>&)>
+		RecursiveCollectPrimsToUpdate;
+	RecursiveCollectPrimsToUpdate = [this, &RecursiveCollectPrimsToUpdate, &Stage](
+										const UE::FSdfPath& PrimPath,
+										EPrimUpdateType UpdateType,
+										TMap<UE::FSdfPath, EPrimUpdateType>& OutPrimsToUpdate,
+										TMap<UE::FSdfPath, EPrimUpdateType>& InOutVisitedPaths
+									)
 	{
 		TRACE_CPUPROFILER_EVENT_SCOPE(RecursiveCollectPrimsToUpdate);
 
@@ -1649,11 +1625,12 @@ void AUsdStageActor::HandleAccumulatedNotices()
 		// mark them as visited, and because we always recurse with resync=false we'd assume those aren't resyncs.
 		// By that same reason this shouldn't be that expensive, as we'll only ever potentially revisit the prims
 		// that are the actual roots of the resyncs
-		if (InOutVisitedPaths.Contains(PrimPath) && !bIsResync)
+		EPrimUpdateType* LastVisitedUpdateType = InOutVisitedPaths.Find(PrimPath);
+		if (LastVisitedUpdateType && (uint8)*LastVisitedUpdateType >= (uint8)UpdateType)
 		{
 			return;
 		}
-		InOutVisitedPaths.Add(PrimPath);
+		InOutVisitedPaths.Add(PrimPath, UpdateType);
 
 		// In some cases USD sends us notices about prims that don't exist anymore: If you rename X to Y,
 		// both X and Y will be on the notice change list, even though X doesn't exist on the stage anymore.
@@ -1668,26 +1645,17 @@ void AUsdStageActor::HandleAccumulatedNotices()
 			// we'd only get that prim itself, and assume we need to spawn assets/components for it (which we
 			// really don't if it's collapsed)
 			UE::FSdfPath UnwoundPath = UsdInfoCache->GetInner().UnwindToNonCollapsedPath(PrimPath, ECollapsingType::Assets);
-			OutPrimsToUpdate.FindOrAdd(UnwoundPath) |= bIsResync;
+			EPrimUpdateType& UnwoundUpdateType = OutPrimsToUpdate.FindOrAdd(UnwoundPath);
+			UnwoundUpdateType = (EPrimUpdateType)FMath::Max((uint8)UnwoundUpdateType, (uint8)UpdateType);
 		}
 
-		// If we're going to cleanup assets/components for a prim according to an old dependency (that doesn't
-		// necessarily exist anymore), we need to make sure we try recreating assets/components for it now, even if
-		// it wasn't part of the actual change, so we'll visit both the new and old info caches.
-		// Even if this was just an info change we still need to do this: We may need to regeneate an asset from it
-		if (OldInfoCache->GetInner().ContainsInfoAboutPrim(PrimPath))
+		// We don't need to recurse via material bind links: Nothing else is affected by a mesh component refreshing
+		// its material overrides in response to an original material prim update
+		if (UpdateType == EPrimUpdateType::MaterialBind)
 		{
-			UE::FSdfPath UnwoundPath = OldInfoCache->GetInner().UnwindToNonCollapsedPath(PrimPath, ECollapsingType::Assets);
-			OutPrimsToUpdate.FindOrAdd(UnwoundPath) |= bIsResync;
+			return;
 		}
 
-		// If our original USD notice resyncs PrimPath, its subtree will need to be rebuilt, yes,
-		// but external prims that depend on prim path (its "main prims") won't need to be *recursively* resynced.
-		// Their hierarchies are fine, they just need to be updated to the fact that PrimPath changed. That is
-		// at most a component update, or regenerating the asset for that particular main prim, but it's entire
-		// hierarchy doesn't need to be rebuilt
-		const bool bRecursiveResync = false;
-
 		// Imagine we have a stage like this:
 		// 		/parent/child1
 		// 		/other
@@ -1697,56 +1665,44 @@ void AUsdStageActor::HandleAccumulatedNotices()
 		// We have to do this on both the stage and old info cache because the change may also have meant that aux/main links
 		// have been modified (i.e. "other" could depend on "child" only now, or only on the old state of the stage, but
 		// we'll still have those assets on the UE level either way, so we need to refresh them)
-		if (bIsResync)
+		if (UpdateType == EPrimUpdateType::Resync)
 		{
-			if (KnownOldPrims.Num() == 0)
+			const TArray<UE::FSdfPath>& NewChildren = UsdInfoCache->GetInner().GetChildren(PrimPath);
+			for (const UE::FSdfPath& Child : NewChildren)
 			{
-				KnownOldPrims = OldInfoCache->GetInner().GetKnownPrims();
-			}
-
-			if (UE::FUsdPrim Prim = Stage.GetPrimAtPath(PrimPath))
-			{
-				for (const UE::FUsdPrim& ChildPrim : Prim.GetChildren())
-				{
-					UE::FSdfPath ChildPrimPath = ChildPrim.GetPrimPath();
-					if (ChildPrimPath.HasPrefix(PrimPath))
-					{
-						RecursiveCollectPrimsToUpdate(ChildPrimPath, bRecursiveResync, OutPrimsToUpdate, InOutVisitedPaths);
-					}
-				}
-			}
-			// TODO: Find a better way of doing this, as this is extremely expensive
-			for (const UE::FSdfPath& KnownOldPrim : KnownOldPrims)
-			{
-				if (KnownOldPrim.HasPrefix(PrimPath))
-				{
-					RecursiveCollectPrimsToUpdate(KnownOldPrim, bRecursiveResync, OutPrimsToUpdate, InOutVisitedPaths);
-				}
+				RecursiveCollectPrimsToUpdate(Child, EPrimUpdateType::Info, OutPrimsToUpdate, InOutVisitedPaths);
 			}
 		}
 
 		TSet<UE::FSdfPath> NewMainPrims = UsdInfoCache->GetInner().GetMainPrims(PrimPath);
-		TSet<UE::FSdfPath> OldMainPrims = OldInfoCache->GetInner().GetMainPrims(PrimPath);
-		OutPrimsToUpdate.Reserve(OutPrimsToUpdate.Num() + NewMainPrims.Num() + OldMainPrims.Num());
+		TSet<UE::FSdfPath> NewMaterialUsers = UsdInfoCache->GetInner().GetMaterialUsers(PrimPath);
+
+		OutPrimsToUpdate.Reserve(OutPrimsToUpdate.Num() + NewMainPrims.Num() + NewMaterialUsers.Num());
+
 		for (const UE::FSdfPath& NewPrimPath : NewMainPrims)
 		{
-			RecursiveCollectPrimsToUpdate(NewPrimPath, bRecursiveResync, OutPrimsToUpdate, InOutVisitedPaths);
+			// If our original USD notice resyncs PrimPath, its subtree will need to be rebuilt, yes,
+			// but external prims that depend on prim path (its "main prims") won't need to be *recursively* resynced.
+			// Their hierarchies are fine, they just need to be updated to the fact that PrimPath changed. That is
+			// at most a component update, or regenerating the asset for that particular main prim, but it's entire
+			// hierarchy doesn't need to be rebuilt
+			RecursiveCollectPrimsToUpdate(NewPrimPath, EPrimUpdateType::Info, OutPrimsToUpdate, InOutVisitedPaths);
 		}
-		for (const UE::FSdfPath& OldPrimPath : OldMainPrims)
+		for (const UE::FSdfPath& NewPrimPath : NewMaterialUsers)
 		{
-			RecursiveCollectPrimsToUpdate(OldPrimPath, bRecursiveResync, OutPrimsToUpdate, InOutVisitedPaths);
+			RecursiveCollectPrimsToUpdate(NewPrimPath, EPrimUpdateType::MaterialBind, OutPrimsToUpdate, InOutVisitedPaths);
 		}
 	};
 
-	// Collect all the paths to update
-	TMap<UE::FSdfPath, bool> PrimsToUpdate;
+	// Collect all the paths to update from the old info cache
+	TMap<UE::FSdfPath, EPrimUpdateType> PrimsToUpdate;
 	{
-		TRACE_CPUPROFILER_EVENT_SCOPE(CollectPathsToUpdate);
+		TRACE_CPUPROFILER_EVENT_SCOPE(CollectOldPathsToUpdate);
 
 		PrimsToUpdate.Reserve(SortedPrimsChangedList.Num());
 
 		// Recursively append main prims to the list of PrimsToUpdate
-		TSet<UE::FSdfPath> VisitedPaths;
+		TMap<UE::FSdfPath, EPrimUpdateType> VisitedPaths;
 		for (const TPair<UE::FSdfPath, bool>& PrimChangedInfo : SortedPrimsChangedList)
 		{
 			const UE::FSdfPath PrimPath = PrimChangedInfo.Key;
@@ -1755,12 +1711,53 @@ void AUsdStageActor::HandleAccumulatedNotices()
 			// Note how we're not modifying SortedPrimsChangedList in-place and are instead adding to a new PrimsToUpdate list.
 			// The intent is that we really only want to process uncollapsed/collapse root main prims, but what is actually on these notices is up to
 			// USD, and could have anything
-			RecursiveCollectPrimsToUpdate(PrimPath, bIsResync, PrimsToUpdate, VisitedPaths);
+			RecursiveCollectPrimsToUpdate(PrimPath, bIsResync ? EPrimUpdateType::Resync : EPrimUpdateType::Info, PrimsToUpdate, VisitedPaths);
 		}
 
 		SortAndCleanPrimsToUpdate(PrimsToUpdate);
 	}
 
+	// Rebuild info cache if needed
+	if (UsdInfoCache && bHasResync)
+	{
+		// The prim path doesn't matter here, it's only used for fetching the parent component (not used on the info cache rebuild)
+		TSharedRef<FUsdSchemaTranslationContext> TranslationContext = FUsdStageActorImpl::CreateUsdSchemaTranslationContext(this, TEXT("/"));
+
+		TArray<UE::FSdfPath> ResyncPaths;
+		ResyncPaths.Reserve(SortedPrimsChangedList.Num());
+		for (const TPair<UE::FSdfPath, bool>& Pair : SortedPrimsChangedList)
+		{
+			if (Pair.Value)
+			{
+				ResyncPaths.Add(Pair.Key);
+			}
+		}
+
+		ResyncedPrimsForThisTransaction = ResyncPaths;
+		UsdInfoCache->GetInner().RebuildCacheForSubtrees(ResyncPaths, TranslationContext.Get());
+
+		// Append the paths to update from the rebuilt info cache
+		//
+		// Note: We don't *reset* prims to update, because whatever assets/components we cleaned up we may also need to
+		// regenerate. Alternatively, any *new* asset/component that we may end up generating also requires
+		// that we cleanup the old asset/component in order to display it.
+		//
+		// The fact that it's a TMap and SortAndCleanPrimsToUpdate will prevent us from doing any extra
+		// work anyway
+		{
+			TRACE_CPUPROFILER_EVENT_SCOPE(CollectNewPathsToUpdate);
+
+			TMap<UE::FSdfPath, EPrimUpdateType> VisitedPaths;
+			for (const TPair<UE::FSdfPath, bool>& PrimChangedInfo : SortedPrimsChangedList)
+			{
+				const UE::FSdfPath PrimPath = PrimChangedInfo.Key;
+				const bool bIsResync = PrimChangedInfo.Value;
+				RecursiveCollectPrimsToUpdate(PrimPath, bIsResync ? EPrimUpdateType::Resync : EPrimUpdateType::Info, PrimsToUpdate, VisitedPaths);
+			}
+			SortAndCleanPrimsToUpdate(PrimsToUpdate);
+		}
+	}
+
 	if (bHasResync)
 	{
 		FUsdStageActorImpl::DeselectActorsAndComponents(this);
@@ -1773,15 +1770,15 @@ void AUsdStageActor::HandleAccumulatedNotices()
 	{
 		TRACE_CPUPROFILER_EVENT_SCOPE(CleaningUpAssets);
 
-		for (const TPair<UE::FSdfPath, bool>& PrimChangedInfo : PrimsToUpdate)
+		for (const TPair<UE::FSdfPath, EPrimUpdateType>& PrimChangedInfo : PrimsToUpdate)
 		{
 			CleanUpAssetsTask.EnterProgressFrame();
 
 			const UE::FSdfPath& PrimPath = PrimChangedInfo.Key;
-			const bool bIsResync = PrimChangedInfo.Value;
+			const bool bIsResync = PrimChangedInfo.Value == EPrimUpdateType::Resync;
 
 			// If it's a new prim, we may not have old info about it
-			if (!OldInfoCache->GetInner().ContainsInfoAboutPrim(PrimPath))
+			if (!UsdInfoCache->GetInner().ContainsInfoAboutPrim(PrimPath))
 			{
 				continue;
 			}
@@ -1793,20 +1790,54 @@ void AUsdStageActor::HandleAccumulatedNotices()
 		}
 	}
 
-	// Regenerate assets before cleaning up components because if material assignments change, we may
-	// discover we need to clear and regenerate additional mesh components to update material overrides
-	TSet<UE::FSdfPath> MaterialUserPrims;
+	RefreshStageTask.EnterProgressFrame();
+	FScopedSlowTask CleanUpComponentsTask(PrimsToUpdate.Num(), LOCTEXT("CleaningUpComponents", "Cleaning up actors and components"));
+	{
+		TRACE_CPUPROFILER_EVENT_SCOPE(CleaningUpComponents);
+
+		for (const TPair<UE::FSdfPath, EPrimUpdateType>& PrimChangedInfo : PrimsToUpdate)
+		{
+			CleanUpComponentsTask.EnterProgressFrame();
+
+			const UE::FSdfPath PrimPath = PrimChangedInfo.Key;
+			const bool bIsResync = PrimChangedInfo.Value == EPrimUpdateType::Resync;
+
+			// If it's a new prim, we may not have old info about it
+			if (!UsdInfoCache->GetInner().ContainsInfoAboutPrim(PrimPath))
+			{
+				continue;
+			}
+
+			UE_LOG(LogUsd, Verbose, TEXT("Cleaning up actors and components for prim path '%s'"), *PrimPath.GetString());
+
+			if (bIsResync && PrimPath.IsAbsoluteRootOrPrimPath())
+			{
+				if (UUsdPrimTwin* UsdPrimTwin = GetRootPrimTwin()->Find(PrimPath.GetString()))
+				{
+					UsdPrimTwin->Clear();
+				}
+			}
+		}
+	}
+
+	// Recreate our LevelSequences before we regenerate components and want to add bindings/tracks
+	// back onto it
+	if (bNeedsAnimationReload)
+	{
+		RegenerateLevelSequence();
+	}
+
 	RefreshStageTask.EnterProgressFrame();
 	FScopedSlowTask RegenerateAssetsTask(PrimsToUpdate.Num(), LOCTEXT("RegeneratingAssets", "Regenerating assets"));
 	{
 		TRACE_CPUPROFILER_EVENT_SCOPE(RegeneratingAssets);
 
-		for (const TPair<UE::FSdfPath, bool>& PrimChangedInfo : PrimsToUpdate)
+		for (const TPair<UE::FSdfPath, EPrimUpdateType>& PrimChangedInfo : PrimsToUpdate)
 		{
 			RegenerateAssetsTask.EnterProgressFrame();
 
 			const UE::FSdfPath& PrimPath = PrimChangedInfo.Key;
-			const bool bIsResync = PrimChangedInfo.Value;
+			const bool bIsResync = PrimChangedInfo.Value == EPrimUpdateType::Resync;
 
 			UE::FUsdPrim PrimToUpdate = Stage.GetPrimAtPath(PrimPath);
 
@@ -1836,86 +1867,20 @@ void AUsdStageActor::HandleAccumulatedNotices()
 				bThisPrimLoadedAssets |= LoadAsset(*TranslationContext, PrimToUpdate);
 			}
 			bHasLoadedOrAbandonedAssets |= bThisPrimLoadedAssets;
-
-			// For UE-120185: If we recreated a material for a prim path we also need to update all components that were using it.
-			// This could be fleshed out further if other asset types require this refresh of "consumer components" but materials
-			// seem to be the only ones that do at the moment.
-			// Note that even after UE-157644 this is still useful: Material prims are not marked as a dependency
-			// of Mesh prims, otherwise we'd have to regenerate the StaticMesh itself when material info changed.
-			// Ideally we'd just set a new material override on the component, which is what this does
-			const bool bPrimGeneratesMaterials = PrimLinkCache->GetInner().GetSingleAssetForPrim<UMaterialInterface>(PrimPath) != nullptr;
-			if (bThisPrimLoadedAssets && bPrimGeneratesMaterials)
-			{
-				MaterialUserPrims.Append(UsdInfoCache->GetInner().GetMaterialUsers(PrimPath));
-			}
-		}
-	}
-
-	if (MaterialUserPrims.Num())
-	{
-		TRACE_CPUPROFILER_EVENT_SCOPE(MaterialUserPrimsPass);
-
-		TSet<UE::FSdfPath> VisitedPaths;
-		PrimsToUpdate.Reserve(PrimsToUpdate.Num() + MaterialUserPrims.Num());
-		for (const UE::FSdfPath& MaterialUserPrim : MaterialUserPrims)
-		{
-			// This also needs to be done recursivly because our "material user prim" may be just something
-			// like a UsdGeomSubset. If we're going to resync its component, we actually need to resync whatever
-			// component it was collapsed into instead
-			const bool bIsResync = true;
-			RecursiveCollectPrimsToUpdate(MaterialUserPrim, bIsResync, PrimsToUpdate, VisitedPaths);
-		}
-		SortAndCleanPrimsToUpdate(PrimsToUpdate);
-	}
-
-	RefreshStageTask.EnterProgressFrame();
-	FScopedSlowTask CleanUpComponentsTask(PrimsToUpdate.Num(), LOCTEXT("CleaningUpComponents", "Cleaning up actors and components"));
-	{
-		TRACE_CPUPROFILER_EVENT_SCOPE(CleaningUpComponents);
-
-		for (const TPair<UE::FSdfPath, bool>& PrimChangedInfo : PrimsToUpdate)
-		{
-			CleanUpComponentsTask.EnterProgressFrame();
-
-			const UE::FSdfPath PrimPath = PrimChangedInfo.Key;
-			const bool bIsResync = PrimChangedInfo.Value;
-
-			// If it's a new prim, we may not have old info about it
-			if (!OldInfoCache->GetInner().ContainsInfoAboutPrim(PrimPath))
-			{
-				continue;
-			}
-
-			UE_LOG(LogUsd, Verbose, TEXT("Cleaning up actors and components for prim path '%s'"), *PrimPath.GetString());
-
-			if (bIsResync && PrimPath.IsAbsoluteRootOrPrimPath())
-			{
-				if (UUsdPrimTwin* UsdPrimTwin = GetRootPrimTwin()->Find(PrimPath.GetString()))
-				{
-					UsdPrimTwin->Clear();
-				}
-			}
 		}
 	}
 
-	// Recreate our LevelSequences before we regenerate components and want to add bindings/tracks
-	// back onto it
-	if (bNeedsAnimationReload)
-	{
-		RegenerateLevelSequence();
-	}
-
 	RefreshStageTask.EnterProgressFrame();
 	FScopedSlowTask RegenerateComponentsTask(PrimsToUpdate.Num(), LOCTEXT("RegeneratingComponents", "Regenerating components"));
 	{
 		TRACE_CPUPROFILER_EVENT_SCOPE(RegeneratingComponents);
 
-		for (const TPair<UE::FSdfPath, bool>& PrimChangedInfo : PrimsToUpdate)
+		for (const TPair<UE::FSdfPath, EPrimUpdateType>& PrimChangedInfo : PrimsToUpdate)
 		{
 			RegenerateComponentsTask.EnterProgressFrame();
 
 			const UE::FSdfPath PrimPath = PrimChangedInfo.Key;
-			const bool bIsResync = PrimChangedInfo.Value;
+			const bool bIsResync = PrimChangedInfo.Value == EPrimUpdateType::Resync;
 
 			if (!UsdInfoCache->GetInner().ContainsInfoAboutPrim(PrimPath))
 			{
@@ -1939,9 +1904,10 @@ void AUsdStageActor::HandleAccumulatedNotices()
 	{
 		TRACE_CPUPROFILER_EVENT_SCOPE(OnPrimChangedBroadcast);
 
-		for (const TPair<UE::FSdfPath, bool>& PrimChangedInfo : PrimsToUpdate)
+		for (const TPair<UE::FSdfPath, EPrimUpdateType>& PrimChangedInfo : PrimsToUpdate)
 		{
-			OnPrimChanged.Broadcast(PrimChangedInfo.Key.GetString(), PrimChangedInfo.Value);
+			const bool bIsResync = PrimChangedInfo.Value == EPrimUpdateType::Resync;
+			OnPrimChanged.Broadcast(PrimChangedInfo.Key.GetString(), bIsResync);
 		}
 	}
 
@@ -3212,6 +3178,8 @@ void AUsdStageActor::OnObjectsReplaced(const TMap<UObject*, UObject*>& ObjectRep
 			NewActor->BBoxCache = BBoxCache;
 			BBoxCache = nullptr;
 
+			NewActor->ResyncedPrimsForThisTransaction = ResyncedPrimsForThisTransaction;
+
 			// It could be that we're automatically recompiling when going into PIE because our blueprint was dirty.
 			// In that case we also need bIsTransitioningIntoPIE to be true to prevent us from calling LoadUsdStage from PostRegisterAllComponents
 			NewActor->bIsTransitioningIntoPIE = bIsTransitioningIntoPIE;
@@ -3309,7 +3277,7 @@ void AUsdStageActor::LoadUsdStage(bool bOpenIfNeeded)
 	if (!UsdInfoCache)
 	{
 		UObject* Outer = this;
-		UsdInfoCache = NewObject<UUsdInfoCache>(Outer, NAME_None, EObjectFlags::RF_Transient | EObjectFlags::RF_Transactional);
+		UsdInfoCache = NewObject<UUsdInfoCache>(Outer, NAME_None, EObjectFlags::RF_Transient);
 	}
 
 	if (!PrimLinkCache)
@@ -3326,7 +3294,7 @@ void AUsdStageActor::LoadUsdStage(bool bOpenIfNeeded)
 	PrimLinkCache->Modify();
 	PrimLinkCache->GetInner().RemoveAllAssetPrimLinks();
 	UsdInfoCache->Modify();
-	UsdInfoCache->GetInner().RebuildCacheForSubtree(StageToLoad.GetPseudoRoot(), TranslationContext.Get());
+	UsdInfoCache->GetInner().RebuildCacheForSubtrees({UE::FSdfPath::AbsoluteRootPath()}, TranslationContext.Get());
 
 	SlowTask.EnterProgressFrame(0.7f);
 	const bool bLoadedOrAbandonedAssets = LoadAssets(*TranslationContext, StageToLoad.GetPseudoRoot());
@@ -3527,6 +3495,17 @@ void AUsdStageActor::SetupBBoxCacheIfNeeded()
 	BBoxCache = MakeShared<UE::FUsdGeomBBoxCache>(Time, static_cast<EUsdPurpose>(PurposesToLoad), bUseExtentsHint, bIgnoreVisibility);
 }
 
+void AUsdStageActor::RebuildInfoCacheFromStoredChanges()
+{
+	TRACE_CPUPROFILER_EVENT_SCOPE(AUsdStageActor::RebuildInfoCacheFromStoredChanges);
+
+	if (UsdInfoCache)
+	{
+		TSharedRef<FUsdSchemaTranslationContext> TranslationContext = FUsdStageActorImpl::CreateUsdSchemaTranslationContext(this, TEXT("/"));
+		UsdInfoCache->GetInner().RebuildCacheForSubtrees(ResyncedPrimsForThisTransaction, TranslationContext.Get());
+	}
+}
+
 UUsdPrimTwin* AUsdStageActor::GetRootPrimTwin()
 {
 	if (!RootUsdTwin)
@@ -3728,6 +3707,15 @@ void AUsdStageActor::PostTransacted(const FTransactionObjectEvent& TransactionEv
 			const bool bUnloadIfNeeded = false;
 			CloseUsdStage(bUnloadIfNeeded);
 			OpenUsdStage();
+
+			// Keep the info cache up to date whenever we undo/redo opening/closing/changing the root layer,
+			// as we don't put the filled in cache into the transaction anymore
+			if (StageState != EUsdStageState::Closed && UsdInfoCache)
+			{
+				TSharedRef<FUsdSchemaTranslationContext> TranslationContext = FUsdStageActorImpl::CreateUsdSchemaTranslationContext(this, TEXT("/"));
+				UsdInfoCache->GetInner().RebuildCacheForSubtrees({UE::FSdfPath::AbsoluteRootPath()}, TranslationContext.Get());
+			}
+
 			RegenerateLevelSequence();
 			RepopulateLevelSequence();
 		}
@@ -3967,11 +3955,29 @@ void AUsdStageActor::Serialize(FArchive& Ar)
 		Ar << BlendShapesByPath;
 		Ar << MaterialToPrimvarToUVIndex;
 		Ar << bIsTransitioningIntoPIE;
+
+		if (!UsdInfoCache)
+		{
+			UObject* Outer = this;
+			UsdInfoCache = NewObject<UUsdInfoCache>(Outer, NAME_None, EObjectFlags::RF_Transient);
+		}
+		Ar << UsdInfoCache;
+
+		if (!PrimLinkCache)
+		{
+			UObject* Outer = this;
+			PrimLinkCache = NewObject<UUsdPrimLinkCache>(Outer, NAME_None, EObjectFlags::RF_Transient | EObjectFlags::RF_Transactional);
+		}
+		Ar << PrimLinkCache;
 	}
 
 	if ((Ar.GetPortFlags() & PPF_DuplicateForPIE) || Ar.IsTransacting())
 	{
 		LevelSequenceHelper.Serialize(Ar);
+
+		// For regular transactions we don't need to serialize the info cache: We'll do partial builds
+		// after the change and when undo/redoing it
+		Ar << ResyncedPrimsForThisTransaction;
 	}
 }
 
@@ -5028,6 +5034,10 @@ void AUsdStageActor::AnimatePrims()
 		GetRootPrimTwin()->PrimPath
 	);
 
+	// For performance reasons we don't want to try computing material overrides on every animation frame.
+	// Material bindings don't change with time code anyway
+	TranslationContext->bAllowRecomputingMaterialOverrides = false;
+
 	// c.f. comment on bSequencerIsAnimating's declaration
 #if WITH_EDITOR
 	if (GEditor)
diff --git a/Engine/Plugins/Importers/USDImporter/Source/USDStage/Private/USDTransactor.cpp b/Engine/Plugins/Importers/USDImporter/Source/USDStage/Private/USDTransactor.cpp
--- a/Engine/Plugins/Importers/USDImporter/Source/USDStage/Private/USDTransactor.cpp
+++ b/Engine/Plugins/Importers/USDImporter/Source/USDStage/Private/USDTransactor.cpp
@@ -703,12 +703,7 @@ namespace UsdUtils
 	}
 
 	/** Applies the field value pairs to all prims on the stage, and returns a list of prim paths for modified prims */
-	TSet<FString> ApplyFieldMapToStage(
-		const FTransactorEditStorage& EditStorage,
-		EApplicationDirection Direction,
-		UE::FUsdStage& Stage,
-		double Time
-	)
+	TSet<FString> ApplyFieldMapToStage(const FTransactorEditStorage& EditStorage, EApplicationDirection Direction, UE::FUsdStage& Stage, double Time)
 	{
 		if (!Stage)
 		{
@@ -992,8 +987,15 @@ namespace UsdUtils
 
 			UE::FUsdStage& Stage = StageActor->GetOrOpenUsdStage();
 
-			TSet<FString>
-				PrimsChanged = UsdUtils::ApplyFieldMapToStage(Values, UsdUtils::EApplicationDirection::Reverse, Stage, StageActor->GetTime());
+			TSet<FString> PrimsChanged = UsdUtils::ApplyFieldMapToStage(	//
+				Values,
+				UsdUtils::EApplicationDirection::Reverse,
+				Stage,
+				StageActor->GetTime()
+			);
+
+			// Partial rebuild of the info cache after we have undone the USD stage changes for this transaction
+			StageActor->RebuildInfoCacheFromStoredChanges();
 
 			if (PrimsChanged.Num() > 0)
 			{
@@ -1055,6 +1057,9 @@ namespace UsdUtils
 			{
 				// Just a common Redo operation or any other type of ConcertSync transaction, so just apply the new values
 				PrimsChanged = UsdUtils::ApplyFieldMapToStage(Values, UsdUtils::EApplicationDirection::Forward, Stage, StageActor->GetTime());
+
+				// Partial rebuild of the info cache after we have redone the USD stage changes for this transaction
+				StageActor->RebuildInfoCacheFromStoredChanges();
 			}
 
 			// If we're redoing or applying ConcertSync we don't want to end up with these values when the transaction finalizes as it could be
diff --git a/Engine/Plugins/Importers/USDImporter/Source/USDStage/Public/USDStageActor.h b/Engine/Plugins/Importers/USDImporter/Source/USDStage/Public/USDStageActor.h
--- a/Engine/Plugins/Importers/USDImporter/Source/USDStage/Public/USDStageActor.h
+++ b/Engine/Plugins/Importers/USDImporter/Source/USDStage/Public/USDStageActor.h
@@ -31,6 +31,10 @@ namespace UE
 	class FUsdGeomBBoxCache;
 	class FUsdPrim;
 }
+namespace UsdUtils
+{
+	class FUsdTransactorImpl;
+}
 
 UENUM()
 enum class EUsdStageState : uint8
@@ -403,6 +407,8 @@ protected:
 	void SetupAssetCacheIfNeeded();
 	void SetupBBoxCacheIfNeeded();
 
+	void RebuildInfoCacheFromStoredChanges();
+
 	bool HasAuthorityOverStage() const;
 
 	void UpdateSpawnedObjectsTransientFlag(bool bTransient);
@@ -429,6 +435,7 @@ protected:
 protected:
 	friend struct FUsdStageActorImpl;
 	friend class FUsdLevelSequenceHelperImpl;
+	friend class UsdUtils::FUsdTransactorImpl;
 
 	UPROPERTY(
 		Category = UsdStageActor,
@@ -464,6 +471,7 @@ protected:
 	 */
 	UsdUtils::FObjectChangesByPath AccumulatedInfoChanges;
 	UsdUtils::FObjectChangesByPath AccumulatedResyncChanges;
+	TArray<UE::FSdfPath> ResyncedPrimsForThisTransaction;
 
 	PRAGMA_DISABLE_DEPRECATION_WARNINGS
 	/** Caches various information about prims that are expensive to query */
diff --git a/Engine/Plugins/Importers/USDImporter/Source/USDStageImporter/Private/USDStageImporter.cpp b/Engine/Plugins/Importers/USDImporter/Source/USDStageImporter/Private/USDStageImporter.cpp
--- a/Engine/Plugins/Importers/USDImporter/Source/USDStageImporter/Private/USDStageImporter.cpp
+++ b/Engine/Plugins/Importers/USDImporter/Source/USDStageImporter/Private/USDStageImporter.cpp
@@ -294,7 +294,7 @@ namespace UE::USDStageImporter::Private
 	{
 		// It's better to always rebuild the info cache because our import options may have changed
 		// from the options used when the cache was first built, which could change collapsing states, etc.
-		TranslationContext.UsdInfoCache->RebuildCacheForSubtree(TranslationContext.Stage.GetPseudoRoot(), TranslationContext);
+		TranslationContext.UsdInfoCache->RebuildCacheForSubtrees({UE::FSdfPath::AbsoluteRootPath()}, TranslationContext);
 	}
 
 	void ImportMaterials(FUsdStageImportContext& ImportContext, FUsdSchemaTranslationContext& TranslationContext)
