USD: Fix meshes still collapsing even when all kinds to collapse were disabled.

diff --git a/Engine/Plugins/Importers/USDImporter/Source/USDSchemas/Private/USDGeomXformableTranslator.cpp b/Engine/Plugins/Importers/USDImporter/Source/USDSchemas/Private/USDGeomXformableTranslator.cpp
--- a/Engine/Plugins/Importers/USDImporter/Source/USDSchemas/Private/USDGeomXformableTranslator.cpp
+++ b/Engine/Plugins/Importers/USDImporter/Source/USDSchemas/Private/USDGeomXformableTranslator.cpp
@@ -1062,7 +1062,8 @@ namespace UE::UsdXformableTranslatorImpl::Private
 
 				// Note that this is false if PrimKind is None
 				const bool bPrimKindCollapses = EnumHasAnyFlags(KindsToCollapse, PrimKind);
-				bool bCanBeCollapsed = bPrimKindCollapses || (PrimKind == EUsdDefaultKind::None && GCollapsePrimsWithoutKind);
+				bool bCanBeCollapsed = (KindsToCollapse != EUsdDefaultKind::None)
+									   && (bPrimKindCollapses || (PrimKind == EUsdDefaultKind::None && GCollapsePrimsWithoutKind));
 
 				if (!bCanBeCollapsed)
 				{
