USD: Enable selection sync by default.

diff --git a/Engine/Plugins/Importers/USDImporter/Source/USDStageEditor/Private/USDStageEditorSettings.h b/Engine/Plugins/Importers/USDImporter/Source/USDStageEditor/Private/USDStageEditorSettings.h
--- a/Engine/Plugins/Importers/USDImporter/Source/USDStageEditor/Private/USDStageEditorSettings.h
+++ b/Engine/Plugins/Importers/USDImporter/Source/USDStageEditor/Private/USDStageEditorSettings.h
@@ -14,7 +14,7 @@ class UUsdStageEditorSettings : public UObject
 public:
 	/** Whether our prim selection in SUSDStageTreeView is kept synchronized with the viewport selection */
 	UPROPERTY(config, EditAnywhere, Category = USD)
-	bool bSelectionSynced = false;
+	bool bSelectionSynced = true;
 
 	/** Whether to automatically set a layer as edit target when isolating it */
 	UPROPERTY(config, EditAnywhere, Category = USD)
