USD: Fix stage editor not being able to select prims if their ancestors are not expanded.

diff --git a/Engine/Plugins/Importers/USDImporter/Source/USDStageEditor/Private/SUSDStageTreeView.cpp b/Engine/Plugins/Importers/USDImporter/Source/USDStageEditor/Private/SUSDStageTreeView.cpp
--- a/Engine/Plugins/Importers/USDImporter/Source/USDStageEditor/Private/SUSDStageTreeView.cpp
+++ b/Engine/Plugins/Importers/USDImporter/Source/USDStageEditor/Private/SUSDStageTreeView.cpp
@@ -524,10 +524,14 @@ FUsdPrimViewModelPtr SUsdStageTreeView::GetItemFromPrimPath(const FString& PrimP
 	FScopedUnrealAllocs UnrealAllocs;
 
 	UE::FSdfPath UsdPrimPath(*PrimPath);
+	if (!UsdStage.GetPrimAtPath(UsdPrimPath))
+	{
+		return nullptr;
+	}
 
 	TFunction<FUsdPrimViewModelPtr(const UE::FSdfPath&, const FUsdPrimViewModelRef&)> FindTreeItemFromPrimPath;
-	FindTreeItemFromPrimPath =
-		[&FindTreeItemFromPrimPath](const UE::FSdfPath& UsdPrimPath, const FUsdPrimViewModelRef& ItemRef) -> FUsdPrimViewModelPtr
+	FindTreeItemFromPrimPath = [&FindTreeItemFromPrimPath,
+								this](const UE::FSdfPath& UsdPrimPath, const FUsdPrimViewModelRef& ItemRef) -> FUsdPrimViewModelPtr
 	{
 		if (ItemRef->UsdPrim.GetPrimPath() == UsdPrimPath)
 		{
@@ -535,6 +539,14 @@ FUsdPrimViewModelPtr SUsdStageTreeView::GetItemFromPrimPath(const FString& PrimP
 		}
 		else
 		{
+			// If we're past the check at the top of this function we know we have a prim for this path.
+			// If we do, then we *must* be able to generate a FUsdPrimViewModelPtr for it (if we dig deep enough),
+			// so let's expand ItemRef's parent on-demand, so that we generate our children that we can step into
+			if (!ItemRef->ShouldGenerateChildren() && ItemRef->ParentItem)
+			{
+				SetItemExpansion(StaticCastSharedRef<FUsdPrimViewModel>(ItemRef->ParentItem->AsShared()), true);
+			}
+
 			for (FUsdPrimViewModelRef ChildItem : ItemRef->Children)
 			{
 				if (FUsdPrimViewModelPtr ChildValue = FindTreeItemFromPrimPath(UsdPrimPath, ChildItem))
@@ -555,6 +567,7 @@ FUsdPrimViewModelPtr SUsdStageTreeView::GetItemFromPrimPath(const FString& PrimP
 
 		FoundItem = FindTreeItemFromPrimPath(PrimPathToSearch, RootItem);
 
+		// If we haven't found an item, try finding an item for an ancestor
 		while (!FoundItem.IsValid())
 		{
 			UE::FSdfPath ParentPrimPath = PrimPathToSearch.GetParentPath();
@@ -619,31 +632,14 @@ void SUsdStageTreeView::SetSelectedPrimPaths(const TArray<FString>& PrimPaths)
 
 void SUsdStageTreeView::SetSelectedPrims(const TArray<UE::FUsdPrim>& Prims)
 {
-	TSet<UE::FSdfPath> PrimPaths;
+	TArray<FString> PrimPaths;
+	PrimPaths.Reserve(Prims.Num());
 	for (const UE::FUsdPrim& Prim : Prims)
 	{
-		PrimPaths.Add(Prim.GetPrimPath());
+		PrimPaths.Add(Prim.GetPrimPath().GetString());
 	}
 
-	TArray<FUsdPrimViewModelRef> ItemsToSelect;
-	ItemsToSelect.Reserve(Prims.Num());
-
-	TFunction<void(const TArray<FUsdPrimViewModelRef>&)> Traverse;
-	Traverse = [this, &PrimPaths, &ItemsToSelect, &Traverse](const TArray<FUsdPrimViewModelRef>& Items)
-	{
-		for (const FUsdPrimViewModelRef& Item : Items)
-		{
-			if (PrimPaths.Contains(Item->UsdPrim.GetPrimPath()))
-			{
-				ItemsToSelect.Add(Item);
-			}
-
-			Traverse(Item->Children);
-		}
-	};
-	Traverse(RootItems);
-
-	SelectItemsInternal(ItemsToSelect);
+	SetSelectedPrimPaths(PrimPaths);
 }
 
 TArray<FString> SUsdStageTreeView::GetSelectedPrimPaths()
