USD: Prevent previous fix from cl 35559871 from expanding more nodes than needed.

diff --git a/Engine/Plugins/Importers/USDImporter/Source/USDStageEditor/Private/SUSDStageTreeView.cpp b/Engine/Plugins/Importers/USDImporter/Source/USDStageEditor/Private/SUSDStageTreeView.cpp
--- a/Engine/Plugins/Importers/USDImporter/Source/USDStageEditor/Private/SUSDStageTreeView.cpp
+++ b/Engine/Plugins/Importers/USDImporter/Source/USDStageEditor/Private/SUSDStageTreeView.cpp
@@ -533,11 +533,12 @@ FUsdPrimViewModelPtr SUsdStageTreeView::GetItemFromPrimPath(const FString& PrimP
 	FindTreeItemFromPrimPath = [&FindTreeItemFromPrimPath,
 								this](const UE::FSdfPath& UsdPrimPath, const FUsdPrimViewModelRef& ItemRef) -> FUsdPrimViewModelPtr
 	{
-		if (ItemRef->UsdPrim.GetPrimPath() == UsdPrimPath)
+		UE::FSdfPath ItemPath = ItemRef->UsdPrim.GetPrimPath();
+		if (ItemPath == UsdPrimPath)
 		{
 			return ItemRef;
 		}
-		else
+		else if (UsdPrimPath.HasPrefix(ItemPath))
 		{
 			// If we're past the check at the top of this function we know we have a prim for this path.
 			// If we do, then we *must* be able to generate a FUsdPrimViewModelPtr for it (if we dig deep enough),
