Removing unnecessary ensure from FRedirectCollector::RemoveAssetPathRedirection.

diff --git a/Engine/Plugins/Importers/USDImporter/Source/USDStageImporter/Private/USDStageImporter.cpp b/Engine/Plugins/Importers/USDImporter/Source/USDStageImporter/Private/USDStageImporter.cpp
--- a/Engine/Plugins/Importers/USDImporter/Source/USDStageImporter/Private/USDStageImporter.cpp
+++ b/Engine/Plugins/Importers/USDImporter/Source/USDStageImporter/Private/USDStageImporter.cpp
@@ -1522,13 +1522,7 @@ namespace UE::USDStageImporter::Private
 		SoftObjectsToRemap.GetKeys(OriginalPaths);
 		for (const FSoftObjectPath& OriginalPath : OriginalPaths)
 		{
-			// We check first here because if we query RemoveAssetPathRedirection with a path that doesn't
-			// have a remapping it will ensure...
-			FSoftObjectPath RedirectedPath = GRedirectCollector.GetAssetPathRedirection(OriginalPath);
-			if (RedirectedPath.IsValid())
-			{
-				GRedirectCollector.RemoveAssetPathRedirection(OriginalPath);
-			}
+			GRedirectCollector.RemoveAssetPathRedirection(OriginalPath);
 		}
 	}
 
diff --git a/Engine/Source/Runtime/CoreUObject/Private/Misc/RedirectCollector.cpp b/Engine/Source/Runtime/CoreUObject/Private/Misc/RedirectCollector.cpp
--- a/Engine/Source/Runtime/CoreUObject/Private/Misc/RedirectCollector.cpp
+++ b/Engine/Source/Runtime/CoreUObject/Private/Misc/RedirectCollector.cpp
@@ -351,8 +351,7 @@ void FRedirectCollector::RemoveAssetPathRedirection(const FSoftObjectPath& Origi
 {
 	FScopeLock ScopeLock(&CriticalSection);
 
-	bool bDidRemove = TryRemoveObjectPathRedirectionInternal(OriginalPath);
-	ensureMsgf(bDidRemove, TEXT("Cannot remove redirection from %s, it was not registered"), *OriginalPath.ToString());
+	TryRemoveObjectPathRedirectionInternal(OriginalPath);
 }
 
 void FRedirectCollector::RemoveAssetPathRedirection(FName OriginalPath)
