USD: Added a stage option to handle when the geometry caches are imported when in stage workflow.

diff --git a/Engine/Plugins/Importers/USDImporter/Source/GeometryCacheUSD/Private/GeometryCacheUSDComponent.cpp b/Engine/Plugins/Importers/USDImporter/Source/GeometryCacheUSD/Private/GeometryCacheUSDComponent.cpp
--- a/Engine/Plugins/Importers/USDImporter/Source/GeometryCacheUSD/Private/GeometryCacheUSDComponent.cpp
+++ b/Engine/Plugins/Importers/USDImporter/Source/GeometryCacheUSD/Private/GeometryCacheUSDComponent.cpp
@@ -34,11 +34,7 @@ void UGeometryCacheUsdComponent::OnRegister()
 		}
 	}
 
-	ClearTrackData();
-	SetupTrackData();
-
-	// Skip code for UGeometryCacheComponent::OnRegister
-	UMeshComponent::OnRegister();
+	UGeometryCacheComponent::OnRegister();
 }
 
 void UGeometryCacheUsdComponent::OnUnregister()
@@ -54,8 +50,5 @@ void UGeometryCacheUsdComponent::OnUnregister()
 		}
 	}
 
-	ClearTrackData();
-
-	// Skip code for UGeometryCacheComponent::OnUnregister
-	UMeshComponent::OnUnregister();
+	UGeometryCacheComponent::OnUnregister();
 }
diff --git a/Engine/Plugins/Importers/USDImporter/Source/GeometryCacheUSD/Private/GeometryCacheUSDSceneProxy.cpp b/Engine/Plugins/Importers/USDImporter/Source/GeometryCacheUSD/Private/GeometryCacheUSDSceneProxy.cpp
--- a/Engine/Plugins/Importers/USDImporter/Source/GeometryCacheUSD/Private/GeometryCacheUSDSceneProxy.cpp
+++ b/Engine/Plugins/Importers/USDImporter/Source/GeometryCacheUSD/Private/GeometryCacheUSDSceneProxy.cpp
@@ -30,7 +30,7 @@ bool FGeomCacheTrackUsdProxy::UpdateMeshData(float Time, bool bLooping, int32& I
 		}
 		return bResult;
 	}
-	return false;
+	return FGeomCacheTrackProxy::UpdateMeshData(Time, bLooping, InOutMeshSampleIndex, OutMeshData);
 }
 
 bool FGeomCacheTrackUsdProxy::GetMeshData(int32 SampleIndex, FGeometryCacheMeshData& OutMeshData)
@@ -41,7 +41,7 @@ bool FGeomCacheTrackUsdProxy::GetMeshData(int32 SampleIndex, FGeometryCacheMeshD
 	{
 		return UsdTrack->GetMeshData(SampleIndex, OutMeshData);
 	}
-	return false;
+	return FGeomCacheTrackProxy::GetMeshData(SampleIndex, OutMeshData);
 }
 
 bool FGeomCacheTrackUsdProxy::IsTopologyCompatible(int32 SampleIndexA, int32 SampleIndexB)
@@ -55,13 +55,17 @@ bool FGeomCacheTrackUsdProxy::IsTopologyCompatible(int32 SampleIndexA, int32 Sam
 
 		return NumVerticesA == NumVerticesB;
 	}
-	return false;
+	return FGeomCacheTrackProxy::IsTopologyCompatible(SampleIndexA, SampleIndexB);
 }
 
 const FVisibilitySample& FGeomCacheTrackUsdProxy::GetVisibilitySample(float Time, const bool bLooping) const
 {
-	// Assume the track is visible for its whole duration
-	return FVisibilitySample::VisibleSample;
+	if (UGeometryCacheTrackUsd* UsdTrack = Cast<UGeometryCacheTrackUsd>(Track))
+	{
+		// Assume the track is visible for its whole duration
+		return FVisibilitySample::VisibleSample;
+	}
+	return FGeomCacheTrackProxy::GetVisibilitySample(Time, bLooping);
 }
 
 void FGeomCacheTrackUsdProxy::FindSampleIndexesFromTime(
@@ -87,4 +91,8 @@ void FGeomCacheTrackUsdProxy::FindSampleIndexesFromTime(
 			InInterpolationFactor = 1.0f - InInterpolationFactor;
 		}
 	}
+	else
+	{
+		FGeomCacheTrackProxy::FindSampleIndexesFromTime(Time, bLooping, bIsPlayingBackwards, OutFrameIndex, OutNextFrameIndex, InInterpolationFactor);
+	}
 }
diff --git a/Engine/Plugins/Importers/USDImporter/Source/USDSchemas/Private/USDGeometryCacheTranslator.cpp b/Engine/Plugins/Importers/USDImporter/Source/USDSchemas/Private/USDGeometryCacheTranslator.cpp
--- a/Engine/Plugins/Importers/USDImporter/Source/USDSchemas/Private/USDGeometryCacheTranslator.cpp
+++ b/Engine/Plugins/Importers/USDImporter/Source/USDSchemas/Private/USDGeometryCacheTranslator.cpp
@@ -67,15 +67,6 @@ static FAutoConsoleVariableRef CVarEnableSubdiv(
 		 "StaticMeshes")
 );
 
-static bool GForceImport = true;
-static FAutoConsoleVariableRef CVarForceImport(
-	TEXT("USD.GeometryCache.ForceImport"),
-	GForceImport,
-	TEXT(
-		"Whether to immediately cache all geometry cache frames into the asset when the stage opens, acting like a fully imported Geometry Cache asset"
-	)
-);
-
 namespace UsdGeometryCacheTranslatorImpl
 {
 	bool ProcessGeometryCacheMaterials(
@@ -405,6 +396,16 @@ namespace UsdGeometryCacheTranslatorImpl
 		return StreamableTrack;
 	}
 
+	void FillGeometryCacheTracks(
+		const FString& RootPrimPath,
+		const TArray<UE::FSdfPath>& MeshPrims,
+		const TArray<int32>& MaterialOffsets,
+		TSharedRef<FUsdSchemaTranslationContext> Context,
+		UGeometryCache* GeometryCache
+	);
+
+	void FinalizeGeometryCache(UGeometryCache* GeometryCache);
+
 	UGeometryCache* CreateGeometryCache(
 		const UE::FUsdPrim& RootPrim,
 		const FMeshDescription& MeshDescription,
@@ -423,15 +424,15 @@ namespace UsdGeometryCacheTranslatorImpl
 		FSHAHash MeshHash = FStaticMeshOperations::ComputeSHAHash(MeshDescription);
 		SHA1.Update(&MeshHash.Hash[0], sizeof(MeshHash.Hash));
 
-		const bool bIsImporting = GForceImport || Context->bIsImporting;
+		const bool bIsImporting = Context->bIsImporting || Context->GeometryCacheImport == EGeometryCacheImport::OnLoad;
 
 		// Frame rate must be taken into account as well since different frame rates must produce different sampling in the tracks
 		SHA1.Update(reinterpret_cast<uint8*>(&Args.FramesPerSecond), sizeof(Args.FramesPerSecond));
 		SHA1.Update(reinterpret_cast<uint8*>(&Args.StartFrame), sizeof(Args.StartFrame));
 		SHA1.Update(reinterpret_cast<uint8*>(&Args.EndFrame), sizeof(Args.EndFrame));
 
-		// Track type depends on if it's importing or not. Import needs to generate a persistent asset with all the frames already sampled
-		SHA1.Update(reinterpret_cast<const uint8*>(&bIsImporting), sizeof(bIsImporting));
+		// Track type depends on how geometry caches are handled. Import needs to generate a persistent asset with all the frames already sampled
+		SHA1.Update(reinterpret_cast<const uint8*>(&Context->GeometryCacheImport), sizeof(Context->GeometryCacheImport));
 		SHA1.Final();
 
 		FSHAHash GeoCacheHash;
@@ -440,15 +441,69 @@ namespace UsdGeometryCacheTranslatorImpl
 
 		const FString DesiredName = FPaths::GetBaseFilename(RootPrimPath);
 
+		// In Never import mode, make the geometry cache transient so it doesn't get saved to disk. It will get recreated since it's lightweight.
+		EObjectFlags ObjectFlags = Context->ObjectFlags;
+		if (Context->GeometryCacheImport == EGeometryCacheImport::Never)
+		{
+			ObjectFlags |= RF_Transient;
+		}
+
 		UGeometryCache* GeometryCache = Context->UsdAssetCache->GetOrCreateCachedAsset<UGeometryCache>(
 			PrefixedGeoCacheHash,
 			DesiredName,
-			Context->ObjectFlags,
+			ObjectFlags,
 			&bOutIsNew
 		);
 
 		if (GeometryCache && bOutIsNew)
 		{
+			if (Context->GeometryCacheImport == EGeometryCacheImport::OnSave)
+			{
+				// In OnSave import mode, register a PreSave callback to convert the USD tracks to streamable tracks
+				GeometryCache->OnPreSave = UGeometryCache::FOnPreSave::CreateLambda(
+					[MeshPaths, RootPrimPath, MaterialOffsets, Context](UGeometryCache* GeometryCache)
+					{
+						// Convert only if there's any USD tracks
+						bool bHasUsdTracks = false;
+						for (UGeometryCacheTrack* Track : GeometryCache->Tracks)
+						{
+							if (UGeometryCacheTrackUsd* UsdTrack = Cast<UGeometryCacheTrackUsd>(Track))
+							{
+								bHasUsdTracks = true;
+								// Make sure to unregister the USD track from the streamer since it will get replaced with a streamable track
+								UsdTrack->UnregisterStream();
+							}
+						}
+
+						if (!bHasUsdTracks)
+						{
+							return;
+						}
+
+						GeometryCache->Tracks.Reset();
+
+						// Create a track for each mesh to be processed and add it to the GeometryCache
+						for (int32 Index = 0; Index < MeshPaths.Num(); ++Index)
+						{
+							const FString& PrimPath = MeshPaths[Index].GetString();
+							UGeometryCacheTrack* Track = CreateStreamableTrack(GeometryCache, PrimPath);
+							GeometryCache->AddTrack(Track);
+
+							TArray<FMatrix> Mats;
+							Mats.Add(FMatrix::Identity);
+							Mats.Add(FMatrix::Identity);
+
+							TArray<float> MatTimes;
+							MatTimes.Add(0.0f);
+							MatTimes.Add(0.0f);
+							Track->SetMatrixSamples(Mats, MatTimes);
+						}
+
+						FillGeometryCacheTracks(RootPrimPath, MeshPaths, MaterialOffsets, Context, GeometryCache);
+
+						FinalizeGeometryCache(GeometryCache);
+					});
+			}
 			if (!bIsImporting)
 			{
 				// StartOffsetTime is the offset applied to the GeometryCache section on the sequencer track, so not relevant when importing
@@ -931,7 +986,7 @@ void FGeometryCacheCreateAssetsTaskChain::SetupTasks()
 				}
 			}
 
-			const bool bIsImporting = GForceImport || Context->bIsImporting;
+			const bool bIsImporting = Context->bIsImporting || Context->GeometryCacheImport == EGeometryCacheImport::OnLoad;
 
 			// Continue with the import steps
 			return bIsImporting && GeometryCache && bIsNew;
@@ -1014,8 +1069,7 @@ USceneComponent* FUsdGeometryCacheTranslator::CreateComponents()
 			return nullptr;
 		}
 
-		const bool bIsImporting = GForceImport || Context->bIsImporting;
-		SceneComponent = CreateComponentsEx({bIsImporting ? UGeometryCacheComponent::StaticClass() : UGeometryCacheUsdComponent::StaticClass()}, {});
+		SceneComponent = CreateComponentsEx({UGeometryCacheUsdComponent::StaticClass()}, {});
 	}
 	else
 	{
@@ -1178,7 +1232,7 @@ void FUsdGeometryCacheTranslator::UpdateComponents(USceneComponent* SceneCompone
 			UsdGroomTranslatorUtils::SetGroomFromPrim(GetPrim(), *Context->PrimLinkCache, SceneComponent);
 		}
 
-		const bool bIsImporting = GForceImport || Context->bIsImporting;
+		const bool bIsImporting = Context->bIsImporting || Context->GeometryCacheImport == EGeometryCacheImport::OnLoad;
 
 		// Defer to xformable translator to set our transforms, visibility, etc. but only when opening the stage: This will be baked in for import.
 		// Don't go through FUsdGeomMeshTranslator::UpdateComponents as it will want to create a static mesh if PrimPath is an animated mesh prim
diff --git a/Engine/Plugins/Importers/USDImporter/Source/USDSchemas/Public/USDSchemaTranslator.h b/Engine/Plugins/Importers/USDImporter/Source/USDSchemas/Public/USDSchemaTranslator.h
--- a/Engine/Plugins/Importers/USDImporter/Source/USDSchemas/Public/USDSchemaTranslator.h
+++ b/Engine/Plugins/Importers/USDImporter/Source/USDSchemas/Public/USDSchemaTranslator.h
@@ -224,6 +224,9 @@ struct USDSCHEMAS_API FUsdSchemaTranslationContext : public TSharedFromThis<FUsd
 	/** Describes what to add to the root bone animation within generated AnimSequences, if anything */
 	EUsdRootMotionHandling RootMotionHandling = EUsdRootMotionHandling::NoAdditionalRootMotion;
 
+	/** How geometry caches are handled in the stage workflow */
+	EGeometryCacheImport GeometryCacheImport = EGeometryCacheImport::Never;
+
 	/** Subdivision level to use for all subdivision meshes on the opened stage. 0 means "don't subdivide" */
 	int32 SubdivisionLevel = 0;
 
diff --git a/Engine/Plugins/Importers/USDImporter/Source/USDStage/Private/USDStageActor.cpp b/Engine/Plugins/Importers/USDImporter/Source/USDStage/Private/USDStageActor.cpp
--- a/Engine/Plugins/Importers/USDImporter/Source/USDStage/Private/USDStageActor.cpp
+++ b/Engine/Plugins/Importers/USDImporter/Source/USDStage/Private/USDStageActor.cpp
@@ -163,6 +163,7 @@ struct FUsdStageActorImpl
 		TranslationContext->RenderContext = StageActor->RenderContext;
 		TranslationContext->MaterialPurpose = StageActor->MaterialPurpose;
 		TranslationContext->RootMotionHandling = StageActor->RootMotionHandling;
+		TranslationContext->GeometryCacheImport = StageActor->GeometryCacheImport;
 		TranslationContext->SubdivisionLevel = StageActor->SubdivisionLevel;
 		TranslationContext->MetadataOptions = StageActor->MetadataOptions;
 		TranslationContext->BlendShapesByPath = &StageActor->BlendShapesByPath;
@@ -583,6 +584,7 @@ struct FUsdStageActorImpl
 			EventAttributes.Emplace(TEXT("RenderContext"), StageActor->RenderContext.ToString());
 			EventAttributes.Emplace(TEXT("MaterialPurpose"), StageActor->MaterialPurpose.ToString());
 			EventAttributes.Emplace(TEXT("RootMotionHandling"), LexToString((uint8)StageActor->RootMotionHandling));
+			EventAttributes.Emplace(TEXT("GeometryCacheImport"), LexToString((uint8)StageActor->GeometryCacheImport));
 			EventAttributes.Emplace(TEXT("SubdivisionLevel"), LexToString(StageActor->SubdivisionLevel));
 
 			UsdUtils::AddAnalyticsAttributes(StageActor->MetadataOptions, EventAttributes);
@@ -913,6 +915,7 @@ AUsdStageActor::AUsdStageActor()
 	: StageState(EUsdStageState::OpenedAndLoaded)
 	, InitialLoadSet(EUsdInitialLoadSet::LoadAll)
 	, InterpolationType(EUsdInterpolationType::Linear)
+	, GeometryCacheImport(EGeometryCacheImport::Never)
 	, KindsToCollapse((int32)(EUsdDefaultKind::Component | EUsdDefaultKind::Subcomponent))
 	, bMergeIdenticalMaterialSlots(true)
 	, bShareAssetsForIdenticalPrims(true)
@@ -2602,6 +2605,20 @@ void AUsdStageActor::SetInterpolationType(EUsdInterpolationType NewType)
 	LoadUsdStage();
 }
 
+void AUsdStageActor::SetGeometryCacheImport(EGeometryCacheImport ImportOption)
+{
+	if (ImportOption == GeometryCacheImport)
+	{
+		return;
+	}
+
+	const bool bMarkDirty = false;
+	Modify(bMarkDirty);
+
+	GeometryCacheImport = ImportOption;
+	LoadUsdStage();
+}
+
 void AUsdStageActor::SetKindsToCollapse(int32 NewKindsToCollapse)
 {
 	const EUsdDefaultKind NewEnum = (EUsdDefaultKind)NewKindsToCollapse;
@@ -4665,6 +4682,12 @@ void AUsdStageActor::HandlePropertyChangedEvent(FPropertyChangedEvent& PropertyC
 		InterpolationType = (EUsdInterpolationType) !((uint8)InterpolationType);
 		SetInterpolationType(CorrectInterpolationType);
 	}
+	else if (PropertyName == GET_MEMBER_NAME_CHECKED(AUsdStageActor, GeometryCacheImport))
+	{
+		const EGeometryCacheImport CorrectImportOption = GeometryCacheImport;
+		GeometryCacheImport = (EGeometryCacheImport) !((uint8)GeometryCacheImport);
+		SetGeometryCacheImport(CorrectImportOption);
+	}
 	else if (PropertyName == GET_MEMBER_NAME_CHECKED(AUsdStageActor, KindsToCollapse))
 	{
 		const int32 CorrectKindsToCollapse = KindsToCollapse;
diff --git a/Engine/Plugins/Importers/USDImporter/Source/USDStage/Public/USDStageActor.h b/Engine/Plugins/Importers/USDImporter/Source/USDStage/Public/USDStageActor.h
--- a/Engine/Plugins/Importers/USDImporter/Source/USDStage/Public/USDStageActor.h
+++ b/Engine/Plugins/Importers/USDImporter/Source/USDStage/Public/USDStageActor.h
@@ -81,6 +81,9 @@ public:
 	UPROPERTY(EditAnywhere, BlueprintReadOnly, Category = "USD", config)
 	EUsdInterpolationType InterpolationType;
 
+	UPROPERTY(EditAnywhere, BlueprintReadOnly, Category = "USD", config)
+	EGeometryCacheImport GeometryCacheImport;
+
 	/**
 	 * Whether to try to combine individual assets and components of the same type on a kind-per-kind basis,
 	 * like multiple Mesh prims into a single Static Mesh
@@ -187,6 +190,9 @@ public:
 	UFUNCTION(BlueprintCallable, Category = "USD", meta = (CallInEditor = "true"))
 	USDSTAGE_API void SetInterpolationType(EUsdInterpolationType NewType);
 
+	UFUNCTION(BlueprintCallable, Category = "USD", meta = (CallInEditor = "true"))
+	USDSTAGE_API void SetGeometryCacheImport(EGeometryCacheImport ImportOption);
+
 	UFUNCTION(BlueprintCallable, Category = "USD", meta = (CallInEditor = "true"))
 	USDSTAGE_API void SetKindsToCollapse(int32 NewKindsToCollapse);
 
diff --git a/Engine/Plugins/Importers/USDImporter/Source/USDStageEditor/Private/SUSDStage.cpp b/Engine/Plugins/Importers/USDImporter/Source/USDStageEditor/Private/SUSDStage.cpp
--- a/Engine/Plugins/Importers/USDImporter/Source/USDStageEditor/Private/SUSDStage.cpp
+++ b/Engine/Plugins/Importers/USDImporter/Source/USDStageEditor/Private/SUSDStage.cpp
@@ -1106,6 +1106,12 @@ void SUsdStage::FillOptionsMenu(FMenuBuilder& MenuBuilder)
 			LOCTEXT("InterpolationType_ToolTip", "Whether to interpolate between time samples linearly or with 'held' (i.e. constant) interpolation"),
 			FNewMenuDelegate::CreateSP(this, &SUsdStage::FillInterpolationTypeSubMenu)
 		);
+
+		MenuBuilder.AddSubMenu(
+			LOCTEXT("GeometryCacheImport", "Geometry Cache Import"),
+			LOCTEXT("GeometryCacheImport_ToolTip", "When to convert geometry caches to persistent assets (in stage workflow)"),
+			FNewMenuDelegate::CreateSP(this, &SUsdStage::FillGeometryCacheImportSubMenu)
+		);
 	}
 	MenuBuilder.EndSection();
 
@@ -2106,6 +2112,142 @@ void SUsdStage::FillInterpolationTypeSubMenu(FMenuBuilder& MenuBuilder)
 	);
 }
 
+void SUsdStage::FillGeometryCacheImportSubMenu(FMenuBuilder& MenuBuilder)
+{
+	MenuBuilder.AddMenuEntry(
+		LOCTEXT("ImportNever", "Never"),
+		LOCTEXT("ImportNever_ToolTip", "Geometry Caches are not imported as persistent assets and are always streamed from the stage instead."),
+		FSlateIcon(),
+		FUIAction(
+			FExecuteAction::CreateLambda(
+				[this]()
+				{
+					if (AUsdStageActor* StageActor = GetStageActorOrCDO())
+					{
+						FScopedTransaction Transaction(FText::Format(
+							LOCTEXT("SetGeometryCacheImportNever", "Set USD stage actor '{0}' to not import geometry caches"),
+							FText::FromString(StageActor->GetActorLabel())
+						));
+
+						// c.f. comment in SUsdStage::FillCollapsingSubMenu
+						TGuardValue<bool> MaintainSelectionGuard(bUpdatingViewportSelection, true);
+
+						StageActor->SetGeometryCacheImport(EGeometryCacheImport::Never);
+						if (StageActor->IsTemplate())
+						{
+							StageActor->SaveConfig();
+						}
+					}
+				}
+			),
+			FCanExecuteAction{},
+			FIsActionChecked::CreateLambda(
+				[this]()
+				{
+					if (AUsdStageActor* StageActor = GetStageActorOrCDO())
+					{
+						return StageActor->GeometryCacheImport == EGeometryCacheImport::Never;
+					}
+					return false;
+				}
+			)
+		),
+		NAME_None,
+		EUserInterfaceActionType::RadioButton
+	);
+
+	MenuBuilder.AddMenuEntry(
+		LOCTEXT("ImportOnLoad", "On Load"),
+		LOCTEXT(
+			"ImportOnLoad_ToolTip",
+			"Geometry Caches are imported as persistents assets on stage load and played back from them."
+		),
+		FSlateIcon(),
+		FUIAction(
+			FExecuteAction::CreateLambda(
+				[this]()
+				{
+					if (AUsdStageActor* StageActor = GetStageActorOrCDO())
+					{
+						FScopedTransaction Transaction(FText::Format(
+							LOCTEXT("SetGeometryCacheImportOnLoad", "Set USD stage actor '{0}' to import geometry caches on stage load"),
+							FText::FromString(StageActor->GetActorLabel())
+						));
+
+						// c.f. comment in SUsdStage::FillCollapsingSubMenu
+						TGuardValue<bool> MaintainSelectionGuard(bUpdatingViewportSelection, true);
+
+						StageActor->SetGeometryCacheImport(EGeometryCacheImport::OnLoad);
+						if (StageActor->IsTemplate())
+						{
+							StageActor->SaveConfig();
+						}
+					}
+				}
+			),
+			FCanExecuteAction{},
+			FIsActionChecked::CreateLambda(
+				[this]()
+				{
+					if (AUsdStageActor* StageActor = GetStageActorOrCDO())
+					{
+						return StageActor->GeometryCacheImport == EGeometryCacheImport::OnLoad;
+					}
+					return false;
+				}
+			)
+		),
+		NAME_None,
+		EUserInterfaceActionType::RadioButton
+	);
+
+	MenuBuilder.AddMenuEntry(
+		LOCTEXT("ImportOnSave", "On Save"),
+		LOCTEXT(
+			"ImportOnSave_ToolTip",
+			"Geometry Caches are streamed from the stage until they are saved at which time they will be imported as persistent assets "
+			"and played back from them instead."
+		),
+		FSlateIcon(),
+		FUIAction(
+			FExecuteAction::CreateLambda(
+				[this]()
+				{
+					if (AUsdStageActor* StageActor = GetStageActorOrCDO())
+					{
+						FScopedTransaction Transaction(FText::Format(
+							LOCTEXT("SetGeometryCacheImportOnSave", "Set USD stage actor '{0}' to import geometry caches on save"),
+							FText::FromString(StageActor->GetActorLabel())
+						));
+
+						// c.f. comment in SUsdStage::FillCollapsingSubMenu
+						TGuardValue<bool> MaintainSelectionGuard(bUpdatingViewportSelection, true);
+
+						StageActor->SetGeometryCacheImport(EGeometryCacheImport::OnSave);
+						if (StageActor->IsTemplate())
+						{
+							StageActor->SaveConfig();
+						}
+					}
+				}
+			),
+			FCanExecuteAction{},
+			FIsActionChecked::CreateLambda(
+				[this]()
+				{
+					if (AUsdStageActor* StageActor = GetStageActorOrCDO())
+					{
+						return StageActor->GeometryCacheImport == EGeometryCacheImport::OnSave;
+					}
+					return false;
+				}
+			)
+		),
+		NAME_None,
+		EUserInterfaceActionType::RadioButton
+	);
+}
+
 void SUsdStage::FillSelectionSubMenu(FMenuBuilder& MenuBuilder)
 {
 	MenuBuilder.AddMenuEntry(
diff --git a/Engine/Plugins/Importers/USDImporter/Source/USDStageEditor/Private/SUSDStage.h b/Engine/Plugins/Importers/USDImporter/Source/USDStageEditor/Private/SUSDStage.h
--- a/Engine/Plugins/Importers/USDImporter/Source/USDStageEditor/Private/SUSDStage.h
+++ b/Engine/Plugins/Importers/USDImporter/Source/USDStageEditor/Private/SUSDStage.h
@@ -100,6 +100,7 @@ protected:
 	void FillInterpolationTypeSubMenu(FMenuBuilder& MenuBuilder);
 	void FillSelectionSubMenu(FMenuBuilder& MenuBuilder);
 	void FillNaniteThresholdSubMenu(FMenuBuilder& MenuBuilder);
+	void FillGeometryCacheImportSubMenu(FMenuBuilder& MenuBuilder);
 
 	void OnLayerIsolated(const UE::FSdfLayer& IsolatedLayer);
 
diff --git a/Engine/Plugins/Runtime/GeometryCache/Source/GeometryCache/Classes/GeometryCache.h b/Engine/Plugins/Runtime/GeometryCache/Source/GeometryCache/Classes/GeometryCache.h
--- a/Engine/Plugins/Runtime/GeometryCache/Source/GeometryCache/Classes/GeometryCache.h
+++ b/Engine/Plugins/Runtime/GeometryCache/Source/GeometryCache/Classes/GeometryCache.h
@@ -3,6 +3,7 @@
 
 #include "Interfaces/Interface_AssetUserData.h"
 #include "RenderCommandFence.h"
+#include "UObject/ObjectSaveContext.h"
 
 #include "GeometryCache.generated.h"
 
@@ -22,6 +23,7 @@ class GEOMETRYCACHE_API UGeometryCache : public UObject, public IInterface_Asset
 	GENERATED_UCLASS_BODY()
 public:
 	//~ Begin UObject Interface.
+	virtual void PreSave(FObjectPreSaveContext SaveContext) override;
 	virtual void Serialize(FArchive& Ar) override;
 	virtual void PostInitProperties() override;
 	virtual FString GetDesc() override;
@@ -100,6 +102,9 @@ public:
 	virtual const TArray<UAssetUserData*>* GetAssetUserDataArray() const override;
 	//~ End IInterface_AssetUserData Interface
 
+	DECLARE_DELEGATE_OneParam(FOnPreSave, UGeometryCache*);
+	FOnPreSave OnPreSave;
+
 private:
 	/** A fence which is used to keep track of the rendering thread releasing the geometry cache resources. */
 	FRenderCommandFence ReleaseResourcesFence;
diff --git a/Engine/Plugins/Runtime/GeometryCache/Source/GeometryCache/Private/GeometryCache.cpp b/Engine/Plugins/Runtime/GeometryCache/Source/GeometryCache/Private/GeometryCache.cpp
--- a/Engine/Plugins/Runtime/GeometryCache/Source/GeometryCache/Private/GeometryCache.cpp
+++ b/Engine/Plugins/Runtime/GeometryCache/Source/GeometryCache/Private/GeometryCache.cpp
@@ -37,6 +37,11 @@ void UGeometryCache::PostInitProperties()
 	Super::PostInitProperties();
 }
 
+void UGeometryCache::PreSave(FObjectPreSaveContext SaveContext)
+{
+	OnPreSave.ExecuteIfBound(this);
+}
+
 void UGeometryCache::Serialize(FArchive& Ar)
 {
 	Ar.UsingCustomVersion(FAnimPhysObjectVersion::GUID);
diff --git a/Engine/Plugins/Runtime/USDCore/Source/UnrealUSDWrapper/Public/UnrealUSDWrapper.h b/Engine/Plugins/Runtime/USDCore/Source/UnrealUSDWrapper/Public/UnrealUSDWrapper.h
--- a/Engine/Plugins/Runtime/USDCore/Source/UnrealUSDWrapper/Public/UnrealUSDWrapper.h
+++ b/Engine/Plugins/Runtime/USDCore/Source/UnrealUSDWrapper/Public/UnrealUSDWrapper.h
@@ -148,6 +148,21 @@ enum class EUsdRootMotionHandling : uint8
 	UseMotionFromSkeleton
 };
 
+/** Stage option on how to handle geometry caches in stage workflow */
+UENUM()
+enum class EGeometryCacheImport : uint8
+{
+	// Never imported (no persistent assets); always streamed from the USD stage
+	Never,
+
+	// Imported on stage load; played back from the persistent assets
+	OnLoad,
+
+	// Imported on save; geometry caches are streamed from the stage until they are saved. 
+	// Afterwards they are played back from the persistent assets
+	OnSave
+};
+
 /** Corresponds to pxr::UsdListPosition, refer to the USD SDK documentation */
 UENUM()
 enum class EUsdListPosition : uint8
