USD: Implement support for muting and reloading sublayers.

diff --git a/Engine/Plugins/Importers/USDImporter/Source/USDStageEditor/Private/SUSDLayersTreeView.cpp b/Engine/Plugins/Importers/USDImporter/Source/USDStageEditor/Private/SUSDLayersTreeView.cpp
--- a/Engine/Plugins/Importers/USDImporter/Source/USDStageEditor/Private/SUSDLayersTreeView.cpp
+++ b/Engine/Plugins/Importers/USDImporter/Source/USDStageEditor/Private/SUSDLayersTreeView.cpp
@@ -67,6 +67,220 @@
 		bool bForce = true;
 		OutputLayer.Save(bForce);
 	}
+
+	void ToggleMuteOrShowWarning(FUsdLayerViewModelRef LayerItem)
+	{
+		if (!LayerItem->IsValid() || !LayerItem->CanMuteLayer())
+		{
+			return;
+		}
+
+		// If the layer is not dirty we can just mute it without worry and early out
+		if (!LayerItem->IsLayerDirty())
+		{
+			LayerItem->ToggleMuteLayer();
+			return;
+		}
+
+		// Show a warning if the layer is dirty, as muting it will discard changes
+		const UUsdProjectSettings* Settings = GetDefault<UUsdProjectSettings>();
+		if (!LayerItem->IsLayerMuted() && Settings && Settings->bShowConfirmationWhenMutingDirtyLayers)
+		{
+			static TWeakPtr<SNotificationItem> Notification;
+
+			FNotificationInfo Toast(LOCTEXT("ConfirmMutingLayer", "Muting dirty layer"));
+			Toast.SubText = FText::Format(
+				LOCTEXT(
+					"ConfirmMutingLayer_Subtext",
+					"Layer '{0}' has unsaved changes that will be lost if muted.\n\nDo you wish to proceed with muting the layer?"
+				),
+				LayerItem->GetDisplayName()
+			);
+			Toast.Image = FCoreStyle::Get().GetBrush(TEXT("MessageLog.Warning"));
+			Toast.bUseLargeFont = false;
+			Toast.bFireAndForget = false;
+			Toast.FadeOutDuration = 0.0f;
+			Toast.ExpireDuration = 0.0f;
+			Toast.bUseThrobber = false;
+			Toast.bUseSuccessFailIcons = false;
+
+			Toast.ButtonDetails.Emplace(
+				LOCTEXT("ConfirmMutingOkAll", "Always proceed"),
+				FText::GetEmpty(),
+				FSimpleDelegate::CreateLambda(
+					[LayerItem]()
+					{
+						if (TSharedPtr<SNotificationItem> PinnedNotification = Notification.Pin())
+						{
+							PinnedNotification->SetCompletionState(SNotificationItem::CS_Success);
+							PinnedNotification->ExpireAndFadeout();
+
+							UUsdProjectSettings* Settings = GetMutableDefault<UUsdProjectSettings>();
+							Settings->bShowConfirmationWhenMutingDirtyLayers = false;
+							Settings->SaveConfig();
+
+							LayerItem->ToggleMuteLayer();
+						}
+					}
+				)
+			);
+
+			Toast.ButtonDetails.Emplace(
+				LOCTEXT("ConfirmMutingOk", "Proceed"),
+				FText::GetEmpty(),
+				FSimpleDelegate::CreateLambda(
+					[LayerItem]()
+					{
+						if (TSharedPtr<SNotificationItem> PinnedNotification = Notification.Pin())
+						{
+							PinnedNotification->SetCompletionState(SNotificationItem::CS_Success);
+							PinnedNotification->ExpireAndFadeout();
+
+							LayerItem->ToggleMuteLayer();
+						}
+					}
+				)
+			);
+
+			Toast.ButtonDetails.Emplace(
+				LOCTEXT("ConfirmMutingCancel", "Cancel"),
+				FText::GetEmpty(),
+				FSimpleDelegate::CreateLambda(
+					[]()
+					{
+						if (TSharedPtr<SNotificationItem> PinnedNotification = Notification.Pin())
+						{
+							PinnedNotification->SetCompletionState(SNotificationItem::CS_Fail);
+							PinnedNotification->ExpireAndFadeout();
+						}
+					}
+				)
+			);
+
+			// Only show one at a time
+			if (!Notification.IsValid())
+			{
+				Notification = FSlateNotificationManager::Get().AddNotification(Toast);
+			}
+
+			if (TSharedPtr<SNotificationItem> PinnedNotification = Notification.Pin())
+			{
+				PinnedNotification->SetCompletionState(SNotificationItem::CS_Pending);
+			}
+		}
+		else
+		{
+			// Don't show prompt, always just mute
+			LayerItem->ToggleMuteLayer();
+		}
+	}
+
+	void ReloadOrShowWarning(FUsdLayerViewModelRef LayerItem)
+	{
+		if (!LayerItem->IsValid() || !LayerItem->CanReload())
+		{
+			return;
+		}
+
+		// If the layer is not dirty we can just mute it without worry and early out
+		if (!LayerItem->IsLayerDirty())
+		{
+			LayerItem->Reload();
+			return;
+		}
+
+		// Show a warning if the layer is dirty, as muting it will discard changes
+		const UUsdProjectSettings* Settings = GetDefault<UUsdProjectSettings>();
+		if (!LayerItem->IsLayerMuted() && Settings && Settings->bShowConfirmationWhenReloadingDirtyLayers)
+		{
+			static TWeakPtr<SNotificationItem> Notification;
+
+			FNotificationInfo Toast(LOCTEXT("ConfirmReloadingDirtyLayer", "Reloading dirty layer"));
+			Toast.SubText = FText::Format(
+				LOCTEXT(
+					"ConfirmReloadingDirtyLayer_Subtext",
+					"Layer '{0}' has unsaved changes that will be lost if reloaded.\n\nDo you wish to proceed with reloading the layer?"
+				),
+				LayerItem->GetDisplayName()
+			);
+			Toast.Image = FCoreStyle::Get().GetBrush(TEXT("MessageLog.Warning"));
+			Toast.bUseLargeFont = false;
+			Toast.bFireAndForget = false;
+			Toast.FadeOutDuration = 0.0f;
+			Toast.ExpireDuration = 0.0f;
+			Toast.bUseThrobber = false;
+			Toast.bUseSuccessFailIcons = false;
+
+			Toast.ButtonDetails.Emplace(
+				LOCTEXT("ConfirmMutingOkAll", "Always proceed"),
+				FText::GetEmpty(),
+				FSimpleDelegate::CreateLambda(
+					[LayerItem]()
+					{
+						if (TSharedPtr<SNotificationItem> PinnedNotification = Notification.Pin())
+						{
+							PinnedNotification->SetCompletionState(SNotificationItem::CS_Success);
+							PinnedNotification->ExpireAndFadeout();
+
+							UUsdProjectSettings* Settings = GetMutableDefault<UUsdProjectSettings>();
+							Settings->bShowConfirmationWhenReloadingDirtyLayers = false;
+							Settings->SaveConfig();
+
+							LayerItem->Reload();
+						}
+					}
+				)
+			);
+
+			Toast.ButtonDetails.Emplace(
+				LOCTEXT("ConfirmMutingOk", "Proceed"),
+				FText::GetEmpty(),
+				FSimpleDelegate::CreateLambda(
+					[LayerItem]()
+					{
+						if (TSharedPtr<SNotificationItem> PinnedNotification = Notification.Pin())
+						{
+							PinnedNotification->SetCompletionState(SNotificationItem::CS_Success);
+							PinnedNotification->ExpireAndFadeout();
+
+							LayerItem->Reload();
+						}
+					}
+				)
+			);
+
+			Toast.ButtonDetails.Emplace(
+				LOCTEXT("ConfirmMutingCancel", "Cancel"),
+				FText::GetEmpty(),
+				FSimpleDelegate::CreateLambda(
+					[]()
+					{
+						if (TSharedPtr<SNotificationItem> PinnedNotification = Notification.Pin())
+						{
+							PinnedNotification->SetCompletionState(SNotificationItem::CS_Fail);
+							PinnedNotification->ExpireAndFadeout();
+						}
+					}
+				)
+			);
+
+			// Only show one at a time
+			if (!Notification.IsValid())
+			{
+				Notification = FSlateNotificationManager::Get().AddNotification(Toast);
+			}
+
+			if (TSharedPtr<SNotificationItem> PinnedNotification = Notification.Pin())
+			{
+				PinnedNotification->SetCompletionState(SNotificationItem::CS_Pending);
+			}
+		}
+		else
+		{
+			// Don't show prompt, always just reload
+			LayerItem->Reload();
+		}
+	}
 }	 // namespace UE::USDLayersTreeViewImpl::Private
 
 class FUsdLayerNameColumn
@@ -112,7 +326,7 @@
 public:
 	FReply OnClicked(const FUsdLayerViewModelRef TreeItem)
 	{
-		ToggleMuteLayer(TreeItem);
+		UE::USDLayersTreeViewImpl::Private::ToggleMuteOrShowWarning(TreeItem);
 
 		return FReply::Handled();
 	}
@@ -121,7 +335,7 @@
 	{
 		const bool bIsButtonHovered = Button.IsValid() && Button->IsHovered();
 
-		if (!CanMuteLayer(TreeItem))
+		if (!TreeItem->CanMuteLayer())
 		{
 			return nullptr;
 		}
@@ -177,7 +391,7 @@
 				.HeightOverride(ItemSize)
 				.WidthOverride(ItemSize)
 				.Visibility(EVisibility::Visible)
-				.ToolTip(SNew(SToolTip).Text(LOCTEXT("CantMuteLayerTooltip", "This layer cannot be muted!")));
+				.ToolTip(SNew(SToolTip).Text(LOCTEXT("CantMuteLayerTooltip", "Cannot mute root layers or edit targets")));
 		}
 
 		TSharedPtr<SButton> Button = SNew(SButton)
@@ -211,119 +425,6 @@
 			];
 		// clang-format on
 	}
-
-protected:
-	bool CanMuteLayer(FUsdLayerViewModelRef LayerItem) const
-	{
-		return LayerItem->IsValid() && LayerItem->CanMuteLayer();
-	}
-
-	void ToggleMuteLayer(FUsdLayerViewModelRef LayerItem)
-	{
-		if (!CanMuteLayer(LayerItem))
-		{
-			return;
-		}
-
-		// If the layer is not dirty we can just mute it without worry and early out
-		if (!LayerItem->IsLayerDirty())
-		{
-			LayerItem->ToggleMuteLayer();
-			return;
-		}
-
-		// Show a warning if the layer is dirty, as muting it will discard changes
-		const UUsdProjectSettings* Settings = GetDefault<UUsdProjectSettings>();
-		if (Settings && Settings->bShowConfirmationWhenMutingDirtyLayers)
-		{
-			static TWeakPtr<SNotificationItem> Notification;
-
-			FNotificationInfo Toast(LOCTEXT("ConfirmMutingLayer", "Muting dirty layer"));
-			Toast.SubText = FText::Format(
-				LOCTEXT(
-					"ConfirmMutingLayer_Subtext",
-					"Layer '{0}' has unsaved changes that will be lost if muted.\n\nDo you wish to proceed muting the layer?"
-				),
-				LayerItem->GetDisplayName()
-			);
-			Toast.Image = FCoreStyle::Get().GetBrush(TEXT("MessageLog.Warning"));
-			Toast.bUseLargeFont = false;
-			Toast.bFireAndForget = false;
-			Toast.FadeOutDuration = 0.0f;
-			Toast.ExpireDuration = 0.0f;
-			Toast.bUseThrobber = false;
-			Toast.bUseSuccessFailIcons = false;
-
-			Toast.ButtonDetails.Emplace(
-				LOCTEXT("ConfirmMutingOkAll", "Always proceed"),
-				FText::GetEmpty(),
-				FSimpleDelegate::CreateLambda(
-					[LayerItem]()
-					{
-						if (TSharedPtr<SNotificationItem> PinnedNotification = Notification.Pin())
-						{
-							PinnedNotification->SetCompletionState(SNotificationItem::CS_Success);
-							PinnedNotification->ExpireAndFadeout();
-
-							UUsdProjectSettings* Settings = GetMutableDefault<UUsdProjectSettings>();
-							Settings->bShowConfirmationWhenMutingDirtyLayers = false;
-							Settings->SaveConfig();
-
-							LayerItem->ToggleMuteLayer();
-						}
-					}
-				)
-			);
-
-			Toast.ButtonDetails.Emplace(
-				LOCTEXT("ConfirmMutingOk", "Proceed"),
-				FText::GetEmpty(),
-				FSimpleDelegate::CreateLambda(
-					[LayerItem]()
-					{
-						if (TSharedPtr<SNotificationItem> PinnedNotification = Notification.Pin())
-						{
-							PinnedNotification->SetCompletionState(SNotificationItem::CS_Success);
-							PinnedNotification->ExpireAndFadeout();
-
-							LayerItem->ToggleMuteLayer();
-						}
-					}
-				)
-			);
-
-			Toast.ButtonDetails.Emplace(
-				LOCTEXT("ConfirmMutingCancel", "Cancel"),
-				FText::GetEmpty(),
-				FSimpleDelegate::CreateLambda(
-					[]()
-					{
-						if (TSharedPtr<SNotificationItem> PinnedNotification = Notification.Pin())
-						{
-							PinnedNotification->SetCompletionState(SNotificationItem::CS_Fail);
-							PinnedNotification->ExpireAndFadeout();
-						}
-					}
-				)
-			);
-
-			// Only show one at a time
-			if (!Notification.IsValid())
-			{
-				Notification = FSlateNotificationManager::Get().AddNotification(Toast);
-			}
-
-			if (TSharedPtr<SNotificationItem> PinnedNotification = Notification.Pin())
-			{
-				PinnedNotification->SetCompletionState(SNotificationItem::CS_Pending);
-			}
-		}
-		else
-		{
-			// Don't show prompt, always just mute
-			LayerItem->ToggleMuteLayer();
-		}
-	}
 };
 
 class FUsdLayerEditColumn
@@ -1140,5 +1241,57 @@
 		);
 
 		LayerOptions.AddMenuEntry(
+			TAttribute<FText>::Create(TAttribute<FText>::FGetter::CreateLambda(
+				[this]()
+				{
+					bool bAnyIsMuted = false;
+					TArray<FUsdLayerViewModelRef> MySelectedItems = GetSelectedItems();
+					for (const FUsdLayerViewModelRef& SelectedItem : MySelectedItems)
+					{
+						if (SelectedItem->IsLayerMuted())
+						{
+							bAnyIsMuted = true;
+							break;
+						}
+					}
+
+					return bAnyIsMuted ? LOCTEXT("UnmuteLayerRightClick_Text", "Unmute") : LOCTEXT("MuteLayerRightClick_Text", "Mute");
+				}
+			)),
+			TAttribute<FText>::Create(TAttribute<FText>::FGetter::CreateLambda(
+				[this]()
+				{
+					const bool bCanMute = CanMuteSelectedLayer();
+
+					bool bAnyIsMuted = false;
+					if (bCanMute)
+					{
+						TArray<FUsdLayerViewModelRef> MySelectedItems = GetSelectedItems();
+						for (const FUsdLayerViewModelRef& SelectedItem : MySelectedItems)
+						{
+							if (SelectedItem->IsLayerMuted())
+							{
+								bAnyIsMuted = true;
+								break;
+							}
+						}
+					}
+
+					return bCanMute
+							   ? (bAnyIsMuted ? LOCTEXT("UnmuteLayerRightClick_ToolTip", "Unmute the selected muted layers")
+											  : LOCTEXT("MuteLayerRightClick_ToolTip", "Mute all selected layers"))
+							   : LOCTEXT("CantMuteLayerRightClick_ToolTip", "Cannot mute root layers or edit targets");
+				}
+			)),
+			FSlateIcon(),
+			FUIAction(
+				FExecuteAction::CreateSP(this, &SUsdLayersTreeView::OnMuteSelectedLayer),
+				FCanExecuteAction::CreateSP(this, &SUsdLayersTreeView::CanMuteSelectedLayer)
+			),
+			NAME_None,
+			EUserInterfaceActionType::Button
+		);
+
+		LayerOptions.AddMenuEntry(
 			LOCTEXT("ClearLayer", "Clear"),
 			LOCTEXT("ClearLayer_ToolTip", "Clears this layer of all data"),
@@ -1152,5 +1305,17 @@
 		);
 
 		LayerOptions.AddMenuEntry(
+			LOCTEXT("ReloadLayer", "Reload"),
+			LOCTEXT("ReloadLayer_ToolTip", "Reloads just this layer from disk"),
+			FSlateIcon(),
+			FUIAction(
+				FExecuteAction::CreateSP(this, &SUsdLayersTreeView::OnReloadSelectedLayers),
+				FCanExecuteAction::CreateSP(this, &SUsdLayersTreeView::CanReloadSelectedLayers)
+			),
+			NAME_None,
+			EUserInterfaceActionType::Button
+		);
+
+		LayerOptions.AddMenuEntry(
 			LOCTEXT("SaveLayer", "Save"),
 			LOCTEXT("SaveLayer_ToolTip", "Saves the layer modifications to disk"),
@@ -1309,6 +1474,55 @@
 	}
 }
 
+bool SUsdLayersTreeView::CanMuteSelectedLayer() const
+{
+	TArray<FUsdLayerViewModelRef> MySelectedItems = GetSelectedItems();
+
+	if (MySelectedItems.Num() < 1)
+	{
+		return false;
+	}
+
+	for (FUsdLayerViewModelRef SelectedItem : MySelectedItems)
+	{
+		if (!SelectedItem->CanMuteLayer())
+		{
+			return false;
+		}
+	}
+
+	return true;
+}
+
+void SUsdLayersTreeView::OnMuteSelectedLayer()
+{
+	TArray<FUsdLayerViewModelRef> MySelectedItems = GetSelectedItems();
+	if (MySelectedItems.Num() < 1)
+	{
+		return;
+	}
+
+	bool bAnyIsMuted = false;
+	for (FUsdLayerViewModelRef SelectedItem : MySelectedItems)
+	{
+		if (SelectedItem->IsLayerMuted())
+		{
+			bAnyIsMuted = true;
+			break;
+		}
+	}
+
+	for (FUsdLayerViewModelRef SelectedItem : MySelectedItems)
+	{
+		if ((bAnyIsMuted && SelectedItem->IsLayerMuted()) || !bAnyIsMuted)
+		{
+			UE::USDLayersTreeViewImpl::Private::ToggleMuteOrShowWarning(SelectedItem);
+		}
+	}
+
+	return;
+}
+
 void SUsdLayersTreeView::OnClearSelectedLayers()
 {
 	TArray<FUsdLayerViewModelRef> MySelectedItems = GetSelectedItems();
@@ -1454,6 +1668,31 @@
 	return false;
 }
 
+void SUsdLayersTreeView::OnReloadSelectedLayers()
+{
+	FScopedTransaction Transaction(LOCTEXT("ReloadLayersTransaction", "Reload selected layers"));
+
+	TArray<FUsdLayerViewModelRef> MySelectedItems = GetSelectedItems();
+	for (FUsdLayerViewModelRef SelectedItem : MySelectedItems)
+	{
+		UE::USDLayersTreeViewImpl::Private::ReloadOrShowWarning(SelectedItem);
+	}
+}
+
+bool SUsdLayersTreeView::CanReloadSelectedLayers() const
+{
+	TArray<FUsdLayerViewModelRef> MySelectedItems = GetSelectedItems();
+	for (FUsdLayerViewModelRef SelectedItem : MySelectedItems)
+	{
+		if (!SelectedItem->CanReload())
+		{
+			return false;
+		}
+	}
+
+	return true;
+}
+
 void SUsdLayersTreeView::OnSaveSelectedLayers()
 {
 	TArray<FUsdLayerViewModelRef> MySelectedItems = GetSelectedItems();
@@ -1605,6 +1844,8 @@
 		return;
 	}
 
+	FScopedTransaction Transaction(LOCTEXT("AddExistingSublayerTransaction", "Add existing sublayer"));
+
 	TArray<FUsdLayerViewModelRef> MySelectedItems = GetSelectedItems();
 
 	for (FUsdLayerViewModelRef SelectedItem : MySelectedItems)
@@ -1625,6 +1866,8 @@
 		return;
 	}
 
+	FScopedTransaction Transaction(LOCTEXT("AddNewSublayerTransaction", "Add new sublayer"));
+
 	TArray<FUsdLayerViewModelRef> MySelectedItems = GetSelectedItems();
 	for (FUsdLayerViewModelRef SelectedItem : MySelectedItems)
 	{
@@ -1666,6 +1909,8 @@
 
 	TArray<FUsdLayerViewModelRef> SelectedLayers = GetSelectedItems();
 
+	FScopedTransaction Transaction(LOCTEXT("RemoveSublayerTransaction", "Remove sublayers"));
+
 	for (FUsdLayerViewModelRef SelectedLayer : SelectedLayers)
 	{
 		if (!CanRemoveLayer(SelectedLayer))
diff --git a/Engine/Plugins/Importers/USDImporter/Source/USDStageEditor/Private/SUSDLayersTreeView.h b/Engine/Plugins/Importers/USDImporter/Source/USDStageEditor/Private/SUSDLayersTreeView.h
--- a/Engine/Plugins/Importers/USDImporter/Source/USDStageEditor/Private/SUSDLayersTreeView.h
+++ b/Engine/Plugins/Importers/USDImporter/Source/USDStageEditor/Private/SUSDLayersTreeView.h
@@ -61,9 +61,15 @@ private:
 	bool CanEditSelectedLayer() const;
 	void OnEditSelectedLayer();
 
+	bool CanMuteSelectedLayer() const;
+	void OnMuteSelectedLayer();
+
 	void OnClearSelectedLayers();
 	bool CanClearSelectedLayers() const;
 
+	void OnReloadSelectedLayers();
+	bool CanReloadSelectedLayers() const;
+
 	void OnSaveSelectedLayers();
 	bool CanSaveSelectedLayers() const;
 
diff --git a/Engine/Plugins/Importers/USDImporter/Source/USDStageEditorViewModels/Private/USDLayersViewModel.cpp b/Engine/Plugins/Importers/USDImporter/Source/USDStageEditorViewModels/Private/USDLayersViewModel.cpp
--- a/Engine/Plugins/Importers/USDImporter/Source/USDStageEditorViewModels/Private/USDLayersViewModel.cpp
+++ b/Engine/Plugins/Importers/USDImporter/Source/USDStageEditorViewModels/Private/USDLayersViewModel.cpp
@@ -164,7 +164,7 @@ void FUsdLayerViewModel::RefreshData()
 	const pxr::SdfLayerHandle& EditTargetLayer = UsdStageRef->GetEditTarget().GetLayer();
 	const pxr::SdfLayerHandle& FocusedEditTarget = IsolatedStageRef ? IsolatedStageRef->GetEditTarget().GetLayer() : pxr::SdfLayerHandle{};
 
-    LayerModel->bIsEditTarget = (IsolatedStageRef) ? LayerModel->bIsInIsolatedStage && FocusedEditTarget == UsdLayer : EditTargetLayer == UsdLayer;
+	LayerModel->bIsEditTarget = (IsolatedStageRef) ? LayerModel->bIsInIsolatedStage && FocusedEditTarget == UsdLayer : EditTargetLayer == UsdLayer;
 
 	LayerModel->bIsMuted = (IsolatedStageRef && LayerModel->bIsInIsolatedStage)
 							   // If we isolating, we're only muted if we're muted on the isolated stage
@@ -327,3 +327,14 @@ bool FUsdLayerViewModel::IsLayerDirty() const
 {
 	return LayerModel->bIsDirty;
 }
+
+bool FUsdLayerViewModel::CanReload() const
+{
+	return !GetLayer().IsAnonymous();
+}
+
+void FUsdLayerViewModel::Reload()
+{
+	const bool bForce = true;
+	GetLayer().Reload(bForce);
+}
diff --git a/Engine/Plugins/Importers/USDImporter/Source/USDStageEditorViewModels/Public/USDLayersViewModel.h b/Engine/Plugins/Importers/USDImporter/Source/USDStageEditorViewModels/Public/USDLayersViewModel.h
--- a/Engine/Plugins/Importers/USDImporter/Source/USDStageEditorViewModels/Public/USDLayersViewModel.h
+++ b/Engine/Plugins/Importers/USDImporter/Source/USDStageEditorViewModels/Public/USDLayersViewModel.h
@@ -55,6 +55,9 @@ public:
 
 	bool IsLayerDirty() const;
 
+	bool CanReload() const;
+	void Reload();
+
 public:
 	TSharedRef<FUsdLayerModel> LayerModel;
 	FUsdLayerViewModel* ParentItem;
diff --git a/Engine/Plugins/Runtime/USDCore/Source/USDClasses/Public/USDProjectSettings.h b/Engine/Plugins/Runtime/USDCore/Source/USDClasses/Public/USDProjectSettings.h
--- a/Engine/Plugins/Runtime/USDCore/Source/USDClasses/Public/USDProjectSettings.h
+++ b/Engine/Plugins/Runtime/USDCore/Source/USDClasses/Public/USDProjectSettings.h
@@ -84,6 +84,9 @@ public:
 	UPROPERTY(config, EditAnywhere, Category = "USD|Dialogs")
 	bool bShowConfirmationWhenMutingDirtyLayers = true;
 
+	UPROPERTY(config, EditAnywhere, Category = "USD|Dialogs")
+	bool bShowConfirmationWhenReloadingDirtyLayers = true;
+
 	// Whether to show the warning dialog when authoring opinions that could have no effect on the composed stage
 	UPROPERTY(config, EditAnywhere, Category = "USD|Dialogs")
 	bool bShowOverriddenOpinionsWarning = true;
diff --git a/Engine/Plugins/Runtime/USDCore/Source/UnrealUSDWrapper/Private/UsdWrappers/SdfLayer.cpp b/Engine/Plugins/Runtime/USDCore/Source/UnrealUSDWrapper/Private/UsdWrappers/SdfLayer.cpp
--- a/Engine/Plugins/Runtime/USDCore/Source/UnrealUSDWrapper/Private/UsdWrappers/SdfLayer.cpp
+++ b/Engine/Plugins/Runtime/USDCore/Source/UnrealUSDWrapper/Private/UsdWrappers/SdfLayer.cpp
@@ -911,6 +911,17 @@ namespace UE
 #endif	  // #if USE_USD_SDK
 	}
 
+	template<typename PtrType>
+	void FSdfLayerBase<PtrType>::Reload(bool bForce)
+	{
+#if USE_USD_SDK
+		if (const PtrType& Ptr = Impl->GetInner())
+		{
+			Ptr->Reload(bForce);
+		}
+#endif	  // #if USE_USD_SDK
+	}
+
 	FString FSdfLayerUtils::SdfComputeAssetPathRelativeToLayer(const FSdfLayer& Anchor, const TCHAR* AssetPath)
 	{
 #if USE_USD_SDK
diff --git a/Engine/Plugins/Runtime/USDCore/Source/UnrealUSDWrapper/Public/UsdWrappers/SdfLayer.h b/Engine/Plugins/Runtime/USDCore/Source/UnrealUSDWrapper/Public/UsdWrappers/SdfLayer.h
--- a/Engine/Plugins/Runtime/USDCore/Source/UnrealUSDWrapper/Public/UsdWrappers/SdfLayer.h
+++ b/Engine/Plugins/Runtime/USDCore/Source/UnrealUSDWrapper/Public/UsdWrappers/SdfLayer.h
@@ -157,6 +157,8 @@ namespace UE
 
 		void Clear();
 
+		void Reload(bool bForce = false);
+
 		bool HasStartTimeCode() const;
 		double GetStartTimeCode() const;
 		void SetStartTimeCode(double TimeCode);
