USD: Implement support for removing and reloading referenced layers and payloads;

Display payloads on the references panel as well (the first column will show R/P depending on whether the item is a reference or payload);
Fix prim info additional panels clipping on top of the layers tree view when they ran out of space to display their content.

diff --git a/Engine/Plugins/Importers/USDImporter/Source/USDStageEditor/Private/SUSDPrimInfo.cpp b/Engine/Plugins/Importers/USDImporter/Source/USDStageEditor/Private/SUSDPrimInfo.cpp
--- a/Engine/Plugins/Importers/USDImporter/Source/USDStageEditor/Private/SUSDPrimInfo.cpp
+++ b/Engine/Plugins/Importers/USDImporter/Source/USDStageEditor/Private/SUSDPrimInfo.cpp
@@ -17,6 +17,22 @@
 
 void SUsdPrimInfo::Construct(const FArguments& InArgs)
 {
+	// auto as this has to fit the different types of list view
+	auto MakeListViewGetter = [](auto ListViewSharedPtrPtr) -> TAttribute<float>
+	{
+		return TAttribute<float>::Create(TAttribute<float>::FGetter::CreateLambda(
+			[ListViewSharedPtrPtr]()
+			{
+				if (ListViewSharedPtrPtr && (*ListViewSharedPtrPtr) && (*ListViewSharedPtrPtr)->GetItems().Num() > 0)
+				{
+					return 0.5f;
+				}
+
+				return 0.0f;
+			}
+		));
+	};
+
 	// clang-format off
 	ChildSlot
 	[
@@ -25,78 +41,58 @@ void SUsdPrimInfo::Construct(const FArguments& InArgs)
 		+SVerticalBox::Slot()
 		.FillHeight(1.f)
 		[
-			SNew(SBox)
-			.Content()
-			[
-				SAssignNew(PropertiesList, SUsdObjectFieldList)
-				.NameColumnText(LOCTEXT("NameColumnText", "Name"))
-				.OnSelectionChanged_Lambda([this](const TSharedPtr<FUsdObjectFieldViewModel>& NewSelection, ESelectInfo::Type SelectionType)
+			SAssignNew(PropertiesList, SUsdObjectFieldList)
+			.NameColumnText(LOCTEXT("NameColumnText", "Name"))
+			.OnSelectionChanged_Lambda([this](const TSharedPtr<FUsdObjectFieldViewModel>& NewSelection, ESelectInfo::Type SelectionType)
+			{
+				// Display property metadata if we have exactly one selected
+				TArray<FString> SelectedFields = PropertiesList->GetSelectedFieldNames();
+				if (PropertiesList && SelectedFields.Num() == 1 && NewSelection &&
+					(NewSelection->Type == EObjectFieldType::Attribute || NewSelection->Type == EObjectFieldType::Relationship)
+				)
+				{
+					PropertyMetadataPanel->SetObjectPath(
+						PropertiesList->GetUsdStage(),
+						*(FString{PropertiesList->GetObjectPath()} + "." + SelectedFields[0])
+					);
+					PropertyMetadataPanel->SetVisibility(EVisibility::Visible);
+				}
+				else
 				{
-					// Display property metadata if we have exactly one selected
-					TArray<FString> SelectedFields = PropertiesList->GetSelectedFieldNames();
-					if (PropertiesList && SelectedFields.Num() == 1 && NewSelection &&
-						(NewSelection->Type == EObjectFieldType::Attribute || NewSelection->Type == EObjectFieldType::Relationship)
-					)
-					{
-						PropertyMetadataPanel->SetObjectPath(
-							PropertiesList->GetUsdStage(),
-							*(FString{PropertiesList->GetObjectPath()} + "." + SelectedFields[0])
-						);
-						PropertyMetadataPanel->SetVisibility(EVisibility::Visible);
-					}
-					else
-					{
-						PropertyMetadataPanel->SetObjectPath({}, TEXT(""));
-						PropertyMetadataPanel->SetVisibility(EVisibility::Collapsed);
-					}
-				})
-			]
+					PropertyMetadataPanel->SetObjectPath({}, TEXT(""));
+					PropertyMetadataPanel->SetVisibility(EVisibility::Collapsed);
+				}
+			})
 		]
 
 		+SVerticalBox::Slot()
-		.AutoHeight()
+		.FillHeight(MakeListViewGetter(&PropertyMetadataPanel))
 		[
-			SNew(SBox)
-			.Content()
-			[
-				SAssignNew(PropertyMetadataPanel, SUsdObjectFieldList)
-				.NameColumnText_Lambda([this]() -> FText
-				{
-					FString PropertyName = FPaths::GetExtension(PropertyMetadataPanel->GetObjectPath());
-					return FText::FromString(FString::Printf(TEXT("%s metadata"), *PropertyName));
-				})
-				.Visibility(EVisibility::Collapsed)
-			]
+			SAssignNew(PropertyMetadataPanel, SUsdObjectFieldList)
+			.NameColumnText_Lambda([this]() -> FText
+			{
+				FString PropertyName = FPaths::GetExtension(PropertyMetadataPanel->GetObjectPath());
+				return FText::FromString(FString::Printf(TEXT("%s metadata"), *PropertyName));
+			})
+			.Visibility(EVisibility::Collapsed)
 		]
 
 		+SVerticalBox::Slot()
-		.AutoHeight()
+		.FillHeight(MakeListViewGetter(&IntegrationsPanel))
 		[
-			SNew(SBox)
-			.Content()
-			[
-				SAssignNew(IntegrationsPanel, SUsdIntegrationsPanel)
-			]
+			SAssignNew(IntegrationsPanel, SUsdIntegrationsPanel)
 		]
 
 		+SVerticalBox::Slot()
-		.AutoHeight()
+		.FillHeight(MakeListViewGetter(&VariantsList))
 		[
-			SNew(SBox)
-			.Content()
-			[
-				SAssignNew(VariantsList, SVariantsList)
-			]
+			SAssignNew(VariantsList, SVariantsList)
 		]
 
 		+SVerticalBox::Slot()
-		.AutoHeight()
+		.FillHeight(MakeListViewGetter(&ReferencesList))
 		[
-			SNew(SBox)
-			.Content()
-			[
-				SAssignNew(ReferencesList, SUsdReferencesList)
-			]
+			SAssignNew(ReferencesList, SUsdReferencesList)
 		]
 	];
 	// clang-format on
diff --git a/Engine/Plugins/Importers/USDImporter/Source/USDStageEditor/Private/Widgets/SUSDReferencesList.cpp b/Engine/Plugins/Importers/USDImporter/Source/USDStageEditor/Private/Widgets/SUSDReferencesList.cpp
--- a/Engine/Plugins/Importers/USDImporter/Source/USDStageEditor/Private/Widgets/SUSDReferencesList.cpp
+++ b/Engine/Plugins/Importers/USDImporter/Source/USDStageEditor/Private/Widgets/SUSDReferencesList.cpp
@@ -4,7 +4,10 @@
 
 #include "SUSDStageEditorStyle.h"
 
+#include "Framework/MultiBox/MultiBoxBuilder.h"
+#include "ScopedTransaction.h"
 #include "Styling/AppStyle.h"
+#include "USDStageActor.h"
 
 #if USE_USD_SDK
 
@@ -29,32 +32,37 @@ TSharedRef<SWidget> SUsdReferenceRow::GenerateWidgetForColumn(const FName& Colum
 	TSharedRef<SWidget> ColumnWidget = SNullWidget::NullWidget;
 
 	// clang-format off
-	if (ColumnName == TEXT("AssetPath"))
+	if (ColumnName == TEXT("ReferenceType"))
+	{
+		SAssignNew(ColumnWidget, STextBlock)
+		.Text(FText::FromString(Reference->bIsPayload ? TEXT("P") : TEXT("R")))
+		.ToolTipText(FText::FromString(Reference->bIsPayload ? TEXT("Payload") : TEXT("Reference")))
+		.Justification(ETextJustify::Center)
+		.Font(FAppStyle::GetFontStyle(UsdReferencesListConstants::NormalFont));
+	}
+
+	else if (ColumnName == TEXT("AssetPath"))
 	{
 		SAssignNew(ColumnWidget, STextBlock)
 		.Text(FText::FromString(Reference->AssetPath.IsEmpty() ? TEXT("(internal reference)") : Reference->AssetPath))
+		.Margin(FMargin(3.0f, 0.0f, 0.0f, 0.0f))
 		.Font(FAppStyle::GetFontStyle(UsdReferencesListConstants::NormalFont));
 	}
 
-	if (ColumnName == TEXT("PrimPath"))
+	else if (ColumnName == TEXT("PrimPath"))
 	{
 		SAssignNew(ColumnWidget, STextBlock)
 		.Text(FText::FromString(Reference->PrimPath.IsEmpty() ? TEXT("(default prim)") : Reference->PrimPath))
+		.Margin(FMargin(3.0f, 0.0f, 0.0f, 0.0f))
 		.Font(FAppStyle::GetFontStyle(UsdReferencesListConstants::NormalFont));
 	}
 
 	return SNew(SBox)
 		.HeightOverride(FUsdStageEditorStyle::Get()->GetFloat("UsdStageEditor.ListItemHeight"))
+		.HAlign(HAlign_Fill)
+		.VAlign(VAlign_Center)
 		[
-			SNew(SHorizontalBox)
-			+SHorizontalBox::Slot()
-			.HAlign(HAlign_Left)
-			.VAlign(VAlign_Center)
-			.Padding(UsdReferencesListConstants::RowPadding)
-			.AutoWidth()
-			[
-				ColumnWidget
-			]
+			ColumnWidget
 		];
 	// clang-format on
 }
@@ -64,6 +72,10 @@ void SUsdReferencesList::Construct(const FArguments& InArgs)
 	// clang-format off
 	SAssignNew(HeaderRowWidget, SHeaderRow)
 
+	+SHeaderRow::Column(FName(TEXT("ReferenceType")))
+	.DefaultLabel(FText::GetEmpty())
+	.FixedWidth(24.0f)
+
 	+SHeaderRow::Column(FName(TEXT("AssetPath")))
 	.DefaultLabel(LOCTEXT("ReferencedPath", "Referenced layers"))
 	.FillWidth(100.f)
@@ -80,6 +92,8 @@ void SUsdReferencesList::Construct(const FArguments& InArgs)
 		.HeaderRow(HeaderRowWidget)
 	);
 
+	OnContextMenuOpening = FOnContextMenuOpening::CreateSP(this, &SUsdReferencesList::ConstructLayerContextMenu);
+
 	SetVisibility(EVisibility::Collapsed); // Start hidden until SetPrimPath displays us
 	// clang-format on
 }
@@ -89,6 +103,114 @@ TSharedRef<ITableRow> SUsdReferencesList::OnGenerateRow(TSharedPtr<FUsdReference
 	return SNew(SUsdReferenceRow, InDisplayNode, OwnerTable);
 }
 
+TSharedPtr<SWidget> SUsdReferencesList::ConstructLayerContextMenu()
+{
+	TSharedRef<SWidget> MenuWidget = SNullWidget::NullWidget;
+
+	FMenuBuilder LayerOptions(true, nullptr);
+	LayerOptions.BeginSection(
+		"Reference",
+		TAttribute<FText>::Create(TAttribute<FText>::FGetter::CreateLambda(
+			[this]()
+			{
+				bool bIsPayload = false;
+				for (const TSharedPtr<FUsdReference>& SelectedItem : GetSelectedItems())
+				{
+					if (SelectedItem->bIsPayload)
+					{
+						bIsPayload = true;
+						break;
+					}
+				}
+
+				return bIsPayload ? LOCTEXT("Payload_Text", "Payload") : LOCTEXT("Reference_Text", "Reference");
+			}
+		))
+	);
+	{
+		LayerOptions.AddMenuEntry(
+			LOCTEXT("RemoveReference", "Remove"),
+			TAttribute<FText>::Create(TAttribute<FText>::FGetter::CreateLambda(
+				[this]()
+				{
+					const bool bCanRemove = CanRemoveReference();
+					return bCanRemove ? LOCTEXT("RemoveReference_ToolTip", "Remove this reference or payload")
+									  : LOCTEXT(
+										  "CantRemoveReference_ToolTip",
+										  "Cannot remove a reference that was introduced from across another reference"
+									  );
+				}
+			)),
+			FSlateIcon(),
+			FUIAction(
+				FExecuteAction::CreateSP(this, &SUsdReferencesList::RemoveReference),
+				FCanExecuteAction::CreateSP(this, &SUsdReferencesList::CanRemoveReference)
+			),
+			NAME_None,
+			EUserInterfaceActionType::Button
+		);
+
+		LayerOptions.AddMenuEntry(
+			LOCTEXT("ReloadReference", "Reload"),
+			LOCTEXT("ReloadReference_ToolTip", "Reloads this reference or payload"),
+			FSlateIcon(),
+			FUIAction(
+				FExecuteAction::CreateSP(this, &SUsdReferencesList::ReloadReference),
+				FCanExecuteAction::CreateSP(this, &SUsdReferencesList::CanReloadReference)
+			),
+			NAME_None,
+			EUserInterfaceActionType::Button
+		);
+	}
+	LayerOptions.EndSection();
+
+	MenuWidget = LayerOptions.MakeWidget();
+
+	return MenuWidget;
+}
+
+bool SUsdReferencesList::CanRemoveReference() const
+{
+	// From https://openusd.org/dev/usdfaq.html#list-edited-composition-arcs:
+	// "The rule, therefore, for meaningfully deleting composition arcs, is that you can only remove an arc
+	// if it was introduced in the same layerStack, as discussed with an example here. This means you cannot
+	// delete a reference that was introduced from across another reference."
+	for (const TSharedPtr<FUsdReference>& SelectedItem : GetSelectedItems())
+	{
+		if (!SelectedItem->bIntroducedInLocalLayerStack)
+		{
+			return false;
+		}
+	}
+
+	return true;
+}
+
+void SUsdReferencesList::RemoveReference()
+{
+	FScopedTransaction Transaction(LOCTEXT("RemoveReferenceTransaction", "Remove reference"));
+
+	for (const TSharedPtr<FUsdReference>& SelectedItem : GetSelectedItems())
+	{
+		ViewModel.RemoveReference(SelectedItem);
+	}
+}
+
+bool SUsdReferencesList::CanReloadReference() const
+{
+	return true;
+}
+
+void SUsdReferencesList::ReloadReference()
+{
+	FScopedTransaction Transaction(LOCTEXT("ReloadReferenceTransaction", "Reload reference"));
+
+	for (const TSharedPtr<FUsdReference>& SelectedItem : GetSelectedItems())
+	{
+		ViewModel.ReloadReference(SelectedItem);
+	}
+}
+
 void SUsdReferencesList::SetPrimPath(const UE::FUsdStageWeak& UsdStage, const TCHAR* PrimPath)
 {
 	ViewModel.UpdateReferences(UsdStage, PrimPath);
diff --git a/Engine/Plugins/Importers/USDImporter/Source/USDStageEditor/Private/Widgets/SUSDReferencesList.h b/Engine/Plugins/Importers/USDImporter/Source/USDStageEditor/Private/Widgets/SUSDReferencesList.h
--- a/Engine/Plugins/Importers/USDImporter/Source/USDStageEditor/Private/Widgets/SUSDReferencesList.h
+++ b/Engine/Plugins/Importers/USDImporter/Source/USDStageEditor/Private/Widgets/SUSDReferencesList.h
@@ -40,6 +40,13 @@ public:
 
 protected:
 	TSharedRef<ITableRow> OnGenerateRow(TSharedPtr<FUsdReference> InDisplayNode, const TSharedRef<STableViewBase>& OwnerTable);
+	TSharedPtr<SWidget> ConstructLayerContextMenu();
+
+	bool CanRemoveReference() const;
+	void RemoveReference();
+
+	bool CanReloadReference() const;
+	void ReloadReference();
 
 private:
 	FUsdReferencesViewModel ViewModel;
diff --git a/Engine/Plugins/Importers/USDImporter/Source/USDStageEditorViewModels/Private/USDReferencesViewModel.cpp b/Engine/Plugins/Importers/USDImporter/Source/USDStageEditorViewModels/Private/USDReferencesViewModel.cpp
--- a/Engine/Plugins/Importers/USDImporter/Source/USDStageEditorViewModels/Private/USDReferencesViewModel.cpp
+++ b/Engine/Plugins/Importers/USDImporter/Source/USDStageEditorViewModels/Private/USDReferencesViewModel.cpp
@@ -2,6 +2,7 @@
 
 #include "USDReferencesViewModel.h"
 
+#include "USDLog.h"
 #include "USDMemory.h"
 #include "USDTypesConversion.h"
 
@@ -11,31 +12,51 @@
 
 #if USE_USD_SDK
 #include "USDIncludesStart.h"
+#include "pxr/usd/pcp/layerStack.h"
+#include "pxr/usd/sdf/payload.h"
 #include "pxr/usd/sdf/reference.h"
+#include "pxr/usd/usd/payloads.h"
 #include "pxr/usd/usd/prim.h"
 #include "pxr/usd/usd/primCompositionQuery.h"
+#include "pxr/usd/usd/references.h"
 #include "USDIncludesEnd.h"
 #endif	  // #if USE_USD_SDK
 
-void FUsdReferencesViewModel::UpdateReferences(const UE::FUsdStageWeak& UsdStage, const TCHAR* PrimPath)
+void FUsdReferencesViewModel::UpdateReferences(const UE::FUsdStageWeak& InUsdStage, const TCHAR* InPrimPath)
 {
+	// We're provided with an empty prim path when we're meant to clear our references and go invisible,
+	// so always do that
 	References.Reset();
 
-	if (!UsdStage || !PrimPath || FString{PrimPath}.IsEmpty() || UE::FSdfPath{PrimPath}.IsAbsoluteRootPath())
+	if (!InUsdStage || !InPrimPath || FString{InPrimPath}.IsEmpty() || UE::FSdfPath{InPrimPath}.IsAbsoluteRootPath())
 	{
 		return;
 	}
 
+	UsdStage = InUsdStage;
+	PrimPath = UE::FSdfPath{InPrimPath};
+
 #if USE_USD_SDK
 	FScopedUsdAllocs UsdAllocs;
 
-	if (pxr::UsdPrim Prim{UsdStage.GetPrimAtPath(UE::FSdfPath(PrimPath))})
+	if (pxr::UsdPrim Prim{UsdStage.GetPrimAtPath(PrimPath)})
 	{
-		pxr::UsdPrimCompositionQuery PrimCompositionQuery = pxr::UsdPrimCompositionQuery::GetDirectReferences(Prim);
+		pxr::UsdPrimCompositionQuery::Filter Filter;
+		Filter.dependencyTypeFilter = pxr::UsdPrimCompositionQuery::DependencyTypeFilter::Direct;
+
+		pxr::UsdPrimCompositionQuery PrimCompositionQuery = pxr::UsdPrimCompositionQuery(Prim, Filter);
+
+		pxr::PcpLayerStackRefPtr RootLayerStack;
 
 		for (const pxr::UsdPrimCompositionQueryArc& CompositionArc : PrimCompositionQuery.GetCompositionArcs())
 		{
-			if (CompositionArc.GetArcType() == pxr::PcpArcTypeReference)
+			pxr::PcpNodeRef IntroducingNode = CompositionArc.GetIntroducingNode();
+
+			if (CompositionArc.GetArcType() == pxr::PcpArcTypeRoot)
+			{
+				RootLayerStack = IntroducingNode.GetLayerStack();
+			}
+			else if (CompositionArc.GetArcType() == pxr::PcpArcTypeReference)
 			{
 				pxr::SdfReferenceEditorProxy ReferenceEditor;
 				pxr::SdfReference UsdReference;
@@ -45,6 +66,24 @@ void FUsdReferencesViewModel::UpdateReferences(const UE::FUsdStageWeak& UsdStage
 					FUsdReference Reference;
 					Reference.AssetPath = UsdToUnreal::ConvertString(UsdReference.GetAssetPath());
 					Reference.PrimPath = UsdToUnreal::ConvertPath(UsdReference.GetPrimPath());
+					Reference.bIntroducedInLocalLayerStack = IntroducingNode.GetLayerStack() == RootLayerStack;
+					Reference.bIsPayload = false;
+
+					References.Add(MakeSharedUnreal<FUsdReference>(MoveTemp(Reference)));
+				}
+			}
+			else if (CompositionArc.GetArcType() == pxr::PcpArcTypePayload)
+			{
+				pxr::SdfPayloadEditorProxy Editor;
+				pxr::SdfPayload UsdPayload;
+
+				if (CompositionArc.GetIntroducingListEditor(&Editor, &UsdPayload))
+				{
+					FUsdReference Reference;
+					Reference.AssetPath = UsdToUnreal::ConvertString(UsdPayload.GetAssetPath());
+					Reference.PrimPath = UsdToUnreal::ConvertPath(UsdPayload.GetPrimPath());
+					Reference.bIntroducedInLocalLayerStack = IntroducingNode.GetLayerStack() == RootLayerStack;
+					Reference.bIsPayload = true;
 
 					References.Add(MakeSharedUnreal<FUsdReference>(MoveTemp(Reference)));
 				}
@@ -53,3 +92,183 @@ void FUsdReferencesViewModel::UpdateReferences(const UE::FUsdStageWeak& UsdStage
 	}
 #endif	  // #if USE_USD_SDK
 }
+
+void FUsdReferencesViewModel::RemoveReference(const TSharedPtr<FUsdReference>& Reference)
+{
+#if USE_USD_SDK
+	FScopedUsdAllocs UsdAllocs;
+
+	if (pxr::UsdPrim Prim{UsdStage.GetPrimAtPath(UE::FSdfPath(PrimPath))})
+	{
+		pxr::UsdPrimCompositionQuery::Filter Filter;
+		Filter.dependencyTypeFilter = pxr::UsdPrimCompositionQuery::DependencyTypeFilter::Direct;
+
+		pxr::UsdPrimCompositionQuery PrimCompositionQuery = pxr::UsdPrimCompositionQuery(Prim, Filter);
+
+		// Annoyingly there is no pxr::UsdReferences::GetReferences() (even though there are add/clear/set?).
+		// We must iterate through all composition arcs again to find the exact reference.
+		//
+		// That is a good idea anyway though because since we first constructed our FUsdReferences in this way,
+		// we know that the path strings will match exactly, whether they're relative/absolute paths,
+		// upper/lower case drive letters, etc.
+		//
+		// Also, using the EditorProxy objects makes it easy to remove them, because otherwise we'd have
+		// to author opinions using pxr::UsdReferences/pxr::UsdPayloads, and have to worry about editing inside
+		// of variants and so on
+		for (const pxr::UsdPrimCompositionQueryArc& CompositionArc : PrimCompositionQuery.GetCompositionArcs())
+		{
+			if (Reference->bIsPayload && CompositionArc.GetArcType() == pxr::PcpArcTypePayload)
+			{
+				pxr::SdfPayloadEditorProxy Editor;
+				pxr::SdfPayload Payload;
+
+				if (CompositionArc.GetIntroducingListEditor(&Editor, &Payload))
+				{
+					FString AssetPath = UsdToUnreal::ConvertString(Payload.GetAssetPath());
+					if (AssetPath != Reference->AssetPath)
+					{
+						continue;
+					}
+
+					FString TargetPrimPath = UsdToUnreal::ConvertPath(Payload.GetPrimPath());
+					if (TargetPrimPath != Reference->PrimPath)
+					{
+						continue;
+					}
+
+					Editor.Remove(Payload);
+				}
+			}
+			else if (!Reference->bIsPayload && CompositionArc.GetArcType() == pxr::PcpArcTypeReference)
+			{
+				pxr::SdfReferenceEditorProxy Editor;
+				pxr::SdfReference SdfReference;
+
+				if (CompositionArc.GetIntroducingListEditor(&Editor, &SdfReference))
+				{
+					FString AssetPath = UsdToUnreal::ConvertString(SdfReference.GetAssetPath());
+					if (AssetPath != Reference->AssetPath)
+					{
+						continue;
+					}
+
+					FString TargetPrimPath = UsdToUnreal::ConvertPath(SdfReference.GetPrimPath());
+					if (TargetPrimPath != Reference->PrimPath)
+					{
+						continue;
+					}
+
+					Editor.Remove(SdfReference);
+				}
+			}
+		}
+	}
+#endif	  // #if USE_USD_SDK
+}
+
+void FUsdReferencesViewModel::ReloadReference(const TSharedPtr<FUsdReference>& Reference)
+{
+#if USE_USD_SDK
+	FScopedUsdAllocs UsdAllocs;
+
+	pxr::SdfLayerRefPtr ReloadedLayer;
+	pxr::SdfPath SdfReferencedPrimPath;
+
+	if (pxr::UsdPrim Prim{UsdStage.GetPrimAtPath(UE::FSdfPath(PrimPath))})
+	{
+		pxr::UsdPrimCompositionQuery::Filter Filter;
+		Filter.dependencyTypeFilter = pxr::UsdPrimCompositionQuery::DependencyTypeFilter::Direct;
+
+		pxr::UsdPrimCompositionQuery PrimCompositionQuery = pxr::UsdPrimCompositionQuery(Prim, Filter);
+
+		for (const pxr::UsdPrimCompositionQueryArc& CompositionArc : PrimCompositionQuery.GetCompositionArcs())
+		{
+			if (CompositionArc.GetArcType() == pxr::PcpArcTypeReference)
+			{
+				pxr::SdfReferenceEditorProxy ReferenceEditor;
+				pxr::SdfReference UsdReference;
+
+				if (CompositionArc.GetIntroducingListEditor(&ReferenceEditor, &UsdReference))
+				{
+					FString AssetPath = UsdToUnreal::ConvertString(UsdReference.GetAssetPath());
+					if (AssetPath != Reference->AssetPath)
+					{
+						continue;
+					}
+
+					FString ReferencedPrimPath = UsdToUnreal::ConvertPath(UsdReference.GetPrimPath());
+					if (ReferencedPrimPath != Reference->PrimPath)
+					{
+						continue;
+					}
+
+					pxr::PcpNodeRef TargetNode = CompositionArc.GetTargetNode();
+					pxr::SdfLayerHandle Layer = CompositionArc.GetTargetLayer();
+
+					SdfReferencedPrimPath = TargetNode.GetPath();
+					ReloadedLayer = pxr::SdfLayerRefPtr{Layer};
+				}
+			}
+			else if (CompositionArc.GetArcType() == pxr::PcpArcTypePayload)
+			{
+				pxr::SdfPayloadEditorProxy Editor;
+				pxr::SdfPayload Payload;
+
+				if (CompositionArc.GetIntroducingListEditor(&Editor, &Payload))
+				{
+					FString AssetPath = UsdToUnreal::ConvertString(Payload.GetAssetPath());
+					if (AssetPath != Reference->AssetPath)
+					{
+						continue;
+					}
+
+					FString ReferencedPrimPath = UsdToUnreal::ConvertPath(Payload.GetPrimPath());
+					if (ReferencedPrimPath != Reference->PrimPath)
+					{
+						continue;
+					}
+
+					pxr::PcpNodeRef TargetNode = CompositionArc.GetTargetNode();
+					pxr::SdfLayerHandle Layer = CompositionArc.GetTargetLayer();
+
+					SdfReferencedPrimPath = TargetNode.GetPath();
+					ReloadedLayer = pxr::SdfLayerRefPtr{Layer};
+				}
+			}
+		}
+	}
+
+	if (!ReloadedLayer || SdfReferencedPrimPath.IsEmpty())
+	{
+		return;
+	}
+
+	// Retrieving the layer stack from the composition query arc above is not perfect: It will give us
+	// all the sublayers of the layer we want to reload, but not any layers referenced by prims in the
+	// subtree of our target prim. Instead, here we "just let USD do it" by instead reopening the referenced
+	// layer with a population mask for our target prim, an retrieving all UsedLayers from the generated stage.
+	// All of those layers are already opened anyway, and with the population mask this should be
+	// pretty fast still, and in turn we avoid having to recurse through the composition query arcs ourselves
+	// which could have some mistakes and miss complex edge cases
+	pxr::UsdStagePopulationMask Mask;
+	Mask.Add(SdfReferencedPrimPath);
+	const pxr::SdfLayerRefPtr SessionLayer = nullptr;
+	const pxr::UsdStage::InitialLoadSet InitialLoadSet = pxr::UsdStage::InitialLoadSet::LoadAll;
+	if (pxr::UsdStageRefPtr TempStage = pxr::UsdStage::OpenMasked(ReloadedLayer, SessionLayer, Mask, InitialLoadSet))
+	{
+		std::vector<pxr::SdfLayerHandle> UsedLayers = TempStage->GetUsedLayers();
+
+		// Reference: SdfLayer::ReloadLayers, which we don't use directly because it takes a std::set<> but just loops over it anyway
+		pxr::SdfChangeBlock Block;
+		for (const pxr::SdfLayerHandle& Layer : UsedLayers)
+		{
+			if (Layer)
+			{
+				UE_LOG(LogUsd, Log, TEXT("Reloading layer '%s'"), *UsdToUnreal::ConvertString(Layer->GetIdentifier()));
+				const bool bForce = true;
+				Layer->Reload(bForce);
+			}
+		}
+	}
+#endif	  // #if USE_USD_SDK
+}
diff --git a/Engine/Plugins/Importers/USDImporter/Source/USDStageEditorViewModels/Public/USDReferencesViewModel.h b/Engine/Plugins/Importers/USDImporter/Source/USDStageEditorViewModels/Public/USDReferencesViewModel.h
--- a/Engine/Plugins/Importers/USDImporter/Source/USDStageEditorViewModels/Public/USDReferencesViewModel.h
+++ b/Engine/Plugins/Importers/USDImporter/Source/USDStageEditorViewModels/Public/USDReferencesViewModel.h
@@ -5,19 +5,27 @@
 #include "Templates/SharedPointer.h"
 
 #include "UsdWrappers/ForwardDeclarations.h"
+#include "UsdWrappers/SdfPath.h"
+#include "UsdWrappers/UsdStage.h"
 
 class USDSTAGEEDITORVIEWMODELS_API FUsdReference : public TSharedFromThis<FUsdReference>
 {
 public:
 	FString AssetPath;
 	FString PrimPath;
+	bool bIntroducedInLocalLayerStack = true;
+	bool bIsPayload = false;
 };
 
 class USDSTAGEEDITORVIEWMODELS_API FUsdReferencesViewModel
 {
 public:
 	void UpdateReferences(const UE::FUsdStageWeak& UsdStage, const TCHAR* PrimPath);
+	void RemoveReference(const TSharedPtr<FUsdReference>& Reference);
+	void ReloadReference(const TSharedPtr<FUsdReference>& Reference);
 
 public:
+	UE::FUsdStageWeak UsdStage;
+	UE::FSdfPath PrimPath;
 	TArray<TSharedPtr<FUsdReference>> References;
 };
