USD: Tweaks to improve performance when reloading stages.

Discard undo/redo buffer when reloading if the cvar is enabled;
Reuse translated prototypes when handling all individual changes within a USD notices;
Always handle layer changed notices together with the prim changed notices;
Prune the notice spam that USD emits when we fully reload the stage;
Fix weird authored opinions when undoing a stage reload.

diff --git a/Engine/Plugins/Importers/USDImporter/Source/USDSchemas/Private/USDSchemaTranslator.cpp b/Engine/Plugins/Importers/USDImporter/Source/USDSchemas/Private/USDSchemaTranslator.cpp
--- a/Engine/Plugins/Importers/USDImporter/Source/USDSchemas/Private/USDSchemaTranslator.cpp
+++ b/Engine/Plugins/Importers/USDImporter/Source/USDSchemas/Private/USDSchemaTranslator.cpp
@@ -232,12 +232,6 @@ void FUsdSchemaTranslationContext::CompleteTasks()
 
 		bFinished = (TranslatorTasks.Num() == 0);
 	}
-
-	// These only make sense within an asset translation period
-	if (UsdInfoCache)
-	{
-		UsdInfoCache->ResetTranslatedPrototypes();
-	}
 }
 
 bool FUsdSchemaTranslator::IsCollapsed(ECollapsingType CollapsingType) const
diff --git a/Engine/Plugins/Importers/USDImporter/Source/USDStage/Private/USDStageActor.cpp b/Engine/Plugins/Importers/USDImporter/Source/USDStage/Private/USDStageActor.cpp
--- a/Engine/Plugins/Importers/USDImporter/Source/USDStage/Private/USDStageActor.cpp
+++ b/Engine/Plugins/Importers/USDImporter/Source/USDStage/Private/USDStageActor.cpp
@@ -1052,23 +1052,21 @@ AUsdStageActor::AUsdStageActor()
 		UsdListener.GetOnSdfLayersChanged().AddLambda(
 			[&, this](const UsdUtils::FLayerToSdfChangeList& LayersToChangeList)
 			{
+				TRACE_CPUPROFILER_EVENT_SCOPE(AUsdStageActor::OnSdfLayersChanged);
+
 				if (!IsListeningToUsdNotices() || LayersToChangeList.Num() == 0)
 				{
 					return;
 				}
 
-				TOptional<TGuardValue<ITransaction*>> SuppressTransaction;
-				if (this->GetOutermost()->HasAnyPackageFlags(PKG_PlayInEditor))
-				{
-					SuppressTransaction.Emplace(GUndo, nullptr);
-				}
-
 				const UE::FUsdStage& Stage = GetUsdStage();
 				if (!Stage)
 				{
 					return;
 				}
 
+				TSet<FString> ChangedLayers;
+
 				// Check to see if any of the stage's layers reloaded, or if we added/removed any layer
 				TSet<UE::FSdfLayer> UsedLayers{Stage.GetUsedLayers()};
 				for (const TPair<UE::FSdfLayerWeak, UsdUtils::FSdfChangeList>& LayerToChangeList : LayersToChangeList)
@@ -1078,18 +1076,17 @@ AUsdStageActor::AUsdStageActor()
 						continue;
 					}
 
+					const FString ChangedLayerIdentifier = LayerToChangeList.Key.GetIdentifier();
+
 					for (const TPair<UE::FSdfPath, UsdUtils::FSdfChangeListEntry>& Change : LayerToChangeList.Value)
 					{
 						if (Change.Value.Flags.bDidReloadContent)
 						{
-							UE_LOG(
-								LogUsd,
-								Verbose,
-								TEXT("Reloading animations because layer '%s' was reloaded"),
-								*LayerToChangeList.Key.GetIdentifier()
-							);
-							RegenerateLevelSequence();
-							RepopulateLevelSequence();
+							// Luckily whenever USD emits a one of these events for a layer reload or sublayer being added/removed,
+							// we also get an object changed notice right after it. This means that we really don't need to do anything
+							// here except to flag that on the next HandleAccumulatedNotices call we really should make sure our
+							// LevelSequence is reloaded, so that we generate subsections for these layers that were added/removed.
+							bLayerReloaded = true;
 							return;
 						}
 
@@ -1098,14 +1095,7 @@ AUsdStageActor::AUsdStageActor()
 							if (SubLayerChange.Value == UsdUtils::ESubLayerChangeType::SubLayerAdded
 								|| SubLayerChange.Value == UsdUtils::ESubLayerChangeType::SubLayerRemoved)
 							{
-								UE_LOG(
-									LogUsd,
-									Verbose,
-									TEXT("Reloading animations because layer '%s' was added or removed"),
-									*LayerToChangeList.Key.GetIdentifier()
-								);
-								RegenerateLevelSequence();
-								RepopulateLevelSequence();
+								bLayerReloaded = true;
 								return;
 							}
 						}
@@ -1353,7 +1343,7 @@ void AUsdStageActor::HandleAccumulatedNotices()
 
 #if USE_USD_SDK
 
-	if (AccumulatedInfoChanges.Num() == 0 && AccumulatedResyncChanges.Num() == 0)
+	if (AccumulatedInfoChanges.Num() == 0 && AccumulatedResyncChanges.Num() == 0 && !bLayerReloaded)
 	{
 		return;
 	}
@@ -1405,7 +1395,10 @@ void AUsdStageActor::HandleAccumulatedNotices()
 
 	bool bHasResync = AccumulatedResyncChanges.Num() > 0;
 
-	bool bNeedsAnimationReload = false;
+	// If any layer changed we'll need to regenerate the LevelSequences for sure (to make sure we update subsequence tracks to match the stage).
+	// We don't have to worry about actually forcing a repopulate: Adding/removing/reloading layers always emits a root resync anyway, which will
+	// already naturally repopulate the level sequences
+	bool bNeedsAnimationReload = bLayerReloaded;
 
 	// The most important thing here is to iterate in parent to child order, so build SortedPrimsChangedList
 	TMap<UE::FSdfPath, bool> SortedPrimsChangedList;
@@ -1830,6 +1823,15 @@ void AUsdStageActor::HandleAccumulatedNotices()
 		RegenerateLevelSequence();
 	}
 
+	// Reset our translated prototypes only here, and reuse them for all individual changes. This because some types of operations
+	// (e.g. reloading a reference used in an instanceable) will cause USD to emit a resync notice for every single instance
+	// of the prototype: By keeping track of which prototypes we translated across all those changes we can do the actual translation
+	//  only once
+	if (UsdInfoCache)
+	{
+		UsdInfoCache->GetInner().ResetTranslatedPrototypes();
+	}
+
 	RefreshStageTask.EnterProgressFrame();
 	FScopedSlowTask RegenerateAssetsTask(PrimsToUpdate.Num(), LOCTEXT("RegeneratingAssets", "Regenerating assets"));
 	{
@@ -1937,6 +1939,7 @@ void AUsdStageActor::HandleAccumulatedNotices()
 
 	AccumulatedInfoChanges.Reset();
 	AccumulatedResyncChanges.Reset();
+	bLayerReloaded = false;
 }
 
 USDSTAGE_API void AUsdStageActor::Reset()
@@ -2052,7 +2055,13 @@ UUsdPrimTwin* AUsdStageActor::ExpandPrim(
 
 	bool bExpandChildren = true;
 
-	if (!TranslationContext.bIsJustRepopulatingLevelSequence)
+	if (TranslationContext.bIsJustRepopulatingLevelSequence)
+	{
+		// For the repopulate, let's only visit the prim twins that already have children and so may actually have
+		// components. We don't want to create brand new components here
+		bExpandChildren = bResync && UsdPrimTwin->GetChildren().Num() > 0;
+	}
+	else if (!TranslationContext.bIsJustRepopulatingLevelSequence)
 	{
 		IUsdSchemasModule& UsdSchemasModule = FModuleManager::Get().LoadModuleChecked<IUsdSchemasModule>(TEXT("USDSchemas"));
 		if (TSharedPtr<FUsdSchemaTranslator> SchemaTranslator = UsdSchemasModule.GetTranslatorRegistry()
@@ -3303,6 +3312,11 @@ void AUsdStageActor::LoadUsdStage(bool bOpenIfNeeded)
 		PrimLinkCache = NewObject<UUsdPrimLinkCache>(Outer, NAME_None, EObjectFlags::RF_Transient | EObjectFlags::RF_Transactional);
 	}
 
+	if (UsdInfoCache)
+	{
+		UsdInfoCache->GetInner().ResetTranslatedPrototypes();
+	}
+
 	RegenerateLevelSequence();
 
 	TSharedRef<FUsdSchemaTranslationContext> TranslationContext = FUsdStageActorImpl::CreateUsdSchemaTranslationContext(this, RootTwin->PrimPath);
@@ -3609,6 +3623,8 @@ void AUsdStageActor::RegenerateLevelSequence()
 
 void AUsdStageActor::RepopulateLevelSequence()
 {
+	TRACE_CPUPROFILER_EVENT_SCOPE(AUsdStageActor::RepopulateLevelSequence);
+
 	TSharedRef<FUsdSchemaTranslationContext> TranslationContext = FUsdStageActorImpl::CreateUsdSchemaTranslationContext(
 		this,
 		UE::FSdfPath::AbsoluteRootPath().GetString()
@@ -3824,9 +3840,10 @@ void AUsdStageActor::HandleTransactionStateChanged(
 		}
 		else if (InTransactionState == ETransactionStateEventType::TransactionFinalized)
 		{
-			ensureAlways(AccumulatedInfoChanges.Num() == 0 && AccumulatedResyncChanges.Num() == 0);
+			ensureAlways(AccumulatedInfoChanges.Num() == 0 && AccumulatedResyncChanges.Num() == 0 && !bLayerReloaded);
 			AccumulatedInfoChanges.Reset();
 			AccumulatedResyncChanges.Reset();
+			bLayerReloaded = false;
 		}
 	}
 
@@ -4841,6 +4858,11 @@ void AUsdStageActor::OnSkelAnimationBaked(const FString& SkeletonPrimPath)
 	// The only way we could have baked a skel animation is via the sequencer, so we know its playing
 	TranslationContext->bSequencerIsAnimating = true;
 
+	if (UsdInfoCache)
+	{
+		UsdInfoCache->GetInner().ResetTranslatedPrototypes();
+	}
+
 	IUsdSchemasModule& UsdSchemasModule = FModuleManager::Get().LoadModuleChecked<IUsdSchemasModule>(TEXT("USDSchemas"));
 	if (TSharedPtr<FUsdSchemaTranslator> SchemaTranslator = UsdSchemasModule.GetTranslatorRegistry()
 																.CreateTranslatorForSchema(TranslationContext, UE::FUsdTyped(SkeletonPrim)))
@@ -5052,6 +5074,11 @@ void AUsdStageActor::AnimatePrims()
 		return;
 	}
 
+	if (UsdInfoCache)
+	{
+		UsdInfoCache->GetInner().ResetTranslatedPrototypes();
+	}
+
 	TSharedRef<FUsdSchemaTranslationContext> TranslationContext = FUsdStageActorImpl::CreateUsdSchemaTranslationContext(
 		this,
 		GetRootPrimTwin()->PrimPath
diff --git a/Engine/Plugins/Importers/USDImporter/Source/USDStage/Public/USDStageActor.h b/Engine/Plugins/Importers/USDImporter/Source/USDStage/Public/USDStageActor.h
--- a/Engine/Plugins/Importers/USDImporter/Source/USDStage/Public/USDStageActor.h
+++ b/Engine/Plugins/Importers/USDImporter/Source/USDStage/Public/USDStageActor.h
@@ -419,7 +419,7 @@ protected:
 
 	void UpdateSpawnedObjectsTransientFlag(bool bTransient);
 
-	void RequestDelayedTransactorReset();
+	USDSTAGE_API void RequestDelayedTransactorReset();
 
 #if WITH_EDITOR
 	void OnBeginPIE(bool bIsSimulating);
@@ -439,6 +439,7 @@ protected:
 	void OnSkelAnimationBaked(const FString& SkeletonPrimPath);
 
 protected:
+	friend class SUsdStage;	   // So that it can call RequestDelayedTransactorReset
 	friend struct FUsdStageActorImpl;
 	friend class FUsdLevelSequenceHelperImpl;
 	friend class UsdUtils::FUsdTransactorImpl;
@@ -477,6 +478,7 @@ protected:
 	 */
 	UsdUtils::FObjectChangesByPath AccumulatedInfoChanges;
 	UsdUtils::FObjectChangesByPath AccumulatedResyncChanges;
+	bool bLayerReloaded = false;
 	TArray<UE::FSdfPath> ResyncedPrimsForThisTransaction;
 
 	PRAGMA_DISABLE_DEPRECATION_WARNINGS
diff --git a/Engine/Plugins/Importers/USDImporter/Source/USDStageEditor/Private/SUSDStage.cpp b/Engine/Plugins/Importers/USDImporter/Source/USDStageEditor/Private/SUSDStage.cpp
--- a/Engine/Plugins/Importers/USDImporter/Source/USDStageEditor/Private/SUSDStage.cpp
+++ b/Engine/Plugins/Importers/USDImporter/Source/USDStageEditor/Private/SUSDStage.cpp
@@ -2652,13 +2652,27 @@ void SUsdStage::FileExportFlattenedLayerStack(const FString& OutputLayer)
 
 void SUsdStage::FileReload()
 {
+	AUsdStageActor* StageActor = ViewModel.UsdStageActor.Get();
+	if (!StageActor)
+	{
+		return;
+	}
+
+	static IConsoleVariable* Cvar = IConsoleManager::Get().FindConsoleVariable(TEXT("USD.DiscardUndoBufferOnStageOpenClose"));
+	const bool bDiscardUndoBufferOnStageOpenClose = Cvar && Cvar->GetBool();
+
 	FScopedTransaction Transaction(LOCTEXT("ReloadTransaction", "Reload USD stage"));
 
-	ViewModel.ReloadStage();
+	TOptional<TGuardValue<ITransaction*>> SuppressTransaction;
+	if (bDiscardUndoBufferOnStageOpenClose)
+	{
+		SuppressTransaction.Emplace(GUndo, nullptr);
+		StageActor->RequestDelayedTransactorReset();
+	}
 
-	const AUsdStageActor* StageActor = ViewModel.UsdStageActor.Get();
+	ViewModel.ReloadStage();
 
-	if (UsdLayersTreeView && StageActor)
+	if (UsdLayersTreeView)
 	{
 		const bool bResync = true;
 		UsdLayersTreeView->Refresh(StageActor->GetBaseUsdStage(), StageActor->GetIsolatedUsdStage(), bResync);
diff --git a/Engine/Plugins/Runtime/USDCore/Source/UnrealUSDWrapper/Private/USDListener.cpp b/Engine/Plugins/Runtime/USDCore/Source/UnrealUSDWrapper/Private/USDListener.cpp
--- a/Engine/Plugins/Runtime/USDCore/Source/UnrealUSDWrapper/Private/USDListener.cpp
+++ b/Engine/Plugins/Runtime/USDCore/Source/UnrealUSDWrapper/Private/USDListener.cpp
@@ -178,8 +178,31 @@ namespace UsdToUnreal
 	)
 	{
 		ConvertPathRange(InNotice.GetChangedInfoOnlyPaths(), OutInfoChanges);
+
+		// If we have a root path reload, just stop right here: We will have to reload everything anyway.
+		// This is handy because otherwise on full reloads USD will emit an info change with bDidReloadContent=true
+		// for most prims on the stage (this could be e.g. tens of thousands of entries, which our downstream code
+		// would uselessly process, sort, serialize, etc.).
+		const static FString RootPath = UE::FSdfPath::AbsoluteRootPath().GetString();
+		if (TArray<UsdUtils::FSdfChangeListEntry>* RootInfoChange = OutInfoChanges.Find(RootPath))
+		{
+			for (const UsdUtils::FSdfChangeListEntry& Entry : *RootInfoChange)
+			{
+				if (Entry.Flags.bDidReloadContent)
+				{
+					OutInfoChanges = {};
+					OutResyncChanges = {};
+					OutResyncChanges.Add(RootPath, *RootInfoChange);
+					return true;
+				}
+			}
+		}
+
 		ConvertPathRange(InNotice.GetResyncedPaths(), OutResyncChanges);
 
+		// Upgrade targetted reload notices into resyncs (this should now only happen when reloading a reference/payload, as reloading
+		// any layer on the local layer stack emits a change for the root path)
+		bool bReloadedContent = false;
 		for (TPair<FString, TArray<UsdUtils::FSdfChangeListEntry>>& InfoPair : OutInfoChanges)
 		{
 			const FString& PrimPath = InfoPair.Key;
@@ -193,6 +216,7 @@ namespace UsdToUnreal
 				if (ChangeIt->Flags.bDidReloadContent)
 				{
 					bUpgrade = true;
+					bReloadedContent = true;
 				}
 
 				// Upgrade visibility changes to resyncs because in case of mesh collapsing having one of the collapsed meshes go visible/invisible
@@ -223,6 +247,27 @@ namespace UsdToUnreal
 			}
 		}
 
+		// For now, dump info changes when handling the notices about reloading layers.
+		//
+		// When we reload any layer, USD will emit some notices about attributes/prims that changed, and also generic notices about those prims
+		// and their ancestors having reloaded their content (via the bDidReloadContent flag).
+		// We'll upgrade those latter notices to resyncs (above), so that we regenerate the assets/components for those prims, but we cannot use the
+		// former info changes about attributes/prims that changed (and how they changed) at all just yet, because that doesn't carry with them the
+		// respective edit target information.
+		//
+		// For example imagine that a prim is only authored in a sublayer/referenced layer, is modified on disk, and the stage is reloaded.
+		// We'll get the notice about what modification took place, but the stage's edit target will be the root layer. If we were to track those
+		// changes, whenever we undid them we'd author the old values of the prim directly on the stage's root layer, which is definitely not what
+		// we want.
+		//
+		// TODO: Remove this and use this information in order to have a "selective resync", updating only the changed prims/attributes's
+		// assets/components. It *could* be possible, but it will be complex as it will involve deducing the right edit target via composition arcs,
+		// and whether new attribute/prim specs where created within this transaction, etc.
+		if (bReloadedContent)
+		{
+			OutInfoChanges = {};
+		}
+
 		return true;
 	}
 #endif	  // USE_USD_SDK
