USD: Fixed issue with dummy USDTransactor.

diff --git a/Engine/Plugins/Importers/USDImporter/Source/USDStage/Private/USDTransactor.cpp b/Engine/Plugins/Importers/USDImporter/Source/USDStage/Private/USDTransactor.cpp
--- a/Engine/Plugins/Importers/USDImporter/Source/USDStage/Private/USDTransactor.cpp
+++ b/Engine/Plugins/Importers/USDImporter/Source/USDStage/Private/USDTransactor.cpp
@@ -1160,19 +1160,17 @@ UUsdTransactor::UUsdTransactor(FVTableHelper& Helper)
 	: Super(Helper)
 {
 }
-UUsdTransactor::UUsdTransactor() = default;
+UUsdTransactor::UUsdTransactor()
+{
+#if USE_USD_SDK
+	Impl = MakeUnique<UsdUtils::FUsdTransactorImpl>();
+#endif
+}
 UUsdTransactor::~UUsdTransactor() = default;
 
 void UUsdTransactor::Initialize(AUsdStageActor* InStageActor)
 {
 	StageActor = InStageActor;
-
-#if USE_USD_SDK
-	if (!IsTemplate())
-	{
-		Impl = MakeUnique<UsdUtils::FUsdTransactorImpl>();
-	}
-#endif	  // USE_USD_SDK
 }
 
 void UUsdTransactor::Update(const UsdUtils::FObjectChangesByPath& NewInfoChanges, const UsdUtils::FObjectChangesByPath& NewResyncChanges)
@@ -1221,12 +1219,6 @@ void UUsdTransactor::Serialize(FArchive& Ar)
 	{
 		Impl->Serialize(Ar);
 	}
-	else
-	{
-		// In case we somehow serialize before we receive a valid Impl, and then later do receive one
-		UsdUtils::FUsdTransactorImpl Dummy;
-		Dummy.Serialize(Ar);
-	}
 }
 
 #if WITH_EDITOR
