USD: Allow turning of the usage of kinds for collapsing

diff --git a/Engine/Plugins/Importers/USDImporter/Source/USDSchemas/Private/USDGeomXformableTranslator.cpp b/Engine/Plugins/Importers/USDImporter/Source/USDSchemas/Private/USDGeomXformableTranslator.cpp
--- a/Engine/Plugins/Importers/USDImporter/Source/USDSchemas/Private/USDGeomXformableTranslator.cpp
+++ b/Engine/Plugins/Importers/USDImporter/Source/USDSchemas/Private/USDGeomXformableTranslator.cpp
@@ -57,7 +57,9 @@ static bool GCollapsePrimsWithoutKind = true;
 static FAutoConsoleVariableRef CVarCollapsePrimsWithoutKind(
 	TEXT("USD.CollapsePrimsWithoutKind"),
 	GCollapsePrimsWithoutKind,
-	TEXT("Allow collapsing prims that have no authored 'Kind' value")
+	TEXT(
+		"If we're collapsing according to kind (i.e. when UsePrimKindsForCollapsing is enabled), this controls how prims should behave if they don't have any authored 'Kind' value. Set this to false, and a prim without kind can't be collapsed by its parent, even if the parent has a collapsible kind."
+	)
 );
 
 static bool GEnableCollision = true;
@@ -1046,7 +1048,7 @@ void FUsdGeomXformableTranslator::CreateAlternativeDrawModeAssets(EUsdDrawMode D
 
 namespace UE::UsdXformableTranslatorImpl::Private
 {
-	bool PrimCollapses(const pxr::UsdPrim& Prim, EUsdDefaultKind KindsToCollapse)
+	bool PrimCollapses(const pxr::UsdPrim& Prim, EUsdDefaultKind KindsToCollapse, bool bUsePrimKindsForCollapsing, bool bCanCollapseByDefault)
 	{
 		UsdUtils::ECollapsingPreference Preference = UsdUtils::GetCollapsingPreference(Prim);
 		switch (Preference)
@@ -1056,21 +1058,24 @@ namespace UE::UsdXformableTranslatorImpl::Private
 				return true;
 				break;
 			}
-			case UsdUtils::ECollapsingPreference::ByKind:
+			case UsdUtils::ECollapsingPreference::Default:
 			{
-				EUsdDefaultKind PrimKind = UsdUtils::GetDefaultKind(Prim);
+				bool bCanBeCollapsed = bCanCollapseByDefault;
+				if (bUsePrimKindsForCollapsing)
+				{
+					EUsdDefaultKind PrimKind = UsdUtils::GetDefaultKind(Prim);
 
-				// Note that this is false if PrimKind is None
-				const bool bPrimKindCollapses = EnumHasAnyFlags(KindsToCollapse, PrimKind);
-				bool bCanBeCollapsed = (KindsToCollapse != EUsdDefaultKind::None)
-									   && (bPrimKindCollapses || (PrimKind == EUsdDefaultKind::None && GCollapsePrimsWithoutKind));
+					// Note that this is false if PrimKind is None
+					const bool bPrimKindCollapses = EnumHasAnyFlags(KindsToCollapse, PrimKind);
+					bCanBeCollapsed = (KindsToCollapse != EUsdDefaultKind::None)
+									  && (bPrimKindCollapses || (PrimKind == EUsdDefaultKind::None && GCollapsePrimsWithoutKind));
 
-				if (!bCanBeCollapsed)
-				{
-					pxr::UsdModelAPI Model{Prim};
-					bCanBeCollapsed = bCanBeCollapsed || Model.IsKind(pxr::TfToken("prop"), pxr::UsdModelAPI::KindValidationNone);
+					if (!bCanBeCollapsed)
+					{
+						pxr::UsdModelAPI Model{Prim};
+						bCanBeCollapsed = bCanBeCollapsed || Model.IsKind(pxr::TfToken("prop"), pxr::UsdModelAPI::KindValidationNone);
+					}
 				}
-
 				return bCanBeCollapsed;
 			}
 			case UsdUtils::ECollapsingPreference::Never:	// Fallthrough
@@ -1117,7 +1122,15 @@ bool FUsdGeomXformableTranslator::CollapsesChildren(ECollapsingType CollapsingTy
 
 	if (Model)
 	{
-		bCollapsesChildren = UE::UsdXformableTranslatorImpl::Private::PrimCollapses(Prim, Context->KindsToCollapse);
+		// If we're not using prim kinds to collapse, we don't want anything to *try collapsing its children* by default
+		const bool bCanCollapseByDefault = false;
+
+		bCollapsesChildren = UE::UsdXformableTranslatorImpl::Private::PrimCollapses(
+			Prim,
+			Context->KindsToCollapse,
+			Context->bUsePrimKindsForCollapsing,
+			bCanCollapseByDefault
+		);
 
 		if (bCollapsesChildren)
 		{
@@ -1149,7 +1162,16 @@ bool FUsdGeomXformableTranslator::CanBeCollapsed(ECollapsingType CollapsingType)
 		return false;
 	}
 
-	return UE::UsdXformableTranslatorImpl::Private::PrimCollapses(UsdPrim, Context->KindsToCollapse);
+	// If we're not using prim kinds to collapse, we still want things to *be collapsible* by default (so that they can be controlled with the
+	// collapsing attributes)
+	const bool bCanCollapseByDefault = true;
+
+	return UE::UsdXformableTranslatorImpl::Private::PrimCollapses(
+		UsdPrim,
+		Context->KindsToCollapse,
+		Context->bUsePrimKindsForCollapsing,
+		bCanCollapseByDefault
+	);
 }
 
 TSet<UE::FSdfPath> FUsdGeomXformableTranslator::CollectAuxiliaryPrims() const
diff --git a/Engine/Plugins/Importers/USDImporter/Source/USDSchemas/Public/USDSchemaTranslator.h b/Engine/Plugins/Importers/USDImporter/Source/USDSchemas/Public/USDSchemaTranslator.h
--- a/Engine/Plugins/Importers/USDImporter/Source/USDSchemas/Public/USDSchemaTranslator.h
+++ b/Engine/Plugins/Importers/USDImporter/Source/USDSchemas/Public/USDSchemaTranslator.h
@@ -282,6 +282,12 @@ struct USDSCHEMAS_API FUsdSchemaTranslationContext : public TSharedFromThis<FUsd
 	 */
 	EUsdDefaultKind KindsToCollapse = EUsdDefaultKind::Component | EUsdDefaultKind::Subcomponent;
 
+	/**
+	 * Use KindsToCollapse to determine when to collapse prim subtrees or not (defaults to enabled).
+	 * Disable this if you want to prevent collapsing, or to control it manually by right-clicking on individual prims.
+	 */
+	bool bUsePrimKindsForCollapsing = true;
+
 	/**
 	 * If enabled, when multiple mesh prims are collapsed into a single static mesh, identical material slots are merged into one slot.
 	 * Otherwise, material slots are simply appended to the list.
diff --git a/Engine/Plugins/Importers/USDImporter/Source/USDStage/Private/USDStageActor.cpp b/Engine/Plugins/Importers/USDImporter/Source/USDStage/Private/USDStageActor.cpp
--- a/Engine/Plugins/Importers/USDImporter/Source/USDStage/Private/USDStageActor.cpp
+++ b/Engine/Plugins/Importers/USDImporter/Source/USDStage/Private/USDStageActor.cpp
@@ -179,6 +179,7 @@ struct FUsdStageActorImpl
 		TranslationContext->bAllowParsingSkeletalAnimations = true;
 
 		TranslationContext->KindsToCollapse = (EUsdDefaultKind)StageActor->KindsToCollapse;
+		TranslationContext->bUsePrimKindsForCollapsing = StageActor->bUsePrimKindsForCollapsing;
 		TranslationContext->bMergeIdenticalMaterialSlots = StageActor->bMergeIdenticalMaterialSlots;
 		TranslationContext->bShareAssetsForIdenticalPrims = StageActor->bShareAssetsForIdenticalPrims;
 
@@ -577,6 +578,7 @@ struct FUsdStageActorImpl
 			EventAttributes.Emplace(TEXT("InitialLoadSet"), LexToString((uint8)StageActor->InitialLoadSet));
 			EventAttributes.Emplace(TEXT("InterpolationType"), LexToString((uint8)StageActor->InterpolationType));
 			EventAttributes.Emplace(TEXT("KindsToCollapse"), LexToString(StageActor->KindsToCollapse));
+			EventAttributes.Emplace(TEXT("bUsePrimKindsForCollapsing"), StageActor->bUsePrimKindsForCollapsing);
 			EventAttributes.Emplace(TEXT("MergeIdenticalMaterialSlots"), LexToString(StageActor->bMergeIdenticalMaterialSlots));
 			EventAttributes.Emplace(TEXT("bShareAssetsForIdenticalPrims"), StageActor->bShareAssetsForIdenticalPrims);
 			EventAttributes.Emplace(TEXT("PurposesToLoad"), LexToString(StageActor->PurposesToLoad));
@@ -916,6 +918,7 @@ AUsdStageActor::AUsdStageActor()
 	, InitialLoadSet(EUsdInitialLoadSet::LoadAll)
 	, InterpolationType(EUsdInterpolationType::Linear)
 	, GeometryCacheImport(EGeometryCacheImport::Never)
+	, bUsePrimKindsForCollapsing(true)
 	, KindsToCollapse((int32)(EUsdDefaultKind::Component | EUsdDefaultKind::Subcomponent))
 	, bMergeIdenticalMaterialSlots(true)
 	, bShareAssetsForIdenticalPrims(true)
@@ -2628,6 +2631,20 @@ void AUsdStageActor::SetGeometryCacheImport(EGeometryCacheImport ImportOption)
 	LoadUsdStage();
 }
 
+void AUsdStageActor::SetUsePrimKindsForCollapsing(bool bUse)
+{
+	if (bUse == bUsePrimKindsForCollapsing)
+	{
+		return;
+	}
+
+	const bool bMarkDirty = false;
+	Modify(bMarkDirty);
+
+	bUsePrimKindsForCollapsing = bUse;
+	LoadUsdStage();
+}
+
 void AUsdStageActor::SetKindsToCollapse(int32 NewKindsToCollapse)
 {
 	const EUsdDefaultKind NewEnum = (EUsdDefaultKind)NewKindsToCollapse;
@@ -4705,6 +4722,12 @@ void AUsdStageActor::HandlePropertyChangedEvent(FPropertyChangedEvent& PropertyC
 		GeometryCacheImport = (EGeometryCacheImport) !((uint8)GeometryCacheImport);
 		SetGeometryCacheImport(CorrectImportOption);
 	}
+	else if (PropertyName == GET_MEMBER_NAME_CHECKED(AUsdStageActor, bUsePrimKindsForCollapsing))
+	{
+		const bool bCorrectUsePrimKindsForCollapsing = bUsePrimKindsForCollapsing;
+		bUsePrimKindsForCollapsing = !bUsePrimKindsForCollapsing;
+		SetUsePrimKindsForCollapsing(bCorrectUsePrimKindsForCollapsing);
+	}
 	else if (PropertyName == GET_MEMBER_NAME_CHECKED(AUsdStageActor, KindsToCollapse))
 	{
 		const int32 CorrectKindsToCollapse = KindsToCollapse;
diff --git a/Engine/Plugins/Importers/USDImporter/Source/USDStage/Public/USDStageActor.h b/Engine/Plugins/Importers/USDImporter/Source/USDStage/Public/USDStageActor.h
--- a/Engine/Plugins/Importers/USDImporter/Source/USDStage/Public/USDStageActor.h
+++ b/Engine/Plugins/Importers/USDImporter/Source/USDStage/Public/USDStageActor.h
@@ -84,6 +84,13 @@ public:
 	UPROPERTY(EditAnywhere, BlueprintReadOnly, Category = "USD", config)
 	EGeometryCacheImport GeometryCacheImport;
 
+	/**
+	 * Use KindsToCollapse to determine when to collapse prim subtrees or not (defaults to enabled).
+	 * Disable this if you want to prevent collapsing, or to control it manually by right-clicking on individual prims.
+	 */
+	UPROPERTY(EditAnywhere, BlueprintReadOnly, Category = "USD", config)
+	bool bUsePrimKindsForCollapsing;
+
 	/**
 	 * Whether to try to combine individual assets and components of the same type on a kind-per-kind basis,
 	 * like multiple Mesh prims into a single Static Mesh
@@ -193,6 +200,9 @@ public:
 	UFUNCTION(BlueprintCallable, Category = "USD", meta = (CallInEditor = "true"))
 	USDSTAGE_API void SetGeometryCacheImport(EGeometryCacheImport ImportOption);
 
+	UFUNCTION(BlueprintCallable, Category = "USD", meta = (CallInEditor = "true"))
+	USDSTAGE_API void SetUsePrimKindsForCollapsing(bool bUse);
+
 	UFUNCTION(BlueprintCallable, Category = "USD", meta = (CallInEditor = "true"))
 	USDSTAGE_API void SetKindsToCollapse(int32 NewKindsToCollapse);
 
diff --git a/Engine/Plugins/Importers/USDImporter/Source/USDStageEditor/Private/SUSDStage.cpp b/Engine/Plugins/Importers/USDImporter/Source/USDStageEditor/Private/SUSDStage.cpp
--- a/Engine/Plugins/Importers/USDImporter/Source/USDStageEditor/Private/SUSDStage.cpp
+++ b/Engine/Plugins/Importers/USDImporter/Source/USDStageEditor/Private/SUSDStage.cpp
@@ -1849,6 +1849,51 @@ void SUsdStage::FillCollapsingSubMenu(FMenuBuilder& MenuBuilder)
 		EUserInterfaceActionType::ToggleButton
 	);
 
+	MenuBuilder.AddMenuEntry(
+		LOCTEXT("UsePrimKindsForCollapsing", "Use prim kinds for collapsing"),
+		LOCTEXT(
+			"UsePrimKindsForCollapsing_ToolTip",
+			"Use KindsToCollapse to determine when to collapse prim subtrees or not (defaults to enabled).\nDisable this if you want to prevent collapsing, or to control it manually by right-clicking on individual prims."
+		),
+		FSlateIcon(),
+		FUIAction(
+			FExecuteAction::CreateLambda(
+				[this]()
+				{
+					if (AUsdStageActor* StageActor = GetStageActorOrCDO())
+					{
+						FScopedTransaction Transaction(FText::Format(
+							LOCTEXT("UsePrimKindsForCollapsingTransaction", "Toggle bUsePrimKindsForCollapsing on USD stage actor '{0}'"),
+							FText::FromString(StageActor->GetActorLabel())
+						));
+
+						// c.f. comment within AddKindToCollapseEntry just below
+						TGuardValue<bool> MaintainSelectionGuard(bUpdatingViewportSelection, true);
+
+						StageActor->SetUsePrimKindsForCollapsing(!StageActor->bUsePrimKindsForCollapsing);
+						if (StageActor->IsTemplate())
+						{
+							StageActor->SaveConfig();
+						}
+					}
+				}
+			),
+			FCanExecuteAction{},
+			FIsActionChecked::CreateLambda(
+				[this]()
+				{
+					if (AUsdStageActor* StageActor = GetStageActorOrCDO())
+					{
+						return StageActor->bUsePrimKindsForCollapsing;
+					}
+					return false;
+				}
+			)
+		),
+		NAME_None,
+		EUserInterfaceActionType::ToggleButton
+	);
+
 	auto AddKindToCollapseEntry = [&](const EUsdDefaultKind Kind, const FText& Text, FCanExecuteAction CanExecuteAction)
 	{
 		MenuBuilder.AddMenuEntry(
@@ -2158,10 +2203,7 @@ void SUsdStage::FillGeometryCacheImportSubMenu(FMenuBuilder& MenuBuilder)
 
 	MenuBuilder.AddMenuEntry(
 		LOCTEXT("ImportOnLoad", "On Load"),
-		LOCTEXT(
-			"ImportOnLoad_ToolTip",
-			"Geometry Caches are imported as persistents assets on stage load and played back from them."
-		),
+		LOCTEXT("ImportOnLoad_ToolTip", "Geometry Caches are imported as persistents assets on stage load and played back from them."),
 		FSlateIcon(),
 		FUIAction(
 			FExecuteAction::CreateLambda(
diff --git a/Engine/Plugins/Importers/USDImporter/Source/USDStageEditor/Private/SUSDStageTreeView.cpp b/Engine/Plugins/Importers/USDImporter/Source/USDStageEditor/Private/SUSDStageTreeView.cpp
--- a/Engine/Plugins/Importers/USDImporter/Source/USDStageEditor/Private/SUSDStageTreeView.cpp
+++ b/Engine/Plugins/Importers/USDImporter/Source/USDStageEditor/Private/SUSDStageTreeView.cpp
@@ -1705,13 +1705,13 @@ void SUsdStageTreeView::FillCollapsingSubmenu(FMenuBuilder& MenuBuilder)
 	);
 
 	MenuBuilder.AddMenuEntry(
-		LOCTEXT("CollapseOnKind", "Collapse according to kind"),
-		LOCTEXT("CollapseOnKind_ToolTip", "Collapse this prim only if its kind is allowed to collapse (default)"),
+		LOCTEXT("CollapseOnKind", "Default"),
+		LOCTEXT("CollapseOnKind_ToolTip", "When 'Use prim kinds for collapsing' is enabled, prims are collapsed according to their kind (default). When disabled, prims won't be collapsed by default."),
 		FSlateIcon(),
 		FUIAction(
-			FExecuteAction::CreateSP(this, &SUsdStageTreeView::OnSetCollapsingPreference, UsdUtils::ECollapsingPreference::ByKind),
+			FExecuteAction::CreateSP(this, &SUsdStageTreeView::OnSetCollapsingPreference, UsdUtils::ECollapsingPreference::Default),
 			FCanExecuteAction::CreateSP(this, &SUsdStageTreeView::DoesPrimExistOnStage),
-			FIsActionChecked::CreateSP(this, &SUsdStageTreeView::DoSelectedPrimsHaveCollapsingPreference, UsdUtils::ECollapsingPreference::ByKind)
+			FIsActionChecked::CreateSP(this, &SUsdStageTreeView::DoSelectedPrimsHaveCollapsingPreference, UsdUtils::ECollapsingPreference::Default)
 		),
 		NAME_None,
 		EUserInterfaceActionType::RadioButton
diff --git a/Engine/Plugins/Importers/USDImporter/Source/USDStageEditorViewModels/Private/USDStageViewModel.cpp b/Engine/Plugins/Importers/USDImporter/Source/USDStageEditorViewModels/Private/USDStageViewModel.cpp
--- a/Engine/Plugins/Importers/USDImporter/Source/USDStageEditorViewModels/Private/USDStageViewModel.cpp
+++ b/Engine/Plugins/Importers/USDImporter/Source/USDStageEditorViewModels/Private/USDStageViewModel.cpp
@@ -328,6 +328,7 @@ void FUsdStageViewModel::ImportStage(const TCHAR* TargetContentFolder, UUsdStage
 		ImportContext.ImportOptions->SubdivisionLevel = StageActor->SubdivisionLevel;
 		ImportContext.ImportOptions->MetadataOptions = StageActor->MetadataOptions;
 		ImportContext.ImportOptions->KindsToCollapse = StageActor->KindsToCollapse;
+		ImportContext.ImportOptions->bUsePrimKindsForCollapsing = StageActor->bUsePrimKindsForCollapsing;
 		ImportContext.ImportOptions->bMergeIdenticalMaterialSlots = StageActor->bMergeIdenticalMaterialSlots;
 		ImportContext.ImportOptions->bShareAssetsForIdenticalPrims = StageActor->bShareAssetsForIdenticalPrims;
 
diff --git a/Engine/Plugins/Importers/USDImporter/Source/USDStageImporter/Private/USDStageImportOptions.cpp b/Engine/Plugins/Importers/USDImporter/Source/USDStageImporter/Private/USDStageImportOptions.cpp
--- a/Engine/Plugins/Importers/USDImporter/Source/USDStageImporter/Private/USDStageImportOptions.cpp
+++ b/Engine/Plugins/Importers/USDImporter/Source/USDStageImporter/Private/USDStageImportOptions.cpp
@@ -49,6 +49,7 @@ UUsdStageImportOptions::UUsdStageImportOptions(const FObjectInitializer& ObjectI
 
 	bPrimPathFolderStructure = false;
 	KindsToCollapse = (int32)(EUsdDefaultKind::Component | EUsdDefaultKind::Subcomponent);
+	bUsePrimKindsForCollapsing = true;
 	bMergeIdenticalMaterialSlots = true;
 	bInterpretLODs = true;
 }
@@ -131,6 +132,7 @@ void UsdUtils::AddAnalyticsAttributes(const UUsdStageImportOptions& Options, TAr
 	InOutAttributes.Emplace(TEXT("ShareAssetsForIdenticalPrims"), Options.bShareAssetsForIdenticalPrims);
 	InOutAttributes.Emplace(TEXT("PrimPathFolderStructure"), LexToString(Options.bPrimPathFolderStructure));
 	InOutAttributes.Emplace(TEXT("KindsToCollapse"), LexToString(Options.KindsToCollapse));
+	InOutAttributes.Emplace(TEXT("bUsePrimKindsForCollapsing"), Options.bUsePrimKindsForCollapsing);
 	InOutAttributes.Emplace(TEXT("MergeIdenticalMaterialSlots"), LexToString(Options.bMergeIdenticalMaterialSlots));
 	InOutAttributes.Emplace(TEXT("InterpretLODs"), LexToString(Options.bInterpretLODs));
 }
diff --git a/Engine/Plugins/Importers/USDImporter/Source/USDStageImporter/Private/USDStageImporter.cpp b/Engine/Plugins/Importers/USDImporter/Source/USDStageImporter/Private/USDStageImporter.cpp
--- a/Engine/Plugins/Importers/USDImporter/Source/USDStageImporter/Private/USDStageImporter.cpp
+++ b/Engine/Plugins/Importers/USDImporter/Source/USDStageImporter/Private/USDStageImporter.cpp
@@ -1941,6 +1941,7 @@ void UUsdStageImporter::ImportFromFile(FUsdStageImportContext& ImportContext)
 	TranslationContext->MetadataOptions = ImportContext.ImportOptions->MetadataOptions;
 	TranslationContext->ParentComponent = ImportContext.SceneActor ? ImportContext.SceneActor->GetRootComponent() : nullptr;
 	TranslationContext->KindsToCollapse = (EUsdDefaultKind)ImportContext.ImportOptions->KindsToCollapse;
+	TranslationContext->bUsePrimKindsForCollapsing = ImportContext.ImportOptions->bUsePrimKindsForCollapsing;
 	TranslationContext->bMergeIdenticalMaterialSlots = ImportContext.ImportOptions->bMergeIdenticalMaterialSlots;
 	TranslationContext->bShareAssetsForIdenticalPrims = ImportContext.ImportOptions->bShareAssetsForIdenticalPrims;
 	TranslationContext->bAllowInterpretingLODs = ImportContext.ImportOptions->bInterpretLODs;
@@ -2153,6 +2154,7 @@ bool UUsdStageImporter::ReimportSingleAsset(
 	TranslationContext->SubdivisionLevel = ImportContext.ImportOptions->SubdivisionLevel;
 	TranslationContext->MetadataOptions = ImportContext.ImportOptions->MetadataOptions;
 	TranslationContext->KindsToCollapse = (EUsdDefaultKind)ImportContext.ImportOptions->KindsToCollapse;
+	TranslationContext->bUsePrimKindsForCollapsing = ImportContext.ImportOptions->bUsePrimKindsForCollapsing;
 	TranslationContext->bMergeIdenticalMaterialSlots = ImportContext.ImportOptions->bMergeIdenticalMaterialSlots;
 	TranslationContext->bShareAssetsForIdenticalPrims = ImportContext.ImportOptions->bShareAssetsForIdenticalPrims;
 	TranslationContext->bAllowInterpretingLODs = ImportContext.ImportOptions->bInterpretLODs;
diff --git a/Engine/Plugins/Importers/USDImporter/Source/USDStageImporter/Public/USDStageImportOptions.h b/Engine/Plugins/Importers/USDImporter/Source/USDStageImporter/Public/USDStageImportOptions.h
--- a/Engine/Plugins/Importers/USDImporter/Source/USDStageImporter/Public/USDStageImportOptions.h
+++ b/Engine/Plugins/Importers/USDImporter/Source/USDStageImporter/Public/USDStageImportOptions.h
@@ -101,7 +101,13 @@ public:
 	bool bUseExistingAssetCache;
 
 	/** Copy assets from an existing UsdAsset cache instead of generating them from scratch, if possible */
-	UPROPERTY(BlueprintReadWrite, config, EditAnywhere, Category = "USD options", meta = (AllowedClasses = "/Script/USDClasses.UsdAssetCache3", EditCondition = "bUseExistingAssetCache"))
+	UPROPERTY(
+		BlueprintReadWrite,
+		config,
+		EditAnywhere,
+		Category = "USD options",
+		meta = (AllowedClasses = "/Script/USDClasses.UsdAssetCache3", EditCondition = "bUseExistingAssetCache")
+	)
 	FSoftObjectPath ExistingAssetCache = nullptr;
 
 	/** Only import prims with these specific purposes from the USD file */
@@ -188,6 +194,13 @@ public:
 	UPROPERTY(BlueprintReadWrite, config, EditAnywhere, Category = "Processing")
 	bool bPrimPathFolderStructure;
 
+	/**
+	 * Use KindsToCollapse to determine when to collapse prim subtrees or not (defaults to enabled).
+	 * Disable this if you want to prevent collapsing, or to control it manually by right-clicking on individual prims.
+	 */
+	UPROPERTY(BlueprintReadWrite, config, EditAnywhere, Category = "Processing")
+	bool bUsePrimKindsForCollapsing;
+
 	/**
 	 * Whether to try to combine individual assets and components of the same type on a kind-per-kind basis,
 	 * like multiple Mesh prims into a single Static Mesh
@@ -197,7 +210,7 @@ public:
 		config,
 		EditAnywhere,
 		Category = "Processing",
-		meta = (Bitmask, BitmaskEnum = "/Script/UnrealUSDWrapper.EUsdDefaultKind")
+		meta = (Bitmask, BitmaskEnum = "/Script/UnrealUSDWrapper.EUsdDefaultKind", EditCondition = bUsePrimKindsForCollapsing)
 	)
 	int32 KindsToCollapse;
 
diff --git a/Engine/Plugins/Runtime/USDCore/Source/USDUtilities/Private/USDConversionUtils.cpp b/Engine/Plugins/Runtime/USDCore/Source/USDUtilities/Private/USDConversionUtils.cpp
--- a/Engine/Plugins/Runtime/USDCore/Source/USDUtilities/Private/USDConversionUtils.cpp
+++ b/Engine/Plugins/Runtime/USDCore/Source/USDUtilities/Private/USDConversionUtils.cpp
@@ -1524,7 +1524,7 @@ UsdUtils::ECollapsingPreference UsdUtils::GetCollapsingPreference(const pxr::Usd
 		}
 	}
 
-	return ECollapsingPreference::ByKind;
+	return ECollapsingPreference::Default;
 }
 
 bool UsdUtils::SetCollapsingPreference(const pxr::UsdPrim& Prim, UsdUtils::ECollapsingPreference NewPreference)
@@ -1554,9 +1554,9 @@ bool UsdUtils::SetCollapsingPreference(const pxr::UsdPrim& Prim, UsdUtils::EColl
 				return Attr.Set(UnrealIdentifiers::CollapsingAllow);
 				break;
 			}
-			case ECollapsingPreference::ByKind:
+			case ECollapsingPreference::Default:
 			{
-				return Attr.Set(UnrealIdentifiers::CollapsingByKind);
+				return Attr.Set(UnrealIdentifiers::CollapsingDefault);
 				break;
 			}
 			case ECollapsingPreference::Never:
diff --git a/Engine/Plugins/Runtime/USDCore/Source/USDUtilities/Public/USDConversionUtils.h b/Engine/Plugins/Runtime/USDCore/Source/USDUtilities/Public/USDConversionUtils.h
--- a/Engine/Plugins/Runtime/USDCore/Source/USDUtilities/Public/USDConversionUtils.h
+++ b/Engine/Plugins/Runtime/USDCore/Source/USDUtilities/Public/USDConversionUtils.h
@@ -296,8 +296,8 @@ namespace UsdUtils
 
 	enum class ECollapsingPreference : uint8
 	{
+		Default,
 		Allow,
-		ByKind,
 		Never
 	};
 
diff --git a/Engine/Plugins/Runtime/USDCore/Source/UnrealUSDWrapper/Private/UnrealUSDWrapper.cpp b/Engine/Plugins/Runtime/USDCore/Source/UnrealUSDWrapper/Private/UnrealUSDWrapper.cpp
--- a/Engine/Plugins/Runtime/USDCore/Source/UnrealUSDWrapper/Private/UnrealUSDWrapper.cpp
+++ b/Engine/Plugins/Runtime/USDCore/Source/UnrealUSDWrapper/Private/UnrealUSDWrapper.cpp
@@ -170,7 +170,7 @@ namespace UnrealIdentifiers
 	const TfToken UnrealCollapsingAPI = TfToken("CollapsingAPI");
 	const TfToken UnrealCollapsingAttr = TfToken("unreal:collapsing");
 	const TfToken CollapsingAllow = TfToken("allow");
-	const TfToken CollapsingByKind = TfToken("byKind");
+	const TfToken CollapsingDefault = TfToken("default");
 	const TfToken CollapsingNever = TfToken("never");
 }	 // namespace UnrealIdentifiers
 
diff --git a/Engine/Plugins/Runtime/USDCore/Source/UnrealUSDWrapper/Public/UnrealUSDWrapper.h b/Engine/Plugins/Runtime/USDCore/Source/UnrealUSDWrapper/Public/UnrealUSDWrapper.h
--- a/Engine/Plugins/Runtime/USDCore/Source/UnrealUSDWrapper/Public/UnrealUSDWrapper.h
+++ b/Engine/Plugins/Runtime/USDCore/Source/UnrealUSDWrapper/Public/UnrealUSDWrapper.h
@@ -158,7 +158,7 @@ enum class EGeometryCacheImport : uint8
 	// Imported on stage load; played back from the persistent assets
 	OnLoad,
 
-	// Imported on save; geometry caches are streamed from the stage until they are saved. 
+	// Imported on save; geometry caches are streamed from the stage until they are saved.
 	// Afterwards they are played back from the persistent assets
 	OnSave
 };
@@ -488,7 +488,7 @@ namespace UnrealIdentifiers
 	extern UNREALUSDWRAPPER_API const pxr::TfToken UnrealCollapsingAPI;
 	extern UNREALUSDWRAPPER_API const pxr::TfToken UnrealCollapsingAttr;
 	extern UNREALUSDWRAPPER_API const pxr::TfToken CollapsingAllow;
-	extern UNREALUSDWRAPPER_API const pxr::TfToken CollapsingByKind;
+	extern UNREALUSDWRAPPER_API const pxr::TfToken CollapsingDefault;
 	extern UNREALUSDWRAPPER_API const pxr::TfToken CollapsingNever;
 #endif	  // #if USE_USD_SDK
 
