USD: Fix infocache and primlinkcache not being moved to the instance when a blueprint deriving the stage actor is recompiled.

diff --git a/Engine/Plugins/Importers/USDImporter/Source/USDStage/Private/USDStageActor.cpp b/Engine/Plugins/Importers/USDImporter/Source/USDStage/Private/USDStageActor.cpp
--- a/Engine/Plugins/Importers/USDImporter/Source/USDStage/Private/USDStageActor.cpp
+++ b/Engine/Plugins/Importers/USDImporter/Source/USDStage/Private/USDStageActor.cpp
@@ -3216,8 +3216,33 @@ void AUsdStageActor::OnObjectsReplaced(const TMap<UObject*, UObject*>& ObjectRep
 			NewActor->OnPrimChanged = OnPrimChanged;
 
 			NewActor->AssetCache = AssetCache;
-			NewActor->UsdInfoCache = UsdInfoCache;
-			NewActor->PrimLinkCache = PrimLinkCache;
+
+			// Move the built info cache over to the new actor
+			const TCHAR* NewName = nullptr;
+			if (UsdInfoCache)
+			{
+				if (NewActor->UsdInfoCache)
+				{
+					NewActor->UsdInfoCache->Rename(NewName, GetTransientPackage());
+				}
+
+				NewActor->UsdInfoCache = UsdInfoCache;
+				UsdInfoCache->Rename(NewName, NewActor);
+				UsdInfoCache = nullptr;
+			}
+
+			// Move the built prim link cache over to the new actor
+			if (PrimLinkCache)
+			{
+				if (NewActor->PrimLinkCache)
+				{
+					NewActor->PrimLinkCache->Rename(NewName, GetTransientPackage());
+				}
+
+				NewActor->PrimLinkCache = PrimLinkCache;
+				PrimLinkCache->Rename(NewName, NewActor);
+				PrimLinkCache = nullptr;
+			}
 
 			NewActor->BBoxCache = BBoxCache;
 			BBoxCache = nullptr;
