USD: Prevent the binding of stage actor properties other than Time onto its LevelSequences.

diff --git a/Engine/Plugins/Importers/USDImporter/Source/USDStage/Private/USDLevelSequenceHelper.cpp b/Engine/Plugins/Importers/USDImporter/Source/USDStage/Private/USDLevelSequenceHelper.cpp
--- a/Engine/Plugins/Importers/USDImporter/Source/USDStage/Private/USDLevelSequenceHelper.cpp
+++ b/Engine/Plugins/Importers/USDImporter/Source/USDStage/Private/USDLevelSequenceHelper.cpp
@@ -589,6 +589,58 @@ cleanup:
 		}
 	}
 
+	void ShowStageActorPropertyTrackWarning(FName PropertyName)
+	{
+		const FText Text = LOCTEXT("TrackUnboundTitle", "USD: Failed to bind property");
+
+		const FText SubText = FText::Format(
+			LOCTEXT(
+				"TrackUnboundMessage",
+				"Cannot bind the Stage Actor property '{0}' to it's own Level Sequence!\n\nThis sequence is an analogue for animation contained in the USD stage. For now it is not possible to create bindings or bind tracks that cannot be translated back into USD information."
+			),
+			FText::FromName(PropertyName)
+		);
+
+		UE_LOG(LogUsd, Warning, TEXT("%s"), *SubText.ToString().Replace(TEXT("\n\n"), TEXT(" ")));
+
+		static TWeakPtr<SNotificationItem> Notification;
+
+		FNotificationInfo Toast(Text);
+		Toast.SubText = SubText;
+		Toast.Image = FCoreStyle::Get().GetBrush(TEXT("MessageLog.Warning"));
+		Toast.bUseLargeFont = false;
+		Toast.bFireAndForget = false;
+		Toast.FadeOutDuration = 0.0f;
+		Toast.ExpireDuration = 0.0f;
+		Toast.bUseThrobber = false;
+		Toast.bUseSuccessFailIcons = false;
+		Toast.ButtonDetails.Emplace(
+			LOCTEXT("TrackUnboundOk", "Ok"),
+			FText::GetEmpty(),
+			FSimpleDelegate::CreateLambda(
+				[]()
+				{
+					if (TSharedPtr<SNotificationItem> PinnedNotification = Notification.Pin())
+					{
+						PinnedNotification->SetCompletionState(SNotificationItem::CS_Success);
+						PinnedNotification->ExpireAndFadeout();
+					}
+				}
+			)
+		);
+
+		// Only show one at a time
+		if (!Notification.IsValid())
+		{
+			Notification = FSlateNotificationManager::Get().AddNotification(Toast);
+		}
+
+		if (TSharedPtr<SNotificationItem> PinnedNotification = Notification.Pin())
+		{
+			PinnedNotification->SetCompletionState(SNotificationItem::CS_Pending);
+		}
+	}
+
 	void ShowVisibilityWarningIfNeeded(const UMovieScenePropertyTrack* PropertyTrack, const UE::FUsdPrim& UsdPrim)
 	{
 		if (!PropertyTrack || !UsdPrim)
@@ -928,7 +980,7 @@ private:
 	void HandleMovieSceneChange(UMovieScene& MovieScene);
 	void HandleSubSectionChange(UMovieSceneSubSection& Section);
 	void HandleControlRigSectionChange(UMovieSceneControlRigParameterSection& Section);
-	void HandleTrackChange(const UMovieSceneTrack& Track, bool bIsMuteChange);
+	void HandleTrackChange(UMovieSceneTrack& Track, bool bIsMuteChange);
 	void HandleDisplayRateChange(const double DisplayRate);
 
 	FDelegateHandle OnObjectTransactedHandle;
@@ -4468,7 +4520,7 @@ void FUsdLevelSequenceHelperImpl::HandleControlRigSectionChange(UMovieSceneContr
 #endif	  // WITH_EDITOR
 }
 
-void FUsdLevelSequenceHelperImpl::HandleTrackChange(const UMovieSceneTrack& Track, bool bIsMuteChange)
+void FUsdLevelSequenceHelperImpl::HandleTrackChange(UMovieSceneTrack& Track, bool bIsMuteChange)
 {
 	if (!StageActor.IsValid())
 	{
@@ -4506,6 +4558,18 @@ void FUsdLevelSequenceHelperImpl::HandleTrackChange(const UMovieSceneTrack& Trac
 		return;
 	}
 
+	// The only stage actor property we allow binding on the transient level sequence is 'Time'. Anything else we
+	// need to force-unbind as not only will it be lost when reloading the stage anyway, but it can even lead to
+	// crashes (e.g. UE-215067)
+	UMovieScenePropertyTrack* PropertyTrack = Cast<UMovieScenePropertyTrack>(&Track);
+	const FName PropertyPath = PropertyTrack ? PropertyTrack->GetPropertyPath() : NAME_None;
+	if (BoundObject == StageActor.Get() && PropertyPath != GET_MEMBER_NAME_CHECKED(AUsdStageActor, Time))
+	{
+		UsdLevelSequenceHelperImpl::ShowStageActorPropertyTrackWarning(PropertyPath);
+		MovieScene->RemoveTrack(*PropertyTrack);
+		return;
+	}
+
 	// Our tracked bindings are always directly to components
 	USceneComponent* BoundSceneComponent = Cast<USceneComponent>(BoundObject);
 	if (!BoundSceneComponent)
@@ -4531,10 +4595,8 @@ void FUsdLevelSequenceHelperImpl::HandleTrackChange(const UMovieSceneTrack& Trac
 	// for the camera component directly, so try again
 	if (!PrimTwin && BoundSceneComponent->IsA<UCineCameraComponent>())
 	{
-		if (const UMovieScenePropertyTrack* PropertyTrack = Cast<const UMovieScenePropertyTrack>(&Track))
+		if (PropertyTrack)
 		{
-			const FName& PropertyPath = PropertyTrack->GetPropertyPath();
-
 			// In the scenario where we're trying to make non-decomposed Camera prims work, we only ever want to write out
 			// actual camera properties from the CameraComponent to the Camera prim. We won't write its USceneComponent
 			// properties, as we will use the ones from the ACineCameraActor's parent USceneComponent instead
@@ -4582,10 +4644,8 @@ void FUsdLevelSequenceHelperImpl::HandleTrackChange(const UMovieSceneTrack& Trac
 
 		if (bIsMuteChange)
 		{
-			if (const UMovieScenePropertyTrack* PropertyTrack = Cast<const UMovieScenePropertyTrack>(&Track))
+			if (PropertyTrack)
 			{
-				const FName& PropertyPath = PropertyTrack->GetPropertyPath();
-
 				TArray<UE::FUsdAttribute> Attrs = UsdUtils::GetAttributesForProperty(UsdPrim, PropertyPath);
 				if (Attrs.Num() > 0)
 				{
@@ -4663,7 +4723,7 @@ void FUsdLevelSequenceHelperImpl::HandleTrackChange(const UMovieSceneTrack& Trac
 			// Right now we don't write out changes to SkeletalAnimation tracks, and only property tracks... the UAnimSequence
 			// asset can't be modified all that much in UE anyway. Later on we may want to enable writing it out anyway though,
 			// and pick up on changes to the section offset or play rate and bake out the UAnimSequence again
-			if (const UMovieScenePropertyTrack* PropertyTrack = Cast<const UMovieScenePropertyTrack>(&Track))
+			if (PropertyTrack)
 			{
 #if WITH_EDITOR
 				UsdLevelSequenceHelperImpl::ShowVisibilityWarningIfNeeded(PropertyTrack, UsdPrim);
