USD: Use "usdc" as default payload format.

diff --git a/Engine/Plugins/Importers/USDImporter/Source/USDExporter/Public/USDAssetOptions.h b/Engine/Plugins/Importers/USDImporter/Source/USDExporter/Public/USDAssetOptions.h
--- a/Engine/Plugins/Importers/USDImporter/Source/USDExporter/Public/USDAssetOptions.h
+++ b/Engine/Plugins/Importers/USDImporter/Source/USDExporter/Public/USDAssetOptions.h
@@ -62,7 +62,7 @@ struct USDEXPORTER_API FUsdMeshAssetOptions
 		Category = "Mesh options",
 		meta = (EditCondition = bUsePayload, GetOptions = "USDExporter.LevelExporterUSDOptions.GetUsdExtensions")
 	)
-	FString PayloadFormat;
+	FString PayloadFormat = TEXT("usdc");
 
 	/** Whether to bake the mesh's assigned material and export these as separate UsdPreviewSurface assets */
 	UPROPERTY(EditAnywhere, config, BlueprintReadWrite, Category = "Material options")
