USD: Fix crash when importing a stage with a component of the previous import already selected.

diff --git a/Engine/Plugins/Importers/USDImporter/Source/USDStageImporter/Private/USDStageImporter.cpp b/Engine/Plugins/Importers/USDImporter/Source/USDStageImporter/Private/USDStageImporter.cpp
--- a/Engine/Plugins/Importers/USDImporter/Source/USDStageImporter/Private/USDStageImporter.cpp
+++ b/Engine/Plugins/Importers/USDImporter/Source/USDStageImporter/Private/USDStageImporter.cpp
@@ -1829,6 +1829,10 @@ void UUsdStageImporter::ImportFromFile(FUsdStageImportContext& ImportContext)
 		return;
 	}
 
+	// Reset the editor selection because the import may destroy actors/components or stomp over existing assets.
+	// It's probably too aggressive to do this, but it should be the fastest and safest at least.
+	GEditor->ResetAllSelectionSets();
+
 	// Reset the transaction buffer so we can't undo the creation of the assets.
 	// We'll also temporarily disable creating new transactions, to prevent code downstream from accidentally creating new ones.
 	// This happens for example on USkeleton::AccumulateCurveMetaData, and can lead to thousands of transactions showing up on
@@ -2044,6 +2048,10 @@ bool UUsdStageImporter::ReimportSingleAsset(
 		return bSuccess;
 	}
 
+	// Reset the editor selection because the import may destroy actors/components or stomp over existing assets.
+	// It's probably too aggressive to do this, but it should be the fastest and safest at least.
+	GEditor->ResetAllSelectionSets();
+
 	// Reset the transaction buffer so we can't undo the creation of the assets.
 	// We'll also temporarily disable creating new transactions, to prevent code downstream from accidentally creating new ones.
 	// This happens for example on USkeleton::AccumulateCurveMetaData, and can lead to thousands of transactions showing up on
