USD: Try to pick up any new existing assets on the asset cache's folder before loading a stage.

diff --git a/Engine/Plugins/Importers/USDImporter/Source/USDStage/Private/USDStageActor.cpp b/Engine/Plugins/Importers/USDImporter/Source/USDStage/Private/USDStageActor.cpp
--- a/Engine/Plugins/Importers/USDImporter/Source/USDStage/Private/USDStageActor.cpp
+++ b/Engine/Plugins/Importers/USDImporter/Source/USDStage/Private/USDStageActor.cpp
@@ -3296,6 +3296,14 @@ void AUsdStageActor::LoadUsdStage(bool bOpenIfNeeded)
 		}
 	}
 
+	// Make sure the asset cache tries picking up any existing asset on its UsdAssets folder before we try
+	// creating new assets. We don't want to do this *too* often as we may have thousands of assets here,
+	// but once before each load should be fine and could potentially save on stage load times
+	if (AssetCache)
+	{
+		AssetCache->RescanAssetDirectory();
+	}
+
 	TOptional<TGuardValue<ITransaction*>> SuppressTransaction;
 	if (GDiscardUndoBufferOnStageOpenClose)
 	{
