USD: Fix broken Mac Editor build.

diff --git a/Engine/Plugins/Importers/USDImporter/Source/USDSchemas/Private/Custom/MDLUSDShadeMaterialTranslator.h b/Engine/Plugins/Importers/USDImporter/Source/USDSchemas/Private/Custom/MDLUSDShadeMaterialTranslator.h
--- a/Engine/Plugins/Importers/USDImporter/Source/USDSchemas/Private/Custom/MDLUSDShadeMaterialTranslator.h
+++ b/Engine/Plugins/Importers/USDImporter/Source/USDSchemas/Private/Custom/MDLUSDShadeMaterialTranslator.h
@@ -17,7 +17,7 @@ class FMdlUsdShadeMaterialTranslator : public FMaterialXUsdShadeMaterialTranslat
 	using Super = FMaterialXUsdShadeMaterialTranslator;
 
 public:
-	UE_DEPRECATED(5.5, "Use UnrealIdentifiers::MdlRenderContext instead");
+	UE_DEPRECATED(5.5, "Use UnrealIdentifiers::MdlRenderContext instead")
 	static FName MdlRenderContext;
 
 public:
