USD: Fix issue packaging QAEnterprise due to WITH_EDITOR check.

diff --git a/Engine/Plugins/Importers/USDImporter/Source/USDStage/Private/USDLevelSequenceHelper.cpp b/Engine/Plugins/Importers/USDImporter/Source/USDStage/Private/USDLevelSequenceHelper.cpp
--- a/Engine/Plugins/Importers/USDImporter/Source/USDStage/Private/USDLevelSequenceHelper.cpp
+++ b/Engine/Plugins/Importers/USDImporter/Source/USDStage/Private/USDLevelSequenceHelper.cpp
@@ -4565,7 +4565,9 @@ void FUsdLevelSequenceHelperImpl::HandleTrackChange(UMovieSceneTrack& Track, boo
 	const FName PropertyPath = PropertyTrack ? PropertyTrack->GetPropertyPath() : NAME_None;
 	if (BoundObject == StageActor.Get() && PropertyPath != GET_MEMBER_NAME_CHECKED(AUsdStageActor, Time))
 	{
+#if WITH_EDITOR
 		UsdLevelSequenceHelperImpl::ShowStageActorPropertyTrackWarning(PropertyPath);
+#endif	  // WITH_EDITOR
 		MovieScene->RemoveTrack(*PropertyTrack);
 		return;
 	}
