USD: Prevent the clear payloads/references buttons from being enabled if the prim doesn't have a payload/reference spec in the local layer stack.

diff --git a/Engine/Plugins/Importers/USDImporter/Source/USDStageEditor/Private/SUSDStageTreeView.cpp b/Engine/Plugins/Importers/USDImporter/Source/USDStageEditor/Private/SUSDStageTreeView.cpp
--- a/Engine/Plugins/Importers/USDImporter/Source/USDStageEditor/Private/SUSDStageTreeView.cpp
+++ b/Engine/Plugins/Importers/USDImporter/Source/USDStageEditor/Private/SUSDStageTreeView.cpp
@@ -858,11 +858,18 @@ TSharedPtr<SWidget> SUsdStageTreeView::ConstructPrimContextMenu()
 
 		PrimOptions.AddMenuEntry(
 			LOCTEXT("ClearReferences", "Clear References"),
-			LOCTEXT("ClearReferences_ToolTip", "Clears the references for this prim"),
+			TAttribute<FText>::Create(TAttribute<FText>::FGetter::CreateLambda(
+				[this]()
+				{
+					return DoesPrimHaveReferenceSpecOnLocalLayerStack()
+							   ? LOCTEXT("ClearReferences_ToolTip", "Clears the references for this prim")
+							   : LOCTEXT("ClearReferencesNoSpec_ToolTip", "The prim doesn't have any reference spec on the current edit target");
+				}
+			)),
 			FSlateIcon(),
 			FUIAction(
 				FExecuteAction::CreateSP(this, &SUsdStageTreeView::OnClearReferences),
-				FCanExecuteAction::CreateSP(this, &SUsdStageTreeView::DoesPrimExistOnEditTarget)
+				FCanExecuteAction::CreateSP(this, &SUsdStageTreeView::DoesPrimHaveReferenceSpecOnLocalLayerStack)
 			),
 			NAME_None,
 			EUserInterfaceActionType::Button
@@ -887,11 +894,18 @@ TSharedPtr<SWidget> SUsdStageTreeView::ConstructPrimContextMenu()
 
 		PrimOptions.AddMenuEntry(
 			LOCTEXT("ClearPayloads", "Clear Payloads"),
-			LOCTEXT("ClearPayloads_ToolTip", "Clears the payloads for this prim"),
+			TAttribute<FText>::Create(TAttribute<FText>::FGetter::CreateLambda(
+				[this]()
+				{
+					return DoesPrimHavePayloadSpecOnLocalLayerStack()
+							   ? LOCTEXT("ClearPayloads_ToolTip", "Clears the payloads for this prim")
+							   : LOCTEXT("ClearPayloadsNoSpec_ToolTip", "The prim doesn't have any payload spec on the current edit target");
+				}
+			)),
 			FSlateIcon(),
 			FUIAction(
 				FExecuteAction::CreateSP(this, &SUsdStageTreeView::OnClearPayloads),
-				FCanExecuteAction::CreateSP(this, &SUsdStageTreeView::DoesPrimExistOnEditTarget)
+				FCanExecuteAction::CreateSP(this, &SUsdStageTreeView::DoesPrimHavePayloadSpecOnLocalLayerStack)
 			),
 			NAME_None,
 			EUserInterfaceActionType::Button
@@ -1469,9 +1483,7 @@ bool SUsdStageTreeView::DoesPrimExistOnEditTarget() const
 
 bool SUsdStageTreeView::DoesPrimHaveSpecOnLocalLayerStack() const
 {
-	TArray<FUsdPrimViewModelRef> MySelectedItems = GetSelectedItems();
-
-	for (FUsdPrimViewModelRef SelectedItem : MySelectedItems)
+	for (FUsdPrimViewModelRef SelectedItem : GetSelectedItems())
 	{
 		if (SelectedItem->HasSpecsOnLocalLayer())
 		{
@@ -1482,6 +1494,32 @@ bool SUsdStageTreeView::DoesPrimHaveSpecOnLocalLayerStack() const
 	return false;
 }
 
+bool SUsdStageTreeView::DoesPrimHaveReferenceSpecOnLocalLayerStack() const
+{
+	for (FUsdPrimViewModelRef SelectedItem : GetSelectedItems())
+	{
+		if (SelectedItem->HasReferencesOnLocalLayer())
+		{
+			return true;
+		}
+	}
+
+	return false;
+}
+
+bool SUsdStageTreeView::DoesPrimHavePayloadSpecOnLocalLayerStack() const
+{
+	for (FUsdPrimViewModelRef SelectedItem : GetSelectedItems())
+	{
+		if (SelectedItem->HasPayloadsOnLocalLayer())
+		{
+			return true;
+		}
+	}
+
+	return false;
+}
+
 bool SUsdStageTreeView::DoSelectedPrimsHaveCollapsingPreference(UsdUtils::ECollapsingPreference TargetPreference) const
 {
 	TArray<FUsdPrimViewModelRef> MySelectedItems = GetSelectedItems();
diff --git a/Engine/Plugins/Importers/USDImporter/Source/USDStageEditor/Private/SUSDStageTreeView.h b/Engine/Plugins/Importers/USDImporter/Source/USDStageEditor/Private/SUSDStageTreeView.h
--- a/Engine/Plugins/Importers/USDImporter/Source/USDStageEditor/Private/SUSDStageTreeView.h
+++ b/Engine/Plugins/Importers/USDImporter/Source/USDStageEditor/Private/SUSDStageTreeView.h
@@ -98,6 +98,8 @@ private:
 	bool DoesPrimExistOnStage() const;
 	bool DoesPrimExistOnEditTarget() const;
 	bool DoesPrimHaveSpecOnLocalLayerStack() const;
+	bool DoesPrimHaveReferenceSpecOnLocalLayerStack() const;
+	bool DoesPrimHavePayloadSpecOnLocalLayerStack() const;
 	bool DoSelectedPrimsHaveCollapsingPreference(UsdUtils::ECollapsingPreference Preference) const;
 
 	void RequestExpansionStateRestore();
diff --git a/Engine/Plugins/Importers/USDImporter/Source/USDStageEditorViewModels/Private/USDPrimViewModel.cpp b/Engine/Plugins/Importers/USDImporter/Source/USDStageEditorViewModels/Private/USDPrimViewModel.cpp
--- a/Engine/Plugins/Importers/USDImporter/Source/USDStageEditorViewModels/Private/USDPrimViewModel.cpp
+++ b/Engine/Plugins/Importers/USDImporter/Source/USDStageEditorViewModels/Private/USDPrimViewModel.cpp
@@ -339,6 +339,61 @@ bool FUsdPrimViewModel::HasSpecsOnLocalLayer() const
 	return false;
 }
 
+bool FUsdPrimViewModel::HasReferencesOnLocalLayer() const
+{
+#if USE_USD_SDK
+	FScopedUsdAllocs UsdAllocs;
+
+	pxr::UsdPrim PxrUsdPrim{UsdPrim};
+	if (PxrUsdPrim)
+	{
+		if (pxr::UsdStageRefPtr PrimUsdStage = PxrUsdPrim.GetStage())
+		{
+			for (const pxr::SdfPrimSpecHandle& Spec : PxrUsdPrim.GetPrimStack())
+			{
+				if (Spec && Spec->HasReferences())
+				{
+					pxr::SdfLayerHandle SpecLayer = Spec->GetLayer();
+					if (PrimUsdStage->GetEditTarget().GetLayer() == SpecLayer && PrimUsdStage->HasLocalLayer(SpecLayer))
+					{
+						return true;
+					}
+				}
+			}
+		}
+	}
+#endif	  // #if USE_USD_SDK
+
+	return false;
+}
+bool FUsdPrimViewModel::HasPayloadsOnLocalLayer() const
+{
+#if USE_USD_SDK
+	FScopedUsdAllocs UsdAllocs;
+
+	pxr::UsdPrim PxrUsdPrim{UsdPrim};
+	if (PxrUsdPrim)
+	{
+		if (pxr::UsdStageRefPtr PrimUsdStage = PxrUsdPrim.GetStage())
+		{
+			for (const pxr::SdfPrimSpecHandle& Spec : PxrUsdPrim.GetPrimStack())
+			{
+				if (Spec && Spec->HasPayloads())
+				{
+					pxr::SdfLayerHandle SpecLayer = Spec->GetLayer();
+					if (PrimUsdStage->GetEditTarget().GetLayer() == SpecLayer && PrimUsdStage->HasLocalLayer(SpecLayer))
+					{
+						return true;
+					}
+				}
+			}
+		}
+	}
+#endif	  // #if USE_USD_SDK
+
+	return false;
+}
+
 void FUsdPrimViewModel::DefinePrim(const TCHAR* PrimName)
 {
 #if USE_USD_SDK
diff --git a/Engine/Plugins/Importers/USDImporter/Source/USDStageEditorViewModels/Public/USDPrimViewModel.h b/Engine/Plugins/Importers/USDImporter/Source/USDStageEditorViewModels/Public/USDPrimViewModel.h
--- a/Engine/Plugins/Importers/USDImporter/Source/USDStageEditorViewModels/Public/USDPrimViewModel.h
+++ b/Engine/Plugins/Importers/USDImporter/Source/USDStageEditorViewModels/Public/USDPrimViewModel.h
@@ -74,6 +74,8 @@ public:
 
 	// Returns true if this prim has at least one spec on the stage's local layer stack
 	bool HasSpecsOnLocalLayer() const;
+	bool HasReferencesOnLocalLayer() const;
+	bool HasPayloadsOnLocalLayer() const;
 
 	void DefinePrim(const TCHAR* PrimName);
 
