USD: Fix isolate layers functionality not using UE transactions or being undo/redoable.

diff --git a/Engine/Plugins/Importers/USDImporter/Source/USDStage/Private/USDStageActor.cpp b/Engine/Plugins/Importers/USDImporter/Source/USDStage/Private/USDStageActor.cpp
--- a/Engine/Plugins/Importers/USDImporter/Source/USDStage/Private/USDStageActor.cpp
+++ b/Engine/Plugins/Importers/USDImporter/Source/USDStage/Private/USDStageActor.cpp
@@ -1178,7 +1178,7 @@ FString AUsdStageActor::GetIsolatedRootLayer() const
 	return IsolatedStage ? IsolatedStage.GetRootLayer().GetIdentifier() : FString{};
 }
 
-void AUsdStageActor::IsolateLayer(const UE::FSdfLayer& Layer)
+void AUsdStageActor::IsolateLayer(const UE::FSdfLayer& Layer, bool bLoadUsdStage)
 {
 	if (IsolatedStage && IsolatedStage.GetRootLayer() == Layer)
 	{
@@ -1198,18 +1198,17 @@ void AUsdStageActor::IsolateLayer(const UE::FSdfLayer& Layer)
 
 	OnPreStageChanged.Broadcast();
 
+	UE::FUsdStage NewIsolatedStage;
+	UE::FUsdStage StageToListenTo;
+
 	// Stop isolating
 	if (!Layer || Layer == UsdStage.GetRootLayer())
 	{
-		IsolatedStage = UE::FUsdStage{};
-
-		UsdListener.Register(UsdStage);
+		NewIsolatedStage = UE::FUsdStage{};
 
-		LoadUsdStage();
-		return;
+		StageToListenTo = UsdStage;
 	}
-
-	if (UsdStage)
+	else if (UsdStage)
 	{
 		// We should only be allowed to isolate a layer belonging to UsdStage's local layer stack, but checking for that
 		// is not trivial given that layers can be muted.
@@ -1248,18 +1247,45 @@ void AUsdStageActor::IsolateLayer(const UE::FSdfLayer& Layer)
 		// We really want our own stage for this and not something from the stage cache.
 		// Plus, this means its easier to cleanup: Just drop our IsolatedStage
 		const bool bUseStageCache = false;
-		IsolatedStage = UnrealUSDWrapper::OpenStage(Layer, {}, EUsdInitialLoadSet::LoadAll, bUseStageCache);
-		IsolatedStage.SetEditTarget(IsolatedStage.GetRootLayer());
-		IsolatedStage.SetInterpolationType(InterpolationType);
+		NewIsolatedStage = UnrealUSDWrapper::OpenStage(Layer, {}, EUsdInitialLoadSet::LoadAll, bUseStageCache);
+		NewIsolatedStage.SetEditTarget(NewIsolatedStage.GetRootLayer());
+		NewIsolatedStage.SetInterpolationType(InterpolationType);
 
-		UsdListener.Register(IsolatedStage);
+		StageToListenTo = NewIsolatedStage;
+	}
 
-		LoadUsdStage();
+	// For now we're kind of using a small hack where we record our isolated layer before and after we change it.
+	// When undo/redoing through the accumulated edits for this transaction, the transactor will play them out in the
+	// correct direction (forwards when redoing, backwards when undoing), and switch the isolated layer when needed and
+	// do nothing when the isolated layer is already correct, so we lose nothing by calling this twice per transaction
+	// and get the benefit of the change application "direction" code automatically working.
+	//
+	// Later on we should hopefully get a command pattern undo/redo system, and we can get something cleaner than this
+	auto RecordIsolatedLayer = [this]()
+	{
+		if (Transactor)
+		{
+			const static UsdUtils::FObjectChangesByPath Empty;
+			Transactor->Update(Empty, Empty);
+		}
+	};
 
-		// Fire this so that the USD Stage Editor knows to refresh.
-		// Plus we kind of changed the active stage too
-		OnStageChanged.Broadcast();
+	RecordIsolatedLayer();
+	{
+		IsolatedStage = NewIsolatedStage;
 	}
+	RecordIsolatedLayer();
+
+	UsdListener.Register(StageToListenTo);
+
+	if (bLoadUsdStage)
+	{
+		LoadUsdStage();
+	}
+
+	// Fire this so that the USD Stage Editor knows to refresh.
+	// Plus we kind of changed the active stage too
+	OnStageChanged.Broadcast();
 }
 
 void AUsdStageActor::OnUsdObjectsChanged(const UsdUtils::FObjectChangesByPath& InfoChanges, const UsdUtils::FObjectChangesByPath& ResyncChanges)
diff --git a/Engine/Plugins/Importers/USDImporter/Source/USDStage/Private/USDTransactor.cpp b/Engine/Plugins/Importers/USDImporter/Source/USDStage/Private/USDTransactor.cpp
--- a/Engine/Plugins/Importers/USDImporter/Source/USDStage/Private/USDTransactor.cpp
+++ b/Engine/Plugins/Importers/USDImporter/Source/USDStage/Private/USDTransactor.cpp
@@ -65,6 +65,7 @@ namespace UsdUtils
 	struct FTransactorRecordedEdits
 	{
 		FString EditTargetIdentifier;
+		FString IsolatedLayerIdentifier;
 		TArray<FTransactorRecordedEdit> Edits;	  // Kept in the order of recording
 	};
 
@@ -131,6 +132,7 @@ FArchive& operator<<(FArchive& Ar, UsdUtils::FTransactorRecordedEdit& Edit)
 FArchive& operator<<(FArchive& Ar, UsdUtils::FTransactorRecordedEdits& Edits)
 {
 	Ar << Edits.EditTargetIdentifier;
+	Ar << Edits.IsolatedLayerIdentifier;
 	Ar << Edits.Edits;
 	return Ar;
 }
@@ -703,8 +705,9 @@ namespace UsdUtils
 	}
 
 	/** Applies the field value pairs to all prims on the stage, and returns a list of prim paths for modified prims */
-	TSet<FString> ApplyFieldMapToStage(const FTransactorEditStorage& EditStorage, EApplicationDirection Direction, UE::FUsdStage& Stage, double Time)
+	TSet<FString> ApplyFieldMapToStage(const FTransactorEditStorage& EditStorage, EApplicationDirection Direction, AUsdStageActor* StageActor)
 	{
+		UE::FUsdStage& Stage = StageActor->GetOrOpenUsdStage();
 		if (!Stage)
 		{
 			return {};
@@ -735,6 +738,19 @@ namespace UsdUtils
 		{
 			const FTransactorRecordedEdits& Edits = EditStorage[EditIndex];
 
+			// Isolate the correct layer (empty string here means to stop isolating / don't isolate anything)
+			const FString CurrentIsolatedLayer = StageActor->GetIsolatedRootLayer();
+			if (CurrentIsolatedLayer != Edits.IsolatedLayerIdentifier)
+			{
+				UE::FSdfLayer LayerToIsolate = Edits.IsolatedLayerIdentifier.IsEmpty() ? UE::FSdfLayer{}
+																					   : UE::FSdfLayer::FindOrOpen(*Edits.IsolatedLayerIdentifier);
+
+				// Don't load here, as we'll get our assets, actors and components back from the UE transaction being
+				// undone/redone already
+				const bool bLoadUsdStage = false;
+				StageActor->IsolateLayer(LayerToIsolate, bLoadUsdStage);
+			}
+
 			if (Edits.EditTargetIdentifier != LastEditTargetIdentifier)
 			{
 #if USE_USD_SDK
@@ -985,13 +1001,10 @@ namespace UsdUtils
 			// transaction buffer.
 			FScopedBlockNotices BlockNotices(StageActor->GetUsdListener());
 
-			UE::FUsdStage& Stage = StageActor->GetOrOpenUsdStage();
-
 			TSet<FString> PrimsChanged = UsdUtils::ApplyFieldMapToStage(	//
 				Values,
 				UsdUtils::EApplicationDirection::Reverse,
-				Stage,
-				StageActor->GetTime()
+				StageActor
 			);
 
 			// Partial rebuild of the info cache after we have undone the USD stage changes for this transaction
@@ -1049,14 +1062,13 @@ namespace UsdUtils
 				PrimsChanged = UsdUtils::ApplyFieldMapToStage(
 					ReceivedValuesBeforeUndo.GetValue(),
 					UsdUtils::EApplicationDirection::Reverse,
-					Stage,
-					StageActor->GetTime()
+					StageActor
 				);
 			}
 			else
 			{
 				// Just a common Redo operation or any other type of ConcertSync transaction, so just apply the new values
-				PrimsChanged = UsdUtils::ApplyFieldMapToStage(Values, UsdUtils::EApplicationDirection::Forward, Stage, StageActor->GetTime());
+				PrimsChanged = UsdUtils::ApplyFieldMapToStage(Values, UsdUtils::EApplicationDirection::Forward, StageActor);
 
 				// Partial rebuild of the info cache after we have redone the USD stage changes for this transaction
 				StageActor->RebuildInfoCacheFromStoredChanges();
@@ -1204,6 +1216,7 @@ void UUsdTransactor::Update(const UsdUtils::FObjectChangesByPath& NewInfoChanges
 
 	UsdUtils::FTransactorRecordedEdits NewEdits;
 	NewEdits.EditTargetIdentifier = EditTarget.GetIdentifier();
+	NewEdits.IsolatedLayerIdentifier = StageActorPtr->GetIsolatedRootLayer();
 
 	UsdUtils::ConvertFieldValueMap(NewInfoChanges, Stage, NewEdits);
 	UsdUtils::ConvertFieldValueMap(NewResyncChanges, Stage, NewEdits);
diff --git a/Engine/Plugins/Importers/USDImporter/Source/USDStage/Public/USDStageActor.h b/Engine/Plugins/Importers/USDImporter/Source/USDStage/Public/USDStageActor.h
--- a/Engine/Plugins/Importers/USDImporter/Source/USDStage/Public/USDStageActor.h
+++ b/Engine/Plugins/Importers/USDImporter/Source/USDStage/Public/USDStageActor.h
@@ -329,8 +329,11 @@ public:
 
 	USDSTAGE_API void SetUsdStage(const UE::FUsdStage& NewStage);
 
-	// Enters isolated mode by creating a new USD Stage using the provided layer as its root
-	USDSTAGE_API void IsolateLayer(const UE::FSdfLayer& Layer);
+	/**
+	 * Enters isolated mode by creating a new USD Stage using the provided layer as its root.
+	 * @param bLoadUsdStage - Whether to load assets, actors and components from the isolated stage
+	 */
+	USDSTAGE_API void IsolateLayer(const UE::FSdfLayer& Layer, bool bLoadUsdStage = true);
 
 	UE_DEPRECATED(5.5, "This function has been renamed into 'RegenerateLevelSequence', which better describes what it does")
 	USDSTAGE_API void ReloadAnimations();
diff --git a/Engine/Plugins/Importers/USDImporter/Source/USDStageEditor/Private/SUSDLayersTreeView.cpp b/Engine/Plugins/Importers/USDImporter/Source/USDStageEditor/Private/SUSDLayersTreeView.cpp
--- a/Engine/Plugins/Importers/USDImporter/Source/USDStageEditor/Private/SUSDLayersTreeView.cpp
+++ b/Engine/Plugins/Importers/USDImporter/Source/USDStageEditor/Private/SUSDLayersTreeView.cpp
@@ -1426,6 +1426,8 @@ void SUsdLayersTreeView::OnIsolateSelectedLayer()
 		return;
 	}
 
+	FScopedTransaction Transaction(LOCTEXT("IsolateTransaction", "Changed layer isolation"));
+
 	for (FUsdLayerViewModelRef SelectedItem : MySelectedItems)
 	{
 		UE::FSdfLayer Layer = SelectedItem->GetLayer();
diff --git a/Engine/Plugins/Importers/USDImporter/Source/USDStageEditor/Private/SUSDStage.cpp b/Engine/Plugins/Importers/USDImporter/Source/USDStageEditor/Private/SUSDStage.cpp
--- a/Engine/Plugins/Importers/USDImporter/Source/USDStageEditor/Private/SUSDStage.cpp
+++ b/Engine/Plugins/Importers/USDImporter/Source/USDStageEditor/Private/SUSDStage.cpp
@@ -2744,6 +2744,8 @@ void SUsdStage::FileClose()
 
 void SUsdStage::OnLayerIsolated(const UE::FSdfLayer& IsolatedLayer)
 {
+	FScopedTransaction Transaction(LOCTEXT("StopIsolatingButtonTransaction", "Changed layer isolation"));
+
 	if (AUsdStageActor* StageActor = ViewModel.UsdStageActor.Get())
 	{
 		StageActor->IsolateLayer(IsolatedLayer);
