USD: Fix multiple issues when manipulating visibility of the prims/components related to point instancers.

Fix hiding point instancer prims not hiding the generated components;
Fix hiding individual HISM components not hiding the prototype prims;
Fix HISM components being left over after hiding/showing point instancer prims;
Add automated test to check for these issues happening again.

diff --git a/Engine/Plugins/Importers/USDImporter/Source/USDSchemas/Private/USDGeomPointInstancerTranslator.cpp b/Engine/Plugins/Importers/USDImporter/Source/USDSchemas/Private/USDGeomPointInstancerTranslator.cpp
--- a/Engine/Plugins/Importers/USDImporter/Source/USDSchemas/Private/USDGeomPointInstancerTranslator.cpp
+++ b/Engine/Plugins/Importers/USDImporter/Source/USDSchemas/Private/USDGeomPointInstancerTranslator.cpp
@@ -92,11 +92,13 @@ namespace UsdGeomPointInstancerTranslatorImpl
 FUsdGeomPointInstancerCreateAssetsTaskChain::FUsdGeomPointInstancerCreateAssetsTaskChain(
 	const TSharedRef<FUsdSchemaTranslationContext>& InContext,
 	const UE::FSdfPath& InPrimPath,
-	bool bInIgnoreTopLevelTransformAndVisibility,
+	bool bInIgnoreTopLevelTransform,
+	bool bInIgnoreTopLevelVisibility,
 	const TOptional<UE::FSdfPath>& InAlternativePrimToLinkAssetsTo
 )
 	: FBuildStaticMeshTaskChain(InContext, InPrimPath, InAlternativePrimToLinkAssetsTo)
-	, bIgnoreTopLevelTransformAndVisibility(bInIgnoreTopLevelTransformAndVisibility)
+	, bIgnoreTopLevelTransform(bInIgnoreTopLevelTransform)
+	, bIgnoreTopLevelVisibility(bInIgnoreTopLevelVisibility)
 {
 	SetupTasks();
 }
@@ -135,7 +137,14 @@ void FUsdGeomPointInstancerCreateAssetsTaskChain::SetupTasks()
 		   Options.bMergeIdenticalMaterialSlots = Context->bMergeIdenticalMaterialSlots;
 		   Options.SubdivisionLevel = Context->SubdivisionLevel;
 
-		   UsdToUnreal::ConvertGeomMeshHierarchy(GetPrim(), AddedMeshDescription, AssignmentInfo, Options, bIgnoreTopLevelTransformAndVisibility);
+		   UsdToUnreal::ConvertGeomMeshSubtree(
+			   GetPrim(),
+			   AddedMeshDescription,
+			   AssignmentInfo,
+			   Options,
+			   bIgnoreTopLevelTransform,
+			   bIgnoreTopLevelVisibility
+		   );
 
 		   return !AddedMeshDescription.IsEmpty();
 	   });
@@ -172,8 +181,12 @@ void FUsdGeomPointInstancerTranslator::CreateAssets()
 		// Don't bake our actual point instancer's transform or visibility into the mesh as its nice to have these on the static mesh component
 		// instead
 		const bool bIgnoreTopLevelTransformAndVisibility = true;
-		Context->TranslatorTasks.Add(MakeShared<FUsdGeomPointInstancerCreateAssetsTaskChain>(Context, PrimPath, bIgnoreTopLevelTransformAndVisibility)
-		);
+		Context->TranslatorTasks.Add(MakeShared<FUsdGeomPointInstancerCreateAssetsTaskChain>(
+			Context,
+			PrimPath,
+			bIgnoreTopLevelTransformAndVisibility,
+			bIgnoreTopLevelTransformAndVisibility
+		));
 	}
 	// Otherwise we're just going to spawn HISM components instead
 	else
@@ -247,11 +260,18 @@ void FUsdGeomPointInstancerTranslator::CreateAssets()
 			}
 			else
 			{
-				// Fully bake the prototype top level transform and visibility into its own static mesh
-				const bool bIgnoreTopLevelTransformAndVisibility = false;
-				Context->TranslatorTasks.Add(
-					MakeShared<FUsdGeomPointInstancerCreateAssetsTaskChain>(Context, UEPrototypePath, bIgnoreTopLevelTransformAndVisibility, PrimPath)
-				);
+				// Bake the prototype transform into the mesh as it must always be applied before the instancing transform, and if we
+				// put it on the HISM component instead it would have been applied *after* the instancing transform
+				const bool bIgnoreTopLevelTransform = false;
+				// Don't bake the visibility into the mesh as we can control it by toggling HiddenInGame on the HISM component directly
+				const bool bIgnoreTopLevelVisibility = true;
+				Context->TranslatorTasks.Add(MakeShared<FUsdGeomPointInstancerCreateAssetsTaskChain>(
+					Context,
+					UEPrototypePath,
+					bIgnoreTopLevelTransform,
+					bIgnoreTopLevelVisibility,
+					PrimPath
+				));
 			}
 		}
 	}
@@ -278,6 +298,7 @@ USceneComponent* FUsdGeomPointInstancerTranslator::CreateComponents()
 	{
 		bCreateChildHISMs = true;
 		MainSceneComponent = CreateComponentsEx({USceneComponent::StaticClass()}, {});
+		MainSceneComponent->ComponentTags.AddUnique(UnrealIdentifiers::PointInstancerTag);
 	}
 	else
 	{
@@ -386,6 +407,9 @@ void FUsdGeomPointInstancerTranslator::UpdateComponents(USceneComponent* PointIn
 		return;
 	}
 
+	// Do this first as we can use it's resulting HiddenInGame value to propagate visibility to the HISM components
+	Super::UpdateComponents(PointInstancerRootComponent);
+
 	// We always spawn exactly an USceneComponent for the "parent" component of the point instancer, so early out if
 	// we dont' have one. This can happen now if we have an alternative draw mode for this point instancer, at which
 	// point this could be an UUsdDrawModeComponent
@@ -548,6 +572,24 @@ void FUsdGeomPointInstancerTranslator::UpdateComponents(USceneComponent* PointIn
 					Context->bShareAssetsForIdenticalPrims
 				);
 			}
+
+			// Handle visibility for the HISM component
+			//
+			// We will in general no longer put the visibility of the prototype itself baked into the mesh, so we
+			// must add it to the HISM component. These HISM components sort of correspond to both the PointInstancer prim
+			// and a particular prototype at the same time (kinda), so here we'll hide the HISM if either the PointInstancer or the
+			// prototype are hidden, as if they were parent/child prims
+			bool bIsVisible = !PointInstancerRootComponent->bHiddenInGame;
+			if (pxr::UsdGeomImageable PrototypeImageable = pxr::UsdGeomImageable{PrototypeUsdPrim.Get()})
+			{
+				pxr::TfToken Visibility;
+				pxr::UsdAttribute VisibilityAttr = PrototypeImageable.GetVisibilityAttr();
+				if (VisibilityAttr && VisibilityAttr.Get(&Visibility, Context->Time) && Visibility == pxr::UsdGeomTokens->invisible)
+				{
+					bIsVisible = false;
+				}
+			}
+			HISMComponent->SetHiddenInGame(!bIsVisible);
 		}
 
 		// Wait on and assign results of the point instancer.
@@ -557,8 +599,6 @@ void FUsdGeomPointInstancerTranslator::UpdateComponents(USceneComponent* PointIn
 			UsdGeomPointInstancerTranslatorImpl::ApplyPointInstanceTransforms(Result.Key, Result.Value);
 		}
 	}
-
-	Super::UpdateComponents(PointInstancerRootComponent);
 }
 
 bool FUsdGeomPointInstancerTranslator::CollapsesChildren(ECollapsingType CollapsingType) const
diff --git a/Engine/Plugins/Importers/USDImporter/Source/USDSchemas/Private/USDGeomXformableTranslator.cpp b/Engine/Plugins/Importers/USDImporter/Source/USDSchemas/Private/USDGeomXformableTranslator.cpp
--- a/Engine/Plugins/Importers/USDImporter/Source/USDSchemas/Private/USDGeomXformableTranslator.cpp
+++ b/Engine/Plugins/Importers/USDImporter/Source/USDSchemas/Private/USDGeomXformableTranslator.cpp
@@ -252,7 +252,8 @@ void FUsdGeomXformableCreateAssetsTaskChain::SetupTasks()
 		   }
 
 		   // We're going to put Prim's transform and visibility on the component, so we don't need to bake it into the combined mesh
-		   const bool bSkipRootPrimTransformAndVis = true;
+		   const bool bSkipRootPrimTransform = true;
+		   const bool bSkipRootPrimVisibility = true;
 
 		   UsdToUnreal::FUsdMeshConversionOptions Options;
 		   Options.TimeCode = Context->Time;
@@ -262,7 +263,14 @@ void FUsdGeomXformableCreateAssetsTaskChain::SetupTasks()
 		   Options.bMergeIdenticalMaterialSlots = Context->bMergeIdenticalMaterialSlots;
 		   Options.SubdivisionLevel = Context->SubdivisionLevel;
 
-		   UsdToUnreal::ConvertGeomMeshHierarchy(GetPrim(), AddedMeshDescription, AssignmentInfo, Options, bSkipRootPrimTransformAndVis);
+		   UsdToUnreal::ConvertGeomMeshSubtree(
+			   GetPrim(),
+			   AddedMeshDescription,
+			   AssignmentInfo,
+			   Options,
+			   bSkipRootPrimTransform,
+			   bSkipRootPrimVisibility
+		   );
 
 		   return !AddedMeshDescription.IsEmpty();
 	   });
diff --git a/Engine/Plugins/Importers/USDImporter/Source/USDSchemas/Private/USDGeometryCacheTranslator.cpp b/Engine/Plugins/Importers/USDImporter/Source/USDSchemas/Private/USDGeometryCacheTranslator.cpp
--- a/Engine/Plugins/Importers/USDImporter/Source/USDSchemas/Private/USDGeometryCacheTranslator.cpp
+++ b/Engine/Plugins/Importers/USDImporter/Source/USDSchemas/Private/USDGeometryCacheTranslator.cpp
@@ -857,8 +857,16 @@ void FGeometryCacheCreateAssetsTaskChain::SetupTasks()
 
 		   // The collapsed mesh description here will be used to cache the GeometryCache asset, but not to fill it since its content will be
 		   // unflattened Bake the prim's transform into the mesh data
-		   const bool bSkipRootPrimTransformAndVis = false;
-		   UsdToUnreal::ConvertGeomMeshHierarchy(GetPrim(), AddedMeshDescription, AssignmentInfo, Options, bSkipRootPrimTransformAndVis);
+		   bool bSkipRootPrimTransform = false;
+		   bool bSkipRootPrimVisiblity = false;
+		   UsdToUnreal::ConvertGeomMeshSubtree(
+			   GetPrim(),
+			   AddedMeshDescription,
+			   AssignmentInfo,
+			   Options,
+			   bSkipRootPrimTransform,
+			   bSkipRootPrimVisiblity
+		   );
 
 		   // If we have at least one valid LOD, we should proceed to the next step
 		   for (const FMeshDescription& MeshDescription : LODIndexToMeshDescription)
diff --git a/Engine/Plugins/Importers/USDImporter/Source/USDSchemas/Public/USDGeomPointInstancerTranslator.h b/Engine/Plugins/Importers/USDImporter/Source/USDSchemas/Public/USDGeomPointInstancerTranslator.h
--- a/Engine/Plugins/Importers/USDImporter/Source/USDSchemas/Public/USDGeomPointInstancerTranslator.h
+++ b/Engine/Plugins/Importers/USDImporter/Source/USDSchemas/Public/USDGeomPointInstancerTranslator.h
@@ -15,7 +15,8 @@ public:
 	explicit FUsdGeomPointInstancerCreateAssetsTaskChain(
 		const TSharedRef<FUsdSchemaTranslationContext>& InContext,
 		const UE::FSdfPath& InPrimPath,
-		bool bIgnoreTopLevelTransformAndVisibility,
+		bool bIgnoreTopLevelTransform,
+		bool bIgnoreTopLevelVisibility,
 		const TOptional<UE::FSdfPath>& AlternativePrimToLinkAssetsTo = {}
 	);
 
@@ -23,7 +24,8 @@ protected:
 	virtual void SetupTasks() override;
 
 private:
-	bool bIgnoreTopLevelTransformAndVisibility = false;
+	bool bIgnoreTopLevelTransform = false;
+	bool bIgnoreTopLevelVisibility = false;
 };
 
 class USDSCHEMAS_API FUsdGeomPointInstancerTranslator : public FUsdGeomXformableTranslator
diff --git a/Engine/Plugins/Importers/USDImporter/Source/USDStage/Private/USDPrimTwin.cpp b/Engine/Plugins/Importers/USDImporter/Source/USDStage/Private/USDPrimTwin.cpp
--- a/Engine/Plugins/Importers/USDImporter/Source/USDStage/Private/USDPrimTwin.cpp
+++ b/Engine/Plugins/Importers/USDImporter/Source/USDStage/Private/USDPrimTwin.cpp
@@ -5,6 +5,7 @@
 #include "USDMemory.h"
 #include "USDStageActor.h"
 
+#include "Components/InstancedStaticMeshComponent.h"
 #include "Engine/World.h"
 
 UUsdPrimTwin& UUsdPrimTwin::AddChild(const FString& InPrimPath)
@@ -105,6 +106,22 @@ void UUsdPrimTwin::Clear()
 			AttachChild->Modify();
 		}
 
+		// When we translate point instancers we secretly add more HISM components for each prototype.
+		// If we're destroying the scene component for the PointInstancer, we should destroy the HISMs too.
+		if (SceneComponent->ComponentTags.Contains(UnrealIdentifiers::PointInstancerTag))
+		{
+			// Copy the array here as DestroyComponent will change the container directly
+			TArray<TObjectPtr<USceneComponent>> AttachChildren = SceneComponent->GetAttachChildren();
+			for (USceneComponent* AttachChild : AttachChildren)
+			{
+				if (UInstancedStaticMeshComponent* ISM = Cast<UInstancedStaticMeshComponent>(AttachChild))
+				{
+					ISM->Modify();
+					ISM->DestroyComponent();
+				}
+			}
+		}
+
 		SceneComponent->Modify();
 		SceneComponent->DestroyComponent();
 		SceneComponent = nullptr;
diff --git a/Engine/Plugins/Runtime/USDCore/Source/USDUtilities/Private/USDGeomMeshConversion.cpp b/Engine/Plugins/Runtime/USDCore/Source/USDUtilities/Private/USDGeomMeshConversion.cpp
--- a/Engine/Plugins/Runtime/USDCore/Source/USDUtilities/Private/USDGeomMeshConversion.cpp
+++ b/Engine/Plugins/Runtime/USDCore/Source/USDUtilities/Private/USDGeomMeshConversion.cpp
@@ -902,7 +902,8 @@ namespace UE::UsdGeomMeshConversion::Private
 		FMeshDescription& OutMeshDescription,
 		UsdUtils::FUsdPrimMaterialAssignmentInfo& OutMaterialAssignments,
 		UsdToUnreal::FUsdMeshConversionOptions& Options,
-		bool bIsFirstPrim,
+		bool bSkipTransform,
+		bool bSkipVisibility,
 		bool bIsInsideSkelRoot
 	)
 	{
@@ -912,9 +913,7 @@ namespace UE::UsdGeomMeshConversion::Private
 			return true;
 		}
 
-		FTransform ChildTransform = Options.AdditionalTransform;
-
-		if (!bIsFirstPrim)
+		if (!bSkipVisibility)
 		{
 			// Ignore invisible child meshes.
 			//
@@ -946,7 +945,11 @@ namespace UE::UsdGeomMeshConversion::Private
 					}
 				}
 			}
+		}
 
+		FTransform ChildTransform = Options.AdditionalTransform;
+		if (!bSkipTransform)
+		{
 			if (pxr::UsdGeomXformable Xformable = pxr::UsdGeomXformable(Prim))
 			{
 				FTransform LocalChildTransform;
@@ -997,14 +1000,15 @@ namespace UE::UsdGeomMeshConversion::Private
 					break;
 				}
 
-				const bool bChildIsFirstPrim = false;
+				const bool bSkipTransformOrVisibility = false;
 
 				bSuccess &= RecursivelyCollapseChildMeshes(
 					ChildPrim,
 					OutMeshDescription,
 					OutMaterialAssignments,
 					Options,
-					bChildIsFirstPrim,
+					bSkipTransformOrVisibility,
+					bSkipTransformOrVisibility,
 					bIsInsideSkelRoot || ChildPrim.IsA<pxr::UsdSkelRoot>()
 				);
 			}
@@ -2789,11 +2793,12 @@ bool UsdToUnreal::ConvertPointInstancerToMesh(
 			}
 
 			const bool bSkipRootPrimTransformAndVisibility = false;
-			ConvertGeomMeshHierarchy(
+			ConvertGeomMeshSubtree(
 				PrototypeUsdPrim,
 				PrototypeMeshDescriptions[PrototypeIndex],
 				PrototypeMaterialAssignments[PrototypeIndex],
 				OptionsCopy,
+				bSkipRootPrimTransformAndVisibility,
 				bSkipRootPrimTransformAndVisibility
 			);
 		}
@@ -2918,15 +2923,36 @@ bool UsdToUnreal::ConvertPointInstancerToMesh(
 	return true;
 }
 
+PRAGMA_DISABLE_DEPRECATION_WARNINGS
 bool UsdToUnreal::ConvertGeomMeshHierarchy(
+	const pxr::UsdPrim& Prim,
+	FMeshDescription& InOutMeshDescription,
+	UsdUtils::FUsdPrimMaterialAssignmentInfo& InOutMaterialAssignments,
+	const FUsdMeshConversionOptions& CommonOptions,
+	bool bSkipRootPrimTransformAndVisibility
+)
+{
+	return ConvertGeomMeshSubtree(
+		Prim,
+		InOutMeshDescription,
+		InOutMaterialAssignments,
+		CommonOptions,
+		bSkipRootPrimTransformAndVisibility,
+		bSkipRootPrimTransformAndVisibility
+	);
+}
+PRAGMA_ENABLE_DEPRECATION_WARNINGS
+
+bool UsdToUnreal::ConvertGeomMeshSubtree(
 	const pxr::UsdPrim& Prim,
 	FMeshDescription& OutMeshDescription,
 	UsdUtils::FUsdPrimMaterialAssignmentInfo& OutMaterialAssignments,
 	const FUsdMeshConversionOptions& Options,
-	bool bSkipRootPrimTransformAndVisibility
+	bool bSkipRootPrimTransform,
+	bool bSkipRootPrimVisibility
 )
 {
-	TRACE_CPUPROFILER_EVENT_SCOPE(UsdToUnreal::ConvertGeomMeshHierarchy);
+	TRACE_CPUPROFILER_EVENT_SCOPE(UsdToUnreal::ConvertGeomMeshSubtree);
 
 	if (!Prim)
 	{
@@ -2947,7 +2973,7 @@ bool UsdToUnreal::ConvertGeomMeshHierarchy(
 	// Prepass to figure out the best primvars to use for the entire collapsed mesh UV sets
 	if (OutMaterialAssignments.PrimvarToUVIndex.Num() == 0)
 	{
-		OutMaterialAssignments.PrimvarToUVIndex = UsdGeomMeshImpl::CollectSubtreePrimvars(Prim, Options, bSkipRootPrimTransformAndVisibility);
+		OutMaterialAssignments.PrimvarToUVIndex = UsdGeomMeshImpl::CollectSubtreePrimvars(Prim, Options, bSkipRootPrimVisibility);
 	}
 
 	const bool bIsInSkelRoot = static_cast<bool>(UsdUtils::GetClosestParentSkelRoot(Prim));
@@ -2956,7 +2982,8 @@ bool UsdToUnreal::ConvertGeomMeshHierarchy(
 		OutMeshDescription,
 		OutMaterialAssignments,
 		OptionsCopy,
-		bSkipRootPrimTransformAndVisibility,
+		bSkipRootPrimTransform,
+		bSkipRootPrimVisibility,
 		bIsInSkelRoot
 	);
 }
diff --git a/Engine/Plugins/Runtime/USDCore/Source/USDUtilities/Public/USDGeomMeshConversion.h b/Engine/Plugins/Runtime/USDCore/Source/USDUtilities/Public/USDGeomMeshConversion.h
--- a/Engine/Plugins/Runtime/USDCore/Source/USDUtilities/Public/USDGeomMeshConversion.h
+++ b/Engine/Plugins/Runtime/USDCore/Source/USDUtilities/Public/USDGeomMeshConversion.h
@@ -155,8 +155,19 @@ namespace UsdToUnreal
 	 * @param OutMaterialAssignments - Input/Output parameter that will be filled with the material data extracted from
 									   UsdSchema. Each item of the array matches the mesh description of same index
 	 * @param CommonOptions - Shared options used for the conversion
+	 * @param bSkipRootPrimTransform - Whether to bake in the transform of SubtreeRootPrim into the mesh
+	 * @param bSkipRootPrimVisiblity - Whether to bake in the visibility of SubtreeRootPrim into the mesh
 	 * @return Whether the conversion was successful or not.
 	 */
+	USDUTILITIES_API bool ConvertGeomMeshSubtree(
+		const pxr::UsdPrim& SubtreeRootPrim,
+		FMeshDescription& InOutMeshDescription,
+		UsdUtils::FUsdPrimMaterialAssignmentInfo& InOutMaterialAssignments,
+		const FUsdMeshConversionOptions& CommonOptions = FUsdMeshConversionOptions::DefaultOptions,
+		bool bSkipRootPrimTransform = false,
+		bool bSkipRootPrimVisiblity = false
+	);
+	UE_DEPRECATED(5.6, "Use ConvertGeomMeshSubtree instead.")
 	USDUTILITIES_API bool ConvertGeomMeshHierarchy(
 		const pxr::UsdPrim& Prim,
 		FMeshDescription& InOutMeshDescription,
@@ -273,7 +284,11 @@ namespace UnrealToUsd
 	 * Creates a baker responsible for baking the skinned mesh of a skeletal mesh component to a Mesh prim,
 	 * returning true if a EBakingType::Skeletal baker was created
 	 */
-	USDUTILITIES_API bool CreateSkeletalAnimationToMeshBaker(UE::FUsdPrim& MeshPrim, USkeletalMeshComponent& Component, struct FComponentBaker& OutBaker);
+	USDUTILITIES_API bool CreateSkeletalAnimationToMeshBaker(
+		UE::FUsdPrim& MeshPrim,
+		USkeletalMeshComponent& Component,
+		struct FComponentBaker& OutBaker
+	);
 }
 #endif	  // USE_USD_SDK
 
diff --git a/Engine/Plugins/Runtime/USDCore/Source/UnrealUSDWrapper/Private/UnrealUSDWrapper.cpp b/Engine/Plugins/Runtime/USDCore/Source/UnrealUSDWrapper/Private/UnrealUSDWrapper.cpp
--- a/Engine/Plugins/Runtime/USDCore/Source/UnrealUSDWrapper/Private/UnrealUSDWrapper.cpp
+++ b/Engine/Plugins/Runtime/USDCore/Source/UnrealUSDWrapper/Private/UnrealUSDWrapper.cpp
@@ -489,6 +489,8 @@ const FName UnrealIdentifiers::MdlRenderContext = TEXT("mdl");
 
 const FString UnrealIdentifiers::UniversalRenderContextDisplayString = TEXT("universal");
 
+const FName UnrealIdentifiers::PointInstancerTag = TEXT("PointInstancer");
+
 FUsdDelegates::FUsdImportDelegate FUsdDelegates::OnPreUsdImport;
 FUsdDelegates::FUsdImportDelegate FUsdDelegates::OnPostUsdImport;
 
diff --git a/Engine/Plugins/Runtime/USDCore/Source/UnrealUSDWrapper/Public/UnrealUSDWrapper.h b/Engine/Plugins/Runtime/USDCore/Source/UnrealUSDWrapper/Public/UnrealUSDWrapper.h
--- a/Engine/Plugins/Runtime/USDCore/Source/UnrealUSDWrapper/Public/UnrealUSDWrapper.h
+++ b/Engine/Plugins/Runtime/USDCore/Source/UnrealUSDWrapper/Public/UnrealUSDWrapper.h
@@ -562,6 +562,10 @@ namespace UnrealIdentifiers
 
 	extern UNREALUSDWRAPPER_API const FString UniversalRenderContextDisplayString; // "universal"
 
+	// This marks a generated scene component as corresponding to a PointInstancer prim.
+	// It is mainly used to let us know to cleanup any HISM child component whenever we're cleaning up the scene component
+	extern UNREALUSDWRAPPER_API const FName PointInstancerTag;
+
 }	 // namespace UnrealIdentifiers
 
 struct UNREALUSDWRAPPER_API FUsdDelegates
