USD: Fix assets and components not being cleaned up when a prim is removed from the stage.

diff --git a/Engine/Plugins/Importers/USDImporter/Source/USDStage/Private/USDStageActor.cpp b/Engine/Plugins/Importers/USDImporter/Source/USDStage/Private/USDStageActor.cpp
--- a/Engine/Plugins/Importers/USDImporter/Source/USDStage/Private/USDStageActor.cpp
+++ b/Engine/Plugins/Importers/USDImporter/Source/USDStage/Private/USDStageActor.cpp
@@ -1801,12 +1801,6 @@ void AUsdStageActor::HandleAccumulatedNotices()
 			const UE::FSdfPath& PrimPath = PrimChangedInfo.Key;
 			const bool bIsResync = PrimChangedInfo.Value == EPrimUpdateType::Resync;
 
-			// If it's a new prim, we may not have old info about it
-			if (!UsdInfoCache->GetInner().ContainsInfoAboutPrim(PrimPath))
-			{
-				continue;
-			}
-
 			UE_LOG(LogUsd, Verbose, TEXT("Cleaning up assets for prim path '%s'"), *PrimPath.GetString());
 
 			const bool bForEntireSubtree = bIsResync;
@@ -1826,12 +1820,6 @@ void AUsdStageActor::HandleAccumulatedNotices()
 			const UE::FSdfPath PrimPath = PrimChangedInfo.Key;
 			const bool bIsResync = PrimChangedInfo.Value == EPrimUpdateType::Resync;
 
-			// If it's a new prim, we may not have old info about it
-			if (!UsdInfoCache->GetInner().ContainsInfoAboutPrim(PrimPath))
-			{
-				continue;
-			}
-
 			UE_LOG(LogUsd, Verbose, TEXT("Cleaning up actors and components for prim path '%s'"), *PrimPath.GetString());
 
 			if (bIsResync && PrimPath.IsAbsoluteRootOrPrimPath())
