USD: Just directly remove a property when undoing the addition of a property instead of trying to set metadata fields as that can lead to invalid USD.

Remove ChangeBlock from application of collapsing schema as that prevents us from being able to undo adding the attribute.

diff --git a/Engine/Plugins/Importers/USDImporter/Source/USDStage/Private/USDTransactor.cpp b/Engine/Plugins/Importers/USDImporter/Source/USDStage/Private/USDTransactor.cpp
--- a/Engine/Plugins/Importers/USDImporter/Source/USDStage/Private/USDTransactor.cpp
+++ b/Engine/Plugins/Importers/USDImporter/Source/USDStage/Private/USDTransactor.cpp
@@ -552,6 +552,13 @@ namespace UsdUtils
 			Time.IsSet() ? *LexToString(Time.GetValue()) : TEXT("<unset>")
 		);
 
+		// If we just want to remove the property don't really bother doing anything else
+		if (bRemoveProperty)
+		{
+			Prim.RemoveProperty(*PropertyName);
+			return true;
+		}
+
 		if (PropertyName == TEXT("kind"))
 		{
 #if USE_USD_SDK
@@ -578,10 +585,6 @@ namespace UsdUtils
 		{
 			if (WrapperValue.IsEmpty())
 			{
-				if (bRemoveProperty)
-				{
-					Prim.RemoveProperty(*PropertyName);
-				}
 				if (Time.IsSet() && Attribute)
 				{
 					Attribute.ClearAtTime(Time.GetValue());
@@ -601,21 +604,14 @@ namespace UsdUtils
 		{
 			if (WrapperValue.IsEmpty())
 			{
-				if (bRemoveProperty)
+				if (Attribute)
 				{
-					Prim.RemoveProperty(*PropertyName);
+					Attribute.Clear();
 				}
-				else
+				else if (Relationship)
 				{
-					if (Attribute)
-					{
-						Attribute.Clear();
-					}
-					else if (Relationship)
-					{
-						bool bRemoveSpec = false;
-						Relationship.ClearTargets(bRemoveSpec);
-					}
+					bool bRemoveSpec = false;
+					Relationship.ClearTargets(bRemoveSpec);
 				}
 			}
 			else
diff --git a/Engine/Plugins/Runtime/USDCore/Source/USDUtilities/Private/USDConversionUtils.cpp b/Engine/Plugins/Runtime/USDCore/Source/USDUtilities/Private/USDConversionUtils.cpp
--- a/Engine/Plugins/Runtime/USDCore/Source/USDUtilities/Private/USDConversionUtils.cpp
+++ b/Engine/Plugins/Runtime/USDCore/Source/USDUtilities/Private/USDConversionUtils.cpp
@@ -1536,7 +1536,8 @@ bool UsdUtils::SetCollapsingPreference(const pxr::UsdPrim& Prim, UsdUtils::EColl
 
 	FScopedUsdAllocs UsdAllocs;
 
-	pxr::SdfChangeBlock ChangeBlock;
+	// Note: Don't use an pxr::SdfChangeBlock here, as USD needs to emit separate notices for the schema addition and attribute
+	// addition, otherwise it will emit an ObjectsChanged notice that *only* contains the schema application details
 
 	const bool bAppliedSchema = UsdUtils::ApplySchema(Prim, UnrealIdentifiers::UnrealCollapsingAPI);
 	if (!bAppliedSchema)
