USD: Remove unnecessary stopping of VolumeComponents whenever we updated them.

diff --git a/Engine/Plugins/Importers/USDImporter/Source/USDSchemas/Private/USDVolVolumeTranslator.cpp b/Engine/Plugins/Importers/USDImporter/Source/USDSchemas/Private/USDVolVolumeTranslator.cpp
--- a/Engine/Plugins/Importers/USDImporter/Source/USDSchemas/Private/USDVolVolumeTranslator.cpp
+++ b/Engine/Plugins/Importers/USDImporter/Source/USDSchemas/Private/USDVolVolumeTranslator.cpp
@@ -1393,9 +1393,6 @@ void FUsdVolVolumeTranslator::UpdateComponents(USceneComponent* SceneComponent)
 							}
 
 							// Now TargetIndex should be pointing at the index of the desired frame within the SVT
-
-							const bool bPlaying = false;
-							VolumeComponent->SetPlaying(bPlaying);
 							VolumeComponent->SetFrame(static_cast<float>(TargetIndex));
 						}
 					}
