USD: Show prim kind on the properties panel even if the prim doesn't have any opinion authored, so that we can author a new one.

diff --git a/Engine/Plugins/Importers/USDImporter/Source/USDStageEditorViewModels/Private/USDObjectFieldViewModel.cpp b/Engine/Plugins/Importers/USDImporter/Source/USDStageEditorViewModels/Private/USDObjectFieldViewModel.cpp
--- a/Engine/Plugins/Importers/USDImporter/Source/USDStageEditorViewModels/Private/USDObjectFieldViewModel.cpp
+++ b/Engine/Plugins/Importers/USDImporter/Source/USDStageEditorViewModels/Private/USDObjectFieldViewModel.cpp
@@ -4,6 +4,7 @@
 
 #include "UnrealUSDWrapper.h"
 #include "USDAttributeUtils.h"
+#include "USDConversionUtils.h"
 #include "USDErrorUtils.h"
 #include "USDMemory.h"
 #include "USDTypesConversion.h"
@@ -195,7 +196,7 @@ void FUsdObjectFieldsViewModel::SetFieldValue(const FString& FieldName, const Us
 
 			// To set stage metadata the edit target must be the root or session layer
 			pxr::UsdStageRefPtr Stage{UsdStage};
-			pxr::UsdEditContext(Stage, Stage->GetRootLayer());
+			pxr::UsdEditContext EditContext{Stage, Stage->GetRootLayer()};
 			bSuccess = UsdStage.SetMetadata(*FieldName, VtValue);
 		}
 		else if (UE::FUsdPrim UsdPrim = UsdStage.GetPrimAtPath(UE::FSdfPath(*ObjectPath)))
@@ -213,7 +214,15 @@ void FUsdObjectFieldsViewModel::SetFieldValue(const FString& FieldName, const Us
 			// Single value, single component of FString
 			if (FieldName == TEXT("kind") && InValue.Entries.Num() == 1 && InValue.Entries[0].Num() == 1 && InValue.Entries[0][0].IsType<FString>())
 			{
-				bSuccess = IUsdPrim::SetKind(UsdPrim, UnrealToUsd::ConvertToken(*(InValue.Entries[0][0].Get<FString>())).Get());
+				FString Kind = InValue.Entries[0][0].Get<FString>();
+				if (Kind.IsEmpty())
+				{
+					bSuccess = IUsdPrim::ClearKind(UsdPrim);
+				}
+				else
+				{
+					bSuccess = IUsdPrim::SetKind(UsdPrim, UnrealToUsd::ConvertToken(*Kind).Get());
+				}
 			}
 			else if (UE::FUsdAttribute Attribute = UsdPrim.GetAttribute(*FieldName))
 			{
@@ -422,6 +431,17 @@ void FUsdObjectFieldsViewModel::Refresh(const UE::FUsdStageWeak& InUsdStage, con
 			);
 			CreateField(EObjectFieldType::Metadata, TEXT("path"), ObjectPath, UsdUtils::EUsdBasicDataTypes::String, Role, bPrimReadOnly);
 
+			// If the prim doesn't have any "kind" authored show a combobox anyway so that we can use this to author new kinds
+			if (!UsdUtils::HasAuthoredKind(UsdPrim))
+			{
+				CreateField(
+					EObjectFieldType::Metadata,
+					TEXT("kind"),
+					UsdToUnreal::ConvertToken(IUsdPrim::GetKind(UsdPrim)),
+					UsdUtils::EUsdBasicDataTypes::Token
+				);
+			}
+
 			FScopedUsdAllocs UsdAllocs;
 
 			pxr::UsdPrim PxrUsdPrim{UsdPrim};
diff --git a/Engine/Plugins/Runtime/USDCore/Source/USDUtilities/Private/USDConversionUtils.cpp b/Engine/Plugins/Runtime/USDCore/Source/USDUtilities/Private/USDConversionUtils.cpp
--- a/Engine/Plugins/Runtime/USDCore/Source/USDUtilities/Private/USDConversionUtils.cpp
+++ b/Engine/Plugins/Runtime/USDCore/Source/USDUtilities/Private/USDConversionUtils.cpp
@@ -1400,6 +1400,15 @@ bool UsdUtils::GetAnimatedBoundsTimeSamples(
 	return bHasAnimatedBounds;
 }
 
+bool UsdUtils::HasAuthoredKind(const pxr::UsdPrim& Prim)
+{
+	FScopedUsdAllocs Allocs;
+
+	pxr::UsdModelAPI Model{Prim};
+	pxr::TfToken KindToken;
+	return Model && Model.GetKind(&KindToken);
+}
+
 EUsdDefaultKind UsdUtils::GetDefaultKind(const pxr::UsdPrim& Prim)
 {
 	FScopedUsdAllocs Allocs;
diff --git a/Engine/Plugins/Runtime/USDCore/Source/USDUtilities/Public/USDConversionUtils.h b/Engine/Plugins/Runtime/USDCore/Source/USDUtilities/Public/USDConversionUtils.h
--- a/Engine/Plugins/Runtime/USDCore/Source/USDUtilities/Public/USDConversionUtils.h
+++ b/Engine/Plugins/Runtime/USDCore/Source/USDUtilities/Public/USDConversionUtils.h
@@ -286,6 +286,9 @@ namespace UsdUtils
 		bool bInIgnoreVisibility = false
 	);
 
+	/** Returns true if this prim has any authored opinion for prim kind */
+	USDUTILITIES_API bool HasAuthoredKind(const pxr::UsdPrim& Prim);
+
 	/** Returns whether Prim belongs to any of the default kinds, or a kind derived from them. The result can be a union of different kinds. */
 	USDUTILITIES_API EUsdDefaultKind GetDefaultKind(const pxr::UsdPrim& Prim);
 
