USD: Fix weird behavior when playing LevelSequences with spawnable stage actors during PIE.

diff --git a/Engine/Plugins/Importers/USDImporter/Source/USDStage/Private/USDStageActor.cpp b/Engine/Plugins/Importers/USDImporter/Source/USDStage/Private/USDStageActor.cpp
--- a/Engine/Plugins/Importers/USDImporter/Source/USDStage/Private/USDStageActor.cpp
+++ b/Engine/Plugins/Importers/USDImporter/Source/USDStage/Private/USDStageActor.cpp
@@ -3739,8 +3739,7 @@ void AUsdStageActor::PostEditChangeProperty(FPropertyChangedEvent& PropertyChang
 	// PostRegister/UnregisterAllComponents which would have sidestepped our checks in SetRootLayer.
 	// Note that any property change event would also end up calling our intended path via OnObjectPropertyChanged, this just prevents us from loading
 	// the same stage again if we don't need to.
-
-	bIsModifyingAProperty = true;
+	TGuardValue<bool> ModifyingPropertyGuard{bIsModifyingAProperty, true};
 	Super::PostEditChangeProperty(PropertyChangedEvent);
 }
 
@@ -3865,6 +3864,7 @@ void AUsdStageActor::PreEditChange(FProperty* PropertyThatWillChange)
 	// PostRegister/Unregister calls in the editor due to AActor::PostEditChangeProperty *and* AActor::PreEditChange.
 	// Here we determine in which cases we should ignore those PostRegister/Unregister calls by using the
 	// bIsModifyingAProperty flag
+	TOptional<TGuardValue<bool>> ModifyingPropertyGuard;
 	if (!IsActorBeingDestroyed())
 	{
 		if ((GEditor && GEditor->bIsSimulatingInEditor && GetWorld() != nullptr) || ReregisterComponentsWhenModified())
@@ -3874,7 +3874,7 @@ void AUsdStageActor::PreEditChange(FProperty* PropertyThatWillChange)
 			// functions. We only care about blocking the calls triggered by AActor::PostEditChangeProperty and AActor::PreEditChange
 			if (PropertyThatWillChange)
 			{
-				bIsModifyingAProperty = true;
+				ModifyingPropertyGuard.Emplace(bIsModifyingAProperty, true);
 			}
 		}
 	}
@@ -4873,8 +4873,6 @@ void AUsdStageActor::HandlePropertyChangedEvent(FPropertyChangedEvent& PropertyC
 
 		SetUsdAssetCache(CorrectCache);
 	}
-
-	bIsModifyingAProperty = false;
 }
 
 bool AUsdStageActor::HasAuthorityOverStage() const
