USD: Fix spawnable stage actors not animating in PIE or showing up in MRQ;

Fix spawnable, loaded stage actors losing their spawns when other stage actors are destroyed.

diff --git a/Engine/Plugins/Importers/USDImporter/Source/USDStage/Private/USDLevelSequenceHelper.cpp b/Engine/Plugins/Importers/USDImporter/Source/USDStage/Private/USDLevelSequenceHelper.cpp
--- a/Engine/Plugins/Importers/USDImporter/Source/USDStage/Private/USDLevelSequenceHelper.cpp
+++ b/Engine/Plugins/Importers/USDImporter/Source/USDStage/Private/USDLevelSequenceHelper.cpp
@@ -3907,6 +3907,12 @@ void FUsdLevelSequenceHelperImpl::OnObjectTransacted(UObject* Object, const clas
 		return;
 	}
 
+	// Never write back to the stage if we don't have authority
+	if (StageActor.IsValid() && !StageActor->HasAuthorityOverStage())
+	{
+		return;
+	}
+
 	if (UMovieScene* MovieScene = Cast<UMovieScene>(Object))
 	{
 		HandleMovieSceneChange(*MovieScene);
diff --git a/Engine/Plugins/Importers/USDImporter/Source/USDStage/Private/USDStageActor.cpp b/Engine/Plugins/Importers/USDImporter/Source/USDStage/Private/USDStageActor.cpp
--- a/Engine/Plugins/Importers/USDImporter/Source/USDStage/Private/USDStageActor.cpp
+++ b/Engine/Plugins/Importers/USDImporter/Source/USDStage/Private/USDStageActor.cpp
@@ -963,7 +963,7 @@ AUsdStageActor::AUsdStageActor()
 	// draw mode components we end up spawning
 	SetupBBoxCacheIfNeeded();
 
-	if (HasAuthorityOverStage())
+	if (!IsTemplate())
 	{
 #if WITH_EDITOR
 		// Update the supported filetypes in our RootPath property
@@ -4146,6 +4146,27 @@ void AUsdStageActor::BeginDestroy()
 	Super::BeginDestroy();
 }
 
+void AUsdStageActor::PostInitProperties()
+{
+	Super::PostInitProperties();
+
+	// Make sure we never keep a prim twin that was provided by our post-constructor initialization, as RootUsdTwin
+	// is not meant to be shared across actors.
+	//
+	// During OnBeginPIE/OnPostPIEStarted we will temporarily make our prim twins into fake subobjects,
+	// so that the PIE process duplicates all our spawns automatically. The annoying part is that if the stage actor
+	// is a spawnable, this actor with prim twin subobjects will be recorded as the "object archetype" for the
+	// spawnable in PIE, meaning that its RootUsdTwin property pointing at that one same prim twin will be copied over
+	// for every subsequent spawn of that spawnable within FObjectInitializer::InitProperties...
+	// The final outcome is that all the stage actor spawns will share the root prim twin, and when one of them clears it,
+	// all other stage actors lose their actors/components.
+	//
+	// Obviously we don't want any of this, so here we make sure that we always ignore those RootUsdTwin values.
+	// Ideally we'd make our RootUsdTwin into a default subobject so all this naturally goes away, but that will
+	// break the blueprint recompilation tricks so it's probably best left alone for now
+	RootUsdTwin = nullptr;
+}
+
 void AUsdStageActor::PostRegisterAllComponents()
 {
 	Super::PostRegisterAllComponents();
diff --git a/Engine/Plugins/Importers/USDImporter/Source/USDStage/Public/USDStageActor.h b/Engine/Plugins/Importers/USDImporter/Source/USDStage/Public/USDStageActor.h
--- a/Engine/Plugins/Importers/USDImporter/Source/USDStage/Public/USDStageActor.h
+++ b/Engine/Plugins/Importers/USDImporter/Source/USDStage/Public/USDStageActor.h
@@ -376,6 +376,7 @@ public:
 	USDSTAGE_API virtual void PostActorCreated() override;
 	USDSTAGE_API virtual void PostRename(UObject* OldOuter, const FName OldName) override;
 	USDSTAGE_API virtual void BeginDestroy() override;
+	USDSTAGE_API virtual void PostInitProperties() override;
 #if WITH_EDITOR
 	USDSTAGE_API virtual void PostEditChangeProperty(FPropertyChangedEvent& PropertyChangedEvent) override;
 	USDSTAGE_API virtual void PostTransacted(const FTransactionObjectEvent& TransactionEvent) override;
