USD: Show the root layer as dirty by default for new stages.

diff --git a/Engine/Plugins/Importers/USDImporter/Source/USDStageEditorViewModels/Private/USDLayersViewModel.cpp b/Engine/Plugins/Importers/USDImporter/Source/USDStageEditorViewModels/Private/USDLayersViewModel.cpp
--- a/Engine/Plugins/Importers/USDImporter/Source/USDStageEditorViewModels/Private/USDLayersViewModel.cpp
+++ b/Engine/Plugins/Importers/USDImporter/Source/USDStageEditorViewModels/Private/USDLayersViewModel.cpp
@@ -141,8 +141,9 @@ void FUsdLayerViewModel::RefreshData()
 #if USE_USD_SDK
 	FScopedUsdAllocs UsdAllocs;
 
-	pxr::UsdStageRefPtr UsdStageRef(UsdStage);
-	pxr::UsdStageRefPtr IsolatedStageRef(IsolatedStage);
+	pxr::UsdStageRefPtr UsdStageRef{UsdStage};
+	pxr::UsdStageRefPtr IsolatedStageRef{IsolatedStage};
+	pxr::UsdStageRefPtr ActiveStage = IsolatedStage ? IsolatedStage : UsdStageRef;
 
 	const std::string UsdLayerIdentifier = UnrealToUsd::ConvertString(*LayerIdentifier).Get();
 
@@ -151,7 +152,9 @@ void FUsdLayerViewModel::RefreshData()
 	pxr::SdfLayerRefPtr UsdLayer = UE::FSdfLayer::FindOrOpen(*LayerIdentifier);
 	if (UsdLayer)
 	{
-		LayerModel->bIsDirty = UsdLayer->IsDirty() && !UsdLayer->IsAnonymous();
+		// Don't show session layers (always anonymous) as dirty, but if the root layer of the stage itself is anonymous
+		// it may be a File->New stage, so we want to show it as dirty to indicate it hasn't been saved yet
+		LayerModel->bIsDirty = UsdLayer->IsDirty() && (!UsdLayer->IsAnonymous() || ActiveStage->GetRootLayer() == UsdLayer);
 	}
 
 	const bool bIsIsolatedRootLayer = IsolatedStageRef && IsolatedStageRef->GetRootLayer() == UsdLayer;
