USD: Export level by default when exporting level sequences, defaulting to the current level.

diff --git a/Engine/Plugins/Importers/USDImporter/Source/USDExporter/Private/LevelExporterUSDOptionsCustomization.cpp b/Engine/Plugins/Importers/USDImporter/Source/USDExporter/Private/LevelExporterUSDOptionsCustomization.cpp
--- a/Engine/Plugins/Importers/USDImporter/Source/USDExporter/Private/LevelExporterUSDOptionsCustomization.cpp
+++ b/Engine/Plugins/Importers/USDImporter/Source/USDExporter/Private/LevelExporterUSDOptionsCustomization.cpp
@@ -198,14 +198,6 @@ void FLevelExporterUSDOptionsCustomization::CustomizeDetails(IDetailLayoutBuilde
 	{
 		OptionsPtr.Reset(LevelSequenceOptions);
 
-		// For now there is no easy way of fetching the level to export from a ULevelSequence... we could potentially try to guess what it is
-		// by looking at the soft object paths, but even those aren't exposed, so here we just default to using the current level as the export level.
-		if (LevelSequenceOptions->Level == nullptr)
-		{
-			const bool bEditorWorldsOnly = true;
-			LevelSequenceOptions->Level = IUsdClassesModule::GetCurrentWorld(bEditorWorldsOnly);
-		}
-
 		PickerTree = SNew(LevelExporterUSDImpl::SLevelPickerList, &LevelSequenceOptions->LevelExportOptions, LevelSequenceOptions->Level->GetWorld());
 		LevelFilterPropName = TEXT("LevelExportOptions.LevelsToIgnore");
 		ExportSublayersPropName = TEXT("LevelExportOptions.bExportSublayers");
diff --git a/Engine/Plugins/Importers/USDImporter/Source/USDExporter/Private/LevelSequenceExporterUSD.cpp b/Engine/Plugins/Importers/USDImporter/Source/USDExporter/Private/LevelSequenceExporterUSD.cpp
--- a/Engine/Plugins/Importers/USDImporter/Source/USDExporter/Private/LevelSequenceExporterUSD.cpp
+++ b/Engine/Plugins/Importers/USDImporter/Source/USDExporter/Private/LevelSequenceExporterUSD.cpp
@@ -1800,6 +1800,12 @@ bool ULevelSequenceExporterUsd::ExportBinary(
 	{
 		Options = GetMutableDefault<ULevelSequenceExporterUsdOptions>();
 
+		// Prefill the level to export with the current level
+		if (!Options->Level.Get())
+		{
+			Options->Level = IUsdClassesModule::GetCurrentWorld();
+		}
+
 		// Prompt with an options dialog if we can
 		if (Options && (!ExportTask || !ExportTask->bAutomated))
 		{
diff --git a/Engine/Plugins/Importers/USDImporter/Source/USDExporter/Public/LevelSequenceExporterUSDOptions.h b/Engine/Plugins/Importers/USDImporter/Source/USDExporter/Public/LevelSequenceExporterUSDOptions.h
--- a/Engine/Plugins/Importers/USDImporter/Source/USDExporter/Public/LevelSequenceExporterUSDOptions.h
+++ b/Engine/Plugins/Importers/USDImporter/Source/USDExporter/Public/LevelSequenceExporterUSDOptions.h
@@ -63,7 +63,7 @@ public:
 	 * Exports the provided level alongside the layer that represents the exported Level Sequence
 	 */
 	UPROPERTY(EditAnywhere, config, BlueprintReadWrite, Category = "Level Sequence")
-	bool bExportLevel = false;
+	bool bExportLevel = true;
 
 	/**
 	 * If checked this will also add the exported level as a sublayer to the USD files emitted for all exported level sequences
