USD: Fix crash when trying to reuse a LevelSequence that was destroyed in PIE.

diff --git a/Engine/Plugins/Importers/USDImporter/Source/USDStage/Private/USDLevelSequenceHelper.cpp b/Engine/Plugins/Importers/USDImporter/Source/USDStage/Private/USDLevelSequenceHelper.cpp
--- a/Engine/Plugins/Importers/USDImporter/Source/USDStage/Private/USDLevelSequenceHelper.cpp
+++ b/Engine/Plugins/Importers/USDImporter/Source/USDStage/Private/USDLevelSequenceHelper.cpp
@@ -1149,7 +1149,6 @@ bool FUsdLevelSequenceHelperImpl::HasData() const
 
 void FUsdLevelSequenceHelperImpl::Clear()
 {
-	// Mark old sequences as garbage so that the track references' to the assets can be immediately ignored
 	for (const TPair<FString, TObjectPtr<ULevelSequence>>& IdentifierAndSeq : LevelSequencesByIdentifier)
 	{
 		if (ULevelSequence* Seq = IdentifierAndSeq.Value.Get())
@@ -1160,8 +1159,6 @@ void FUsdLevelSequenceHelperImpl::Clear()
 				GEditor->GetEditorSubsystem<UAssetEditorSubsystem>()->CloseAllEditorsForAsset(Seq);
 			}
 #endif	  // WITH_EDITOR
-
-			Seq->MarkAsGarbage();
 		}
 	}
 
diff --git a/Engine/Plugins/Importers/USDImporter/Source/USDStage/Private/USDStageActor.cpp b/Engine/Plugins/Importers/USDImporter/Source/USDStage/Private/USDStageActor.cpp
--- a/Engine/Plugins/Importers/USDImporter/Source/USDStage/Private/USDStageActor.cpp
+++ b/Engine/Plugins/Importers/USDImporter/Source/USDStage/Private/USDStageActor.cpp
@@ -2075,7 +2075,7 @@ UUsdPrimTwin* AUsdStageActor::ExpandPrim(
 	{
 		// For the repopulate, let's only visit the prim twins that already have children and so may actually have
 		// components. We don't want to create brand new components here
-		bExpandChildren = bResync && UsdPrimTwin->GetChildren().Num() > 0;
+		bExpandChildren = UsdPrimTwin->GetChildren().Num() > 0;
 	}
 	else if (!TranslationContext.bIsJustRepopulatingLevelSequence)
 	{
@@ -3503,12 +3503,6 @@ void AUsdStageActor::UnloadUsdStage()
 					if (ULevelSequence* ValidSequence = LevelSequencePtr.Get())
 					{
 						GEditor->GetEditorSubsystem<UAssetEditorSubsystem>()->CloseAllEditorsForAsset(ValidSequence);
-
-						// It's important to mark as garbage so that DeleteUnreferencedAssets can delete assets
-						// that are referenced by LevelSequence tracks, like AnimSequences and geometry caches.
-						// DeleteUnreferencedAssets will run on the *next* tick, so doing this on the AsyncTask is still
-						// going to take place before it
-						ValidSequence->MarkAsGarbage();
 					}
 				}
 			);
@@ -4031,6 +4025,12 @@ void AUsdStageActor::PostDuplicate(bool bDuplicateForPIE)
 	if (bDuplicateForPIE)
 	{
 		OpenUsdStage();
+
+		// We always want our own LevelSequence though, otherwise we could end up with some strange behavior, like:
+		// PIE -> Change LevelSequence -> PIE actor is inside of a FScopedBlockNoticeListening from HandleTrackChange,
+		// but the actor back in the editor is not -> Editor actor writes changes back to the stage anyway
+		RegenerateLevelSequence();
+		RepopulateLevelSequence();
 	}
 	else
 	{
