USD: Fix edit in instance proxy warning still showing when setting material overrides in instanceables.

diff --git a/Engine/Plugins/Importers/USDImporter/Source/USDStage/Private/USDLevelSequenceHelper.cpp b/Engine/Plugins/Importers/USDImporter/Source/USDStage/Private/USDLevelSequenceHelper.cpp
--- a/Engine/Plugins/Importers/USDImporter/Source/USDStage/Private/USDLevelSequenceHelper.cpp
+++ b/Engine/Plugins/Importers/USDImporter/Source/USDStage/Private/USDLevelSequenceHelper.cpp
@@ -4647,6 +4647,12 @@ void FUsdLevelSequenceHelperImpl::HandleTrackChange(UMovieSceneTrack& Track, boo
 		ensure(!FoundExistingGuid || *FoundExistingGuid == PossessableGuid);
 		Bindings.ObjectClassToBindingGuid.Emplace(ComponentClass, PossessableGuid);
 
+		// We can't do anything if our prim is an instance proxy
+		if (UsdUtils::NotifyIfInstanceProxy(UsdPrim))
+		{
+			return;
+		}
+
 		if (bIsMuteChange)
 		{
 			if (PropertyTrack)
diff --git a/Engine/Plugins/Importers/USDImporter/Source/USDStage/Private/USDStageActor.cpp b/Engine/Plugins/Importers/USDImporter/Source/USDStage/Private/USDStageActor.cpp
--- a/Engine/Plugins/Importers/USDImporter/Source/USDStage/Private/USDStageActor.cpp
+++ b/Engine/Plugins/Importers/USDImporter/Source/USDStage/Private/USDStageActor.cpp
@@ -7,6 +7,7 @@
 #include "USDAssetCache2.h"
 #include "USDAssetCache3.h"
 #include "USDAssetUserData.h"
+#include "USDAttributeUtils.h"
 #include "USDClassesModule.h"
 #include "USDConversionUtils.h"
 #include "USDDrawModeComponent.h"
@@ -4618,88 +4619,100 @@ void AUsdStageActor::OnObjectPropertyChanged(UObject* ObjectBeingModified, FProp
 
 #if USE_USD_SDK
 
-		UnrealToUsd::ConvertLiveLinkProperties(Controller ? Cast<UActorComponent>(Controller) : PrimSceneComponent, UsdPrim);
-
-		UnrealToUsd::ConvertSceneComponent(CurrentStage, PrimSceneComponent, UsdPrim);
-
-		// When we parse a Gprim like a Cube or a Cylinder, we'll always generate some "default" meshes (e.g. Cylinder with
-		// height always equal 1), and combine the Xform and the effect of the prim's attributes (e.g. height/width) into
-		// a SINGLE transform, and put that on the component (this approach allows attribute animation purely with Sequencer tracks).
-		// When we modify any property and want to write back out to USD however, we'll write that combined transform as the prim's
-		// transform. This means we must also "reset" the (e.g. height/width) attributes, so that the combined transform stays
-		// consistent
-		const bool bDefaultValues = true;
-		const bool bTimeSampleValues = false;
-		UsdUtils::AuthorIdentityTransformGprimAttributes(UsdPrim, bDefaultValues, bTimeSampleValues);
-
+		// We can author material overrides even in instance proxies now, so we should always be able to do this
 		if (UMeshComponent* MeshComponent = Cast<UMeshComponent>(PrimSceneComponent))
 		{
 			UnrealToUsd::ConvertMeshComponent(CurrentStage, MeshComponent, UsdPrim);
 		}
-		else if (UUsdDrawModeComponent* DrawModeComponent = Cast<UUsdDrawModeComponent>(PrimSceneComponent))
-		{
-			const static TSet<FName> BoundsProperties = {
-				GET_MEMBER_NAME_CHECKED(UUsdDrawModeComponent, BoundsMin),
-				GET_MEMBER_NAME_CHECKED(UUsdDrawModeComponent, BoundsMax),
-			};
 
-			// If we just manually tweaked the extents, also author those back out to USD as extents opinions
-			const bool bWriteExtents = BoundsProperties.Contains(PropertyChangedEvent.GetMemberPropertyName());
-			const double UsdTimeCode = UsdUtils::GetDefaultTimeCode();
-			UnrealToUsd::ConvertDrawModeComponent(*DrawModeComponent, UsdPrim, bWriteExtents, UsdTimeCode);
-		}
-		else if (UsdPrim && UsdPrim.IsA(TEXT("Camera")))
+		if (UsdPrim.IsInstanceProxy())
 		{
-			// Our component may be pointing directly at a camera component in case we recreated an exported
-			// ACineCameraActor (see UE-120826)
-			if (UCineCameraComponent* RecreatedCameraComponent = Cast<UCineCameraComponent>(PrimSceneComponent))
-			{
-				UnrealToUsd::ConvertCameraComponent(*RecreatedCameraComponent, UsdPrim);
-			}
-			// Or it could have been just a generic Camera prim, at which case we'll have spawned an entire new
-			// ACineCameraActor for it. In this scenario our prim twin is pointing at the root component, so we need
-			// to dig to the actual UCineCameraComponent to write out the camera data.
-			// We should only do this when the Prim actually corresponds to the Camera though, or else we'll also catch
-			// the prim/component pair that corresponds to the root scene component in case we recreated an exported
-			// ACineCameraActor.
-			else if (ACineCameraActor* CameraActor = Cast<ACineCameraActor>(PrimSceneComponent->GetOwner()))
+			if (PropertyChangedEvent.GetPropertyName() != GET_MEMBER_NAME_CHECKED(UMeshComponent, OverrideMaterials))
 			{
-				if (UCineCameraComponent* CameraComponent = CameraActor->GetCineCameraComponent())
-				{
-					UnrealToUsd::ConvertCameraComponent(*CameraComponent, UsdPrim);
-				}
+				UsdUtils::NotifyIfInstanceProxy(UsdPrim);
 			}
 		}
-		else if (ALight* LightActor = Cast<ALight>(PrimSceneComponent->GetOwner()))
+		else
 		{
-			if (ULightComponent* LightComponent = LightActor->GetLightComponent())
-			{
-				UnrealToUsd::ConvertLightComponent(*LightComponent, UsdPrim, UsdUtils::GetDefaultTimeCode());
+			UnrealToUsd::ConvertLiveLinkProperties(Controller ? Cast<UActorComponent>(Controller) : PrimSceneComponent, UsdPrim);
+
+			UnrealToUsd::ConvertSceneComponent(CurrentStage, PrimSceneComponent, UsdPrim);
 
-				if (UDirectionalLightComponent* DirectionalLight = Cast<UDirectionalLightComponent>(LightComponent))
+			// When we parse a Gprim like a Cube or a Cylinder, we'll always generate some "default" meshes (e.g. Cylinder with
+			// height always equal 1), and combine the Xform and the effect of the prim's attributes (e.g. height/width) into
+			// a SINGLE transform, and put that on the component (this approach allows attribute animation purely with Sequencer tracks).
+			// When we modify any property and want to write back out to USD however, we'll write that combined transform as the prim's
+			// transform. This means we must also "reset" the (e.g. height/width) attributes, so that the combined transform stays
+			// consistent
+			const bool bDefaultValues = true;
+			const bool bTimeSampleValues = false;
+			UsdUtils::AuthorIdentityTransformGprimAttributes(UsdPrim, bDefaultValues, bTimeSampleValues);
+
+			if (UUsdDrawModeComponent* DrawModeComponent = Cast<UUsdDrawModeComponent>(PrimSceneComponent))
+			{
+				const static TSet<FName> BoundsProperties = {
+					GET_MEMBER_NAME_CHECKED(UUsdDrawModeComponent, BoundsMin),
+					GET_MEMBER_NAME_CHECKED(UUsdDrawModeComponent, BoundsMax),
+				};
+
+				// If we just manually tweaked the extents, also author those back out to USD as extents opinions
+				const bool bWriteExtents = BoundsProperties.Contains(PropertyChangedEvent.GetMemberPropertyName());
+				const double UsdTimeCode = UsdUtils::GetDefaultTimeCode();
+				UnrealToUsd::ConvertDrawModeComponent(*DrawModeComponent, UsdPrim, bWriteExtents, UsdTimeCode);
+			}
+			else if (UsdPrim && UsdPrim.IsA(TEXT("Camera")))
+			{
+				// Our component may be pointing directly at a camera component in case we recreated an exported
+				// ACineCameraActor (see UE-120826)
+				if (UCineCameraComponent* RecreatedCameraComponent = Cast<UCineCameraComponent>(PrimSceneComponent))
 				{
-					UnrealToUsd::ConvertDirectionalLightComponent(*DirectionalLight, UsdPrim, UsdUtils::GetDefaultTimeCode());
+					UnrealToUsd::ConvertCameraComponent(*RecreatedCameraComponent, UsdPrim);
 				}
-				else if (URectLightComponent* RectLight = Cast<URectLightComponent>(LightComponent))
+				// Or it could have been just a generic Camera prim, at which case we'll have spawned an entire new
+				// ACineCameraActor for it. In this scenario our prim twin is pointing at the root component, so we need
+				// to dig to the actual UCineCameraComponent to write out the camera data.
+				// We should only do this when the Prim actually corresponds to the Camera though, or else we'll also catch
+				// the prim/component pair that corresponds to the root scene component in case we recreated an exported
+				// ACineCameraActor.
+				else if (ACineCameraActor* CameraActor = Cast<ACineCameraActor>(PrimSceneComponent->GetOwner()))
 				{
-					UnrealToUsd::ConvertRectLightComponent(*RectLight, UsdPrim, UsdUtils::GetDefaultTimeCode());
+					if (UCineCameraComponent* CameraComponent = CameraActor->GetCineCameraComponent())
+					{
+						UnrealToUsd::ConvertCameraComponent(*CameraComponent, UsdPrim);
+					}
 				}
-				else if (UPointLightComponent* PointLight = Cast<UPointLightComponent>(LightComponent))
+			}
+			else if (ALight* LightActor = Cast<ALight>(PrimSceneComponent->GetOwner()))
+			{
+				if (ULightComponent* LightComponent = LightActor->GetLightComponent())
 				{
-					UnrealToUsd::ConvertPointLightComponent(*PointLight, UsdPrim, UsdUtils::GetDefaultTimeCode());
+					UnrealToUsd::ConvertLightComponent(*LightComponent, UsdPrim, UsdUtils::GetDefaultTimeCode());
 
-					if (USpotLightComponent* SpotLight = Cast<USpotLightComponent>(LightComponent))
+					if (UDirectionalLightComponent* DirectionalLight = Cast<UDirectionalLightComponent>(LightComponent))
 					{
-						UnrealToUsd::ConvertSpotLightComponent(*SpotLight, UsdPrim, UsdUtils::GetDefaultTimeCode());
+						UnrealToUsd::ConvertDirectionalLightComponent(*DirectionalLight, UsdPrim, UsdUtils::GetDefaultTimeCode());
+					}
+					else if (URectLightComponent* RectLight = Cast<URectLightComponent>(LightComponent))
+					{
+						UnrealToUsd::ConvertRectLightComponent(*RectLight, UsdPrim, UsdUtils::GetDefaultTimeCode());
+					}
+					else if (UPointLightComponent* PointLight = Cast<UPointLightComponent>(LightComponent))
+					{
+						UnrealToUsd::ConvertPointLightComponent(*PointLight, UsdPrim, UsdUtils::GetDefaultTimeCode());
+
+						if (USpotLightComponent* SpotLight = Cast<USpotLightComponent>(LightComponent))
+						{
+							UnrealToUsd::ConvertSpotLightComponent(*SpotLight, UsdPrim, UsdUtils::GetDefaultTimeCode());
+						}
 					}
 				}
 			}
-		}
-		// In contrast to the other light types, the USkyLightComponent is the root component of the ASkyLight
-		else if (USkyLightComponent* SkyLightComponent = Cast<USkyLightComponent>(PrimSceneComponent))
-		{
-			UnrealToUsd::ConvertLightComponent(*SkyLightComponent, UsdPrim, UsdUtils::GetDefaultTimeCode());
-			UnrealToUsd::ConvertSkyLightComponent(*SkyLightComponent, UsdPrim, UsdUtils::GetDefaultTimeCode());
+			// In contrast to the other light types, the USkyLightComponent is the root component of the ASkyLight
+			else if (USkyLightComponent* SkyLightComponent = Cast<USkyLightComponent>(PrimSceneComponent))
+			{
+				UnrealToUsd::ConvertLightComponent(*SkyLightComponent, UsdPrim, UsdUtils::GetDefaultTimeCode());
+				UnrealToUsd::ConvertSkyLightComponent(*SkyLightComponent, UsdPrim, UsdUtils::GetDefaultTimeCode());
+			}
 		}
 #endif	  // #if USE_USD_SDK
 
diff --git a/Engine/Plugins/Runtime/USDCore/Source/USDUtilities/Private/USDIntegrationUtils.cpp b/Engine/Plugins/Runtime/USDCore/Source/USDUtilities/Private/USDIntegrationUtils.cpp
--- a/Engine/Plugins/Runtime/USDCore/Source/USDUtilities/Private/USDIntegrationUtils.cpp
+++ b/Engine/Plugins/Runtime/USDCore/Source/USDUtilities/Private/USDIntegrationUtils.cpp
@@ -171,11 +171,6 @@ void UnrealToUsd::ConvertLiveLinkProperties(const UActorComponent* InComponent,
 
 	UE::FSdfChangeBlock ChangeBlock;
 
-	if (UsdUtils::NotifyIfInstanceProxy(InOutPrim))
-	{
-		return;
-	}
-
 	// Skeletal LiveLink case
 	if (const USkeletalMeshComponent* SkeletalComponent = Cast<USkeletalMeshComponent>(InComponent))
 	{
diff --git a/Engine/Plugins/Runtime/USDCore/Source/USDUtilities/Private/USDLightConversion.cpp b/Engine/Plugins/Runtime/USDCore/Source/USDUtilities/Private/USDLightConversion.cpp
--- a/Engine/Plugins/Runtime/USDCore/Source/USDUtilities/Private/USDLightConversion.cpp
+++ b/Engine/Plugins/Runtime/USDCore/Source/USDUtilities/Private/USDLightConversion.cpp
@@ -408,11 +408,6 @@ bool UnrealToUsd::ConvertLightComponent(const ULightComponentBase& LightComponen
 		return false;
 	}
 
-	if (UsdUtils::NotifyIfInstanceProxy(Prim))
-	{
-		return false;
-	}
-
 	if (pxr::UsdAttribute Attr = LightAPI.CreateIntensityAttr())
 	{
 		Attr.Set<float>(LightComponent.Intensity, UsdTimeCode);
@@ -493,11 +488,6 @@ bool UnrealToUsd::ConvertDirectionalLightComponent(const UDirectionalLightCompon
 		return false;
 	}
 
-	if (UsdUtils::NotifyIfInstanceProxy(Prim))
-	{
-		return false;
-	}
-
 	if (pxr::UsdAttribute Attr = Light.CreateAngleAttr())
 	{
 		Attr.Set<float>(LightComponent.LightSourceAngle, UsdTimeCode);
@@ -525,11 +515,6 @@ bool UnrealToUsd::ConvertRectLightComponent(const URectLightComponent& LightComp
 		return false;
 	}
 
-	if (UsdUtils::NotifyIfInstanceProxy(Prim))
-	{
-		return false;
-	}
-
 	FUsdStageInfo StageInfo(Prim.GetStage());
 
 	// Disk light
@@ -607,11 +592,6 @@ bool UnrealToUsd::ConvertPointLightComponent(const UPointLightComponent& LightCo
 		return false;
 	}
 
-	if (UsdUtils::NotifyIfInstanceProxy(Prim))
-	{
-		return false;
-	}
-
 	FUsdStageInfo StageInfo(Prim.GetStage());
 
 	if (pxr::UsdAttribute Attr = Light.CreateRadiusAttr())
@@ -663,11 +643,6 @@ bool UnrealToUsd::ConvertSkyLightComponent(const USkyLightComponent& LightCompon
 		return false;
 	}
 
-	if (UsdUtils::NotifyIfInstanceProxy(Prim))
-	{
-		return false;
-	}
-
 #if WITH_EDITORONLY_DATA
 	FUsdStageInfo StageInfo(Prim.GetStage());
 
@@ -717,11 +692,6 @@ bool UnrealToUsd::ConvertSpotLightComponent(const USpotLightComponent& LightComp
 		return false;
 	}
 
-	if (UsdUtils::NotifyIfInstanceProxy(Prim))
-	{
-		return false;
-	}
-
 	if (pxr::UsdAttribute ConeAngleAttr = ShapingAPI.CreateShapingConeAngleAttr())
 	{
 		ConeAngleAttr.Set<float>(LightComponent.OuterConeAngle, UsdTimeCode);
diff --git a/Engine/Plugins/Runtime/USDCore/Source/USDUtilities/Private/USDPrimConversion.cpp b/Engine/Plugins/Runtime/USDCore/Source/USDUtilities/Private/USDPrimConversion.cpp
--- a/Engine/Plugins/Runtime/USDCore/Source/USDUtilities/Private/USDPrimConversion.cpp
+++ b/Engine/Plugins/Runtime/USDCore/Source/USDUtilities/Private/USDPrimConversion.cpp
@@ -2103,11 +2103,6 @@ bool UnrealToUsd::ConvertCameraComponent(const UCineCameraComponent& CameraCompo
 		return false;
 	}
 
-	if (UsdUtils::NotifyIfInstanceProxy(Prim))
-	{
-		return false;
-	}
-
 	FUsdStageInfo StageInfo(Prim.GetStage());
 
 	if (pxr::UsdAttribute Attr = GeomCamera.CreateFocalLengthAttr())
@@ -2153,11 +2148,6 @@ bool UnrealToUsd::ConvertAudioComponent(const UAudioComponent& AudioComponent, p
 		return false;
 	}
 
-	if (UsdUtils::NotifyIfInstanceProxy(Prim))
-	{
-		return false;
-	}
-
 	if (!bFilePathOnly)
 	{
 		// auralmode
@@ -2258,11 +2248,6 @@ bool UnrealToUsd::ConvertAudioSection(
 		return false;
 	}
 
-	if (UsdUtils::NotifyIfInstanceProxy(Prim))
-	{
-		return false;
-	}
-
 	UMovieScene* MovieScene = AudioSection.GetTypedOuter<UMovieScene>();
 	if (!MovieScene)
 	{
@@ -3173,11 +3158,6 @@ bool UnrealToUsd::ConvertSceneComponent(const pxr::UsdStageRefPtr& Stage, const
 		return false;
 	}
 
-	if (UsdUtils::NotifyIfInstanceProxy(UsdPrim))
-	{
-		return false;
-	}
-
 	FScopedUsdAllocs UsdAllocs;
 
 	// Transform
@@ -3313,11 +3293,6 @@ bool UnrealToUsd::ConvertHierarchicalInstancedStaticMeshComponent(
 		return false;
 	}
 
-	if (UsdUtils::NotifyIfInstanceProxy(UsdPrim))
-	{
-		return false;
-	}
-
 	UsdStageRefPtr Stage = UsdPrim.GetStage();
 	FUsdStageInfo StageInfo{Stage};
 
@@ -3839,11 +3814,6 @@ bool UnrealToUsd::ConvertXformable(const FTransform& RelativeTransform, pxr::Usd
 		return false;
 	}
 
-	if (UsdUtils::NotifyIfInstanceProxy(UsdPrim))
-	{
-		return false;
-	}
-
 	FUsdStageInfo StageInfo(UsdPrim.GetStage());
 	pxr::GfMatrix4d UsdTransform = UnrealToUsd::ConvertTransform(StageInfo, RelativeTransform);
 
@@ -3925,11 +3895,6 @@ bool UnrealToUsd::ConvertInstancedFoliageActor(const AInstancedFoliageActor& Act
 		return false;
 	}
 
-	if (UsdUtils::NotifyIfInstanceProxy(UsdPrim))
-	{
-		return false;
-	}
-
 	UsdStageRefPtr Stage = UsdPrim.GetStage();
 	FUsdStageInfo StageInfo{Stage};
 
@@ -5204,11 +5169,6 @@ UnrealToUsd::FPropertyTrackWriter UnrealToUsd::CreatePropertyTrackWriter(
 	pxr::UsdStageRefPtr UsdStage = UsdPrim.GetStage();
 	FUsdStageInfo StageInfo{UsdStage};
 
-	if (UsdUtils::NotifyIfInstanceProxy(Prim))
-	{
-		return Result;
-	}
-
 	TArray<pxr::UsdAttribute> Attrs = {pxr::UsdAttribute{}};
 	pxr::UsdAttribute& Attr = Attrs[0];
 	{
diff --git a/Engine/Plugins/Runtime/USDCore/Source/USDUtilities/Private/USDShadeConversion.cpp b/Engine/Plugins/Runtime/USDCore/Source/USDUtilities/Private/USDShadeConversion.cpp
--- a/Engine/Plugins/Runtime/USDCore/Source/USDUtilities/Private/USDShadeConversion.cpp
+++ b/Engine/Plugins/Runtime/USDCore/Source/USDUtilities/Private/USDShadeConversion.cpp
@@ -2243,7 +2243,7 @@ namespace UE::UsdShadeConversion::Private
 				}
 
 				std::string ConnectedShaderName = ConnectedShader.GetPrim().GetName().GetString();
-				
+
 				bool bComputeShaderGraph = false;
 				mx::NodePtr ConnectedNode = Document->getNode(ConnectedShaderName);
 				if (!ConnectedNode)
@@ -2433,7 +2433,7 @@ mx::DocumentPtr UsdToUnreal::CreateSurfaceMaterial(const pxr::UsdPrim& UsdShadeM
 		return MtlxSurfaceShader;
 	};
 
-	//Create the shader graph for the surface shader and displacement shader	
+	//Create the shader graph for the surface shader and displacement shader
 	if (mx::NodePtr MaterialXSurfaceShader = CreateShaderGraph(SurfaceShader, Document))
 	{
 		mx::NodeDefPtr SurfaceMaterialNodeDef = Document->getNodeDef("ND_surfacematerial");
@@ -3214,11 +3214,6 @@ bool UsdUtils::MarkMaterialPrimWithWorldSpaceNormals(const UE::FUsdPrim& Materia
 		return false;
 	}
 
-	if (UsdUtils::NotifyIfInstanceProxy(MaterialPrim))
-	{
-		return false;
-	}
-
 	const bool bCustom = true;
 	pxr::UsdAttribute Attr = UsdPrim.CreateAttribute(UnrealIdentifiers::WorldSpaceNormals, pxr::SdfValueTypeNames->Bool, bCustom);
 	if (!Attr)
