USD: Use ISM instead of HISM components whenever possible.

diff --git a/Engine/Plugins/Importers/USDImporter/Source/USDExporter/Private/USDConversionBlueprintContext.cpp b/Engine/Plugins/Importers/USDImporter/Source/USDExporter/Private/USDConversionBlueprintContext.cpp
--- a/Engine/Plugins/Importers/USDImporter/Source/USDExporter/Private/USDConversionBlueprintContext.cpp
+++ b/Engine/Plugins/Importers/USDImporter/Source/USDExporter/Private/USDConversionBlueprintContext.cpp
@@ -249,7 +249,12 @@ bool UUsdConversionBlueprintContext::ConvertDrawModeComponent(const UUsdDrawMode
 #endif	  // USE_USD_SDK
 }
 
-bool UUsdConversionBlueprintContext::ConvertAudioComponent(const UAudioComponent* Component, const FString& PrimPath, float TimeCode, bool bFilePathOnly)
+bool UUsdConversionBlueprintContext::ConvertAudioComponent(
+	const UAudioComponent* Component,
+	const FString& PrimPath,
+	float TimeCode,
+	bool bFilePathOnly
+)
 {
 #if USE_USD_SDK
 	UE::FUsdPrim Prim = UnrealToUsdImpl::GetPrim(Stage, PrimPath);
@@ -279,6 +284,21 @@ bool UUsdConversionBlueprintContext::ConvertSceneComponent(const USceneComponent
 #endif	  // USE_USD_SDK
 }
 
+bool UUsdConversionBlueprintContext::ConvertIsmComponent(const UInstancedStaticMeshComponent* Component, const FString& PrimPath, float TimeCode)
+{
+#if USE_USD_SDK
+	UE::FUsdPrim Prim = UnrealToUsdImpl::GetPrim(Stage, PrimPath);
+	if (!Prim || !Component)
+	{
+		return false;
+	}
+
+	return UnrealToUsd::ConvertInstancedStaticMeshComponent(Component, Prim, TimeCode == FLT_MAX ? UsdUtils::GetDefaultTimeCode() : TimeCode);
+#else
+	return false;
+#endif	  // USE_USD_SDK
+}
+
 bool UUsdConversionBlueprintContext::ConvertHismComponent(
 	const UHierarchicalInstancedStaticMeshComponent* Component,
 	const FString& PrimPath,
diff --git a/Engine/Plugins/Importers/USDImporter/Source/USDExporter/Public/USDConversionBlueprintContext.h b/Engine/Plugins/Importers/USDImporter/Source/USDExporter/Public/USDConversionBlueprintContext.h
--- a/Engine/Plugins/Importers/USDImporter/Source/USDExporter/Public/USDConversionBlueprintContext.h
+++ b/Engine/Plugins/Importers/USDImporter/Source/USDExporter/Public/USDConversionBlueprintContext.h
@@ -19,6 +19,7 @@ class ALandscapeProxy;
 class UAudioComponent;
 class UCineCameraComponent;
 class UDirectionalLightComponent;
+class UInstancedStaticMeshComponent;
 class UHierarchicalInstancedStaticMeshComponent;
 class ULevelExporterUSDOptions;
 class ULightComponentBase;
@@ -130,6 +131,9 @@ public:
 	UFUNCTION(BlueprintCallable, Category = "Component conversion")
 	bool ConvertSceneComponent(const USceneComponent* Component, const FString& PrimPath);
 
+	UFUNCTION(BlueprintCallable, Category = "Component conversion")
+	bool ConvertIsmComponent(const UInstancedStaticMeshComponent* Component, const FString& PrimPath, float TimeCode = 3.402823466e+38F);
+
 	UFUNCTION(BlueprintCallable, Category = "Component conversion")
 	bool ConvertHismComponent(const UHierarchicalInstancedStaticMeshComponent* Component, const FString& PrimPath, float TimeCode = 3.402823466e+38F);
 
diff --git a/Engine/Plugins/Importers/USDImporter/Source/USDSchemas/Private/USDGeomPointInstancerTranslator.cpp b/Engine/Plugins/Importers/USDImporter/Source/USDSchemas/Private/USDGeomPointInstancerTranslator.cpp
--- a/Engine/Plugins/Importers/USDImporter/Source/USDSchemas/Private/USDGeomPointInstancerTranslator.cpp
+++ b/Engine/Plugins/Importers/USDImporter/Source/USDSchemas/Private/USDGeomPointInstancerTranslator.cpp
@@ -79,7 +79,7 @@ namespace UsdGeomPointInstancerTranslatorImpl
 
 		if (StaticMesh)
 		{
-			StaticMesh->CreateBodySetup();					// BodySetup is required for HISM component
+			StaticMesh->CreateBodySetup();					// BodySetup is required for ISM component
 			StaticMesh->MarkAsNotHavingNavigationData();	// Needed or else it will warn if we try cooking with body setup
 		}
 
@@ -188,7 +188,7 @@ void FUsdGeomPointInstancerTranslator::CreateAssets()
 			bIgnoreTopLevelTransformAndVisibility
 		));
 	}
-	// Otherwise we're just going to spawn HISM components instead
+	// Otherwise we're just going to spawn ISM components instead
 	else
 	{
 		const pxr::UsdRelationship& Prototypes = PointInstancer.GetPrototypesRel();
@@ -261,9 +261,9 @@ void FUsdGeomPointInstancerTranslator::CreateAssets()
 			else
 			{
 				// Bake the prototype transform into the mesh as it must always be applied before the instancing transform, and if we
-				// put it on the HISM component instead it would have been applied *after* the instancing transform
+				// put it on the ISM component instead it would have been applied *after* the instancing transform
 				const bool bIgnoreTopLevelTransform = false;
-				// Don't bake the visibility into the mesh as we can control it by toggling HiddenInGame on the HISM component directly
+				// Don't bake the visibility into the mesh as we can control it by toggling HiddenInGame on the ISM component directly
 				const bool bIgnoreTopLevelVisibility = true;
 				Context->TranslatorTasks.Add(MakeShared<FUsdGeomPointInstancerCreateAssetsTaskChain>(
 					Context,
@@ -288,15 +288,15 @@ USceneComponent* FUsdGeomPointInstancerTranslator::CreateComponents()
 	}
 
 	// Otherwise, the plan here is to create an USceneComponent that corresponds to the PointInstancer prim itself, and then spawn a child
-	// HISM component for each prototype.
+	// ISM component for each prototype.
 	// We always request a scene component here explicitly or else we'll be upgraded to a static mesh component by the mechanism that
 	// handles collapsed meshes/static mesh components for the GeomXFormable translator.
-	bool bCreateChildHISMs = false;
+	bool bCreateChildISMs = false;
 	USceneComponent* MainSceneComponent = nullptr;
 	EUsdDrawMode DrawMode = UsdUtils::GetAppliedDrawMode(GetPrim());
 	if (DrawMode == EUsdDrawMode::Default)
 	{
-		bCreateChildHISMs = true;
+		bCreateChildISMs = true;
 		MainSceneComponent = CreateComponentsEx({USceneComponent::StaticClass()}, {});
 		MainSceneComponent->ComponentTags.AddUnique(UnrealIdentifiers::PointInstancerTag);
 	}
@@ -305,8 +305,8 @@ USceneComponent* FUsdGeomPointInstancerTranslator::CreateComponents()
 		MainSceneComponent = CreateAlternativeDrawModeComponents(DrawMode);
 	}
 
-	// Actually create the child HISM components for each point instancer prototype
-	if (bCreateChildHISMs)
+	// Actually create the child ISM components for each point instancer prototype
+	if (bCreateChildISMs)
 	{
 		FScopedUsdAllocs UsdAllocs;
 
@@ -328,11 +328,11 @@ USceneComponent* FUsdGeomPointInstancerTranslator::CreateComponents()
 			return MainSceneComponent;
 		}
 
-		// Lets pretend ParentComponent is pointing to the parent USceneComponent while we create the child HISMs, so they get
+		// Lets pretend ParentComponent is pointing to the parent USceneComponent while we create the child ISMs, so they get
 		// automatically attached to it as children
 		TGuardValue<USceneComponent*> ParentComponentGuard{Context->ParentComponent, MainSceneComponent};
 
-		TArray<TFuture<TTuple<UHierarchicalInstancedStaticMeshComponent*, TArray<FTransform>>>> Tasks;
+		TArray<TFuture<TTuple<UInstancedStaticMeshComponent*, TArray<FTransform>>>> Tasks;
 		FScopedSlowTask PrototypePathsSlowTask(
 			(float)PrototypePaths.size(),
 			LOCTEXT("GeomPointCreateComponents", "Creating HierarchicalInstancedStaticMeshComponents for point instancers")
@@ -374,19 +374,22 @@ USceneComponent* FUsdGeomPointInstancerTranslator::CreateComponents()
 				}
 			}
 
-			using UHISMComponent = UHierarchicalInstancedStaticMeshComponent;
+			UClass* ComponentClass = Context->bAllowInterpretingLODs ? UHierarchicalInstancedStaticMeshComponent::StaticClass()
+																	 : UInstancedStaticMeshComponent::StaticClass();
 
-			FUsdGeomXformableTranslator XformableTranslator{UHISMComponent::StaticClass(), Context, UE::FUsdTyped(PrototypeUsdPrim)};
+			FUsdGeomXformableTranslator XformableTranslator{ComponentClass, Context, UE::FUsdTyped(PrototypeUsdPrim)};
 
 			const bool bNeedsActor = false;
-			UHISMComponent* HISMComponent = Cast<UHISMComponent>(XformableTranslator.CreateComponentsEx({UHISMComponent::StaticClass()}, bNeedsActor)
+			UInstancedStaticMeshComponent* ISMComponent = Cast<UInstancedStaticMeshComponent>(
+				XformableTranslator.CreateComponentsEx({ComponentClass}, bNeedsActor)
 			);
-			if (!HISMComponent)
+
+			if (!ISMComponent)
 			{
 				UE_LOG(
 					LogUsd,
 					Warning,
-					TEXT("Failed to generate HISM component for prototype '%s' for PointInstancer '%s'"),
+					TEXT("Failed to generate ISM/HISM component for prototype '%s' for PointInstancer '%s'"),
 					*PrototypePathStr,
 					*PrimPath.GetString()
 				);
@@ -407,7 +410,7 @@ void FUsdGeomPointInstancerTranslator::UpdateComponents(USceneComponent* PointIn
 		return;
 	}
 
-	// Do this first as we can use it's resulting HiddenInGame value to propagate visibility to the HISM components
+	// Do this first as we can use it's resulting HiddenInGame value to propagate visibility to the ISM components
 	Super::UpdateComponents(PointInstancerRootComponent);
 
 	// We always spawn exactly an USceneComponent for the "parent" component of the point instancer, so early out if
@@ -433,18 +436,18 @@ void FUsdGeomPointInstancerTranslator::UpdateComponents(USceneComponent* PointIn
 			return;
 		}
 
-		// Lets pretend ParentComponent is pointing to the parent USceneComponent while we create the child HISMs, so they get
+		// Lets pretend ParentComponent is pointing to the parent USceneComponent while we create the child ISMs, so they get
 		// automatically attached to it as children
 		TGuardValue<USceneComponent*> ParentComponentGuard{Context->ParentComponent, PointInstancerRootComponent};
 
 		const TArray<TObjectPtr<USceneComponent>>& AttachedChildren = PointInstancerRootComponent->GetAttachChildren();
-		TArray<UHierarchicalInstancedStaticMeshComponent*> AttachedHISMs;
-		AttachedHISMs.Reserve(AttachedChildren.Num());
+		TArray<UInstancedStaticMeshComponent*> AttachedISMs;
+		AttachedISMs.Reserve(AttachedChildren.Num());
 		for (const TObjectPtr<USceneComponent>& AttachedChild : AttachedChildren)
 		{
-			if (UHierarchicalInstancedStaticMeshComponent* HISM = Cast<UHierarchicalInstancedStaticMeshComponent>(AttachedChild))
+			if (UInstancedStaticMeshComponent* ISM = Cast<UInstancedStaticMeshComponent>(AttachedChild))
 			{
-				AttachedHISMs.Add(HISM);
+				AttachedISMs.Add(ISM);
 			}
 		}
 
@@ -467,7 +470,7 @@ void FUsdGeomPointInstancerTranslator::UpdateComponents(USceneComponent* PointIn
 			}
 		}
 
-		TArray<TFuture<TTuple<UHierarchicalInstancedStaticMeshComponent*, TArray<FTransform>>>> Tasks;
+		TArray<TFuture<TTuple<UInstancedStaticMeshComponent*, TArray<FTransform>>>> Tasks;
 		size_t NumPrototypePaths = PrototypePaths.Get().size();
 		FScopedSlowTask PrototypePathsSlowTask(
 			(float)NumPrototypePaths,
@@ -492,21 +495,21 @@ void FUsdGeomPointInstancerTranslator::UpdateComponents(USceneComponent* PointIn
 				continue;
 			}
 
-			UHierarchicalInstancedStaticMeshComponent* HISMComponent = nullptr;
+			UInstancedStaticMeshComponent* ISMComponent = nullptr;
 
 			// The user could have just manually deleted the component, so we must check
-			if (!AttachedHISMs.IsValidIndex(PrototypeIndex))
+			if (!AttachedISMs.IsValidIndex(PrototypeIndex))
 			{
 				UE_LOG(
 					LogUsd,
 					Warning,
-					TEXT("Failed to find corresponding HISM component for prototype '%s' of PointInstancer '%s'. Cancelling component update"),
+					TEXT("Failed to find corresponding ISM/HISM component for prototype '%s' of PointInstancer '%s'. Cancelling component update"),
 					*UsdToUnreal::ConvertPath(PrototypePath),
 					*PrimPath.GetString()
 				);
 				break;
 			}
-			HISMComponent = AttachedHISMs[PrototypeIndex];
+			ISMComponent = AttachedISMs[PrototypeIndex];
 
 			// If our prototype was a LOD mesh we will have used the path of one of the actual LOD meshes to start the FGeomMeshCreateAssetsTaskChain,
 			// so we have to look for our resulting mesh with the same path
@@ -532,19 +535,19 @@ void FUsdGeomPointInstancerTranslator::UpdateComponents(USceneComponent* PointIn
 			{
 				StaticMesh = Iter->second;
 			}
-			UsdGeomPointInstancerTranslatorImpl::SetStaticMesh(StaticMesh, *HISMComponent);
+			UsdGeomPointInstancerTranslatorImpl::SetStaticMesh(StaticMesh, *ISMComponent);
 
 			// Evaluating point instancer can take a long time and is thread-safe. Move to async task while we work on something else.
 			pxr::UsdTimeCode TimeCode{Context->Time};
 			FUsdStageInfo StageInfo{Prim.GetStage()};
 			Tasks.Emplace(Async(
 				EAsyncExecution::ThreadPool,
-				[TimeCode, StageInfo, PointInstancer, PrototypeIndex, HISMComponent]()
+				[TimeCode, StageInfo, PointInstancer, PrototypeIndex, ISMComponent]()
 				{
 					TArray<FTransform> InstanceTransforms;
 					UsdUtils::GetPointInstancerTransforms(StageInfo, PointInstancer, PrototypeIndex, TimeCode, InstanceTransforms);
 
-					return MakeTuple(HISMComponent, MoveTemp(InstanceTransforms));
+					return MakeTuple(ISMComponent, MoveTemp(InstanceTransforms));
 				}
 			));
 
@@ -560,7 +563,7 @@ void FUsdGeomPointInstancerTranslator::UpdateComponents(USceneComponent* PointIn
 				MeshTranslationImpl::SetMaterialOverrides(
 					PrototypeUsdPrim.Get(),
 					ExistingAssignments,
-					*HISMComponent,
+					*ISMComponent,
 					*Context->UsdAssetCache,
 					*Context->UsdInfoCache,
 					*Context->PrimLinkCache,
@@ -573,11 +576,11 @@ void FUsdGeomPointInstancerTranslator::UpdateComponents(USceneComponent* PointIn
 				);
 			}
 
-			// Handle visibility for the HISM component
+			// Handle visibility for the ISM component
 			//
 			// We will in general no longer put the visibility of the prototype itself baked into the mesh, so we
-			// must add it to the HISM component. These HISM components sort of correspond to both the PointInstancer prim
-			// and a particular prototype at the same time (kinda), so here we'll hide the HISM if either the PointInstancer or the
+			// must add it to the ISM component. These ISM components sort of correspond to both the PointInstancer prim
+			// and a particular prototype at the same time (kinda), so here we'll hide the ISM if either the PointInstancer or the
 			// prototype are hidden, as if they were parent/child prims
 			bool bIsVisible = !PointInstancerRootComponent->bHiddenInGame;
 			if (pxr::UsdGeomImageable PrototypeImageable = pxr::UsdGeomImageable{PrototypeUsdPrim.Get()})
@@ -589,13 +592,13 @@ void FUsdGeomPointInstancerTranslator::UpdateComponents(USceneComponent* PointIn
 					bIsVisible = false;
 				}
 			}
-			HISMComponent->SetHiddenInGame(!bIsVisible);
+			ISMComponent->SetHiddenInGame(!bIsVisible);
 		}
 
 		// Wait on and assign results of the point instancer.
 		for (auto& Future : Tasks)
 		{
-			TTuple<UHierarchicalInstancedStaticMeshComponent*, TArray<FTransform>> Result{Future.Get()};
+			TTuple<UInstancedStaticMeshComponent*, TArray<FTransform>> Result{Future.Get()};
 			UsdGeomPointInstancerTranslatorImpl::ApplyPointInstanceTransforms(Result.Key, Result.Value);
 		}
 	}
diff --git a/Engine/Plugins/Importers/USDImporter/Source/USDSchemas/Private/USDGeomXformableTranslator.cpp b/Engine/Plugins/Importers/USDImporter/Source/USDSchemas/Private/USDGeomXformableTranslator.cpp
--- a/Engine/Plugins/Importers/USDImporter/Source/USDSchemas/Private/USDGeomXformableTranslator.cpp
+++ b/Engine/Plugins/Importers/USDImporter/Source/USDSchemas/Private/USDGeomXformableTranslator.cpp
@@ -584,7 +584,7 @@ USceneComponent* FUsdGeomXformableTranslator::CreateComponentsEx(TOptional<TSubc
 		if (!GEnableCollision)
 		{
 			// In most cases this will have no benefit memory-wise, as regular UStaticMeshComponents build their physics meshes anyway
-			// when registering, regardless of these. HISM components will *not* build them though, so disabling the cvar may lead
+			// when registering, regardless of these. ISM components will *not* build them though, so disabling the cvar may lead
 			// to some memory savings
 			if (UPrimitiveComponent* PrimComp = Cast<UPrimitiveComponent>(SceneComponent))
 			{
diff --git a/Engine/Plugins/Importers/USDImporter/Source/USDStage/Private/USDPrimTwin.cpp b/Engine/Plugins/Importers/USDImporter/Source/USDStage/Private/USDPrimTwin.cpp
--- a/Engine/Plugins/Importers/USDImporter/Source/USDStage/Private/USDPrimTwin.cpp
+++ b/Engine/Plugins/Importers/USDImporter/Source/USDStage/Private/USDPrimTwin.cpp
@@ -106,8 +106,8 @@ void UUsdPrimTwin::Clear()
 			AttachChild->Modify();
 		}
 
-		// When we translate point instancers we secretly add more HISM components for each prototype.
-		// If we're destroying the scene component for the PointInstancer, we should destroy the HISMs too.
+		// When we translate point instancers we secretly add more ISM components for each prototype.
+		// If we're destroying the scene component for the PointInstancer, we should destroy the ISMs too.
 		if (SceneComponent->ComponentTags.Contains(UnrealIdentifiers::PointInstancerTag))
 		{
 			// Copy the array here as DestroyComponent will change the container directly
diff --git a/Engine/Plugins/Importers/USDImporter/Source/USDStage/Private/USDStageActor.cpp b/Engine/Plugins/Importers/USDImporter/Source/USDStage/Private/USDStageActor.cpp
--- a/Engine/Plugins/Importers/USDImporter/Source/USDStage/Private/USDStageActor.cpp
+++ b/Engine/Plugins/Importers/USDImporter/Source/USDStage/Private/USDStageActor.cpp
@@ -2106,8 +2106,8 @@ UUsdPrimTwin* AUsdStageActor::ExpandPrim(
 
 					for (const TObjectPtr<USceneComponent>& Child : TwinSceneComponent->GetAttachChildren())
 					{
-						UHierarchicalInstancedStaticMeshComponent* HISMComponent = Cast<UHierarchicalInstancedStaticMeshComponent>(Child.Get());
-						ObjectsToWatch.Remove(HISMComponent);
+						UInstancedStaticMeshComponent* ISMComponent = Cast<UInstancedStaticMeshComponent>(Child.Get());
+						ObjectsToWatch.Remove(ISMComponent);
 					}
 				}
 				SchemaTranslator->UpdateComponents(TwinSceneComponent);
@@ -2202,19 +2202,19 @@ UUsdPrimTwin* AUsdStageActor::ExpandPrim(
 			const TArray<TObjectPtr<USceneComponent>>& ChildComponents = TwinSceneComponent->GetAttachChildren();
 			for (const TObjectPtr<USceneComponent>& Child : ChildComponents)
 			{
-				UHierarchicalInstancedStaticMeshComponent* HISMComponent = Cast<UHierarchicalInstancedStaticMeshComponent>(Child.Get());
-				if (!HISMComponent)
+				UInstancedStaticMeshComponent* ISMComponent = Cast<UInstancedStaticMeshComponent>(Child.Get());
+				if (!ISMComponent)
 				{
 					continue;
 				}
 
-				UStaticMesh* HISMMesh = HISMComponent->GetStaticMesh();
-				if (!HISMMesh)
+				UStaticMesh* ISMMesh = ISMComponent->GetStaticMesh();
+				if (!ISMMesh)
 				{
 					continue;
 				}
 
-				UUsdAssetUserData* UserData = UsdUnreal::ObjectUtils::GetAssetUserData(HISMMesh);
+				UUsdAssetUserData* UserData = UsdUnreal::ObjectUtils::GetAssetUserData(ISMMesh);
 				if (!UserData)
 				{
 					continue;
@@ -2224,7 +2224,7 @@ UUsdPrimTwin* AUsdStageActor::ExpandPrim(
 				{
 					if (PrototypePaths.Contains(Path))
 					{
-						ObjectsToWatch.Add(HISMComponent, Path);
+						ObjectsToWatch.Add(ISMComponent, Path);
 					}
 				}
 			}
@@ -4581,7 +4581,7 @@ void AUsdStageActor::OnObjectPropertyChanged(UObject* ObjectBeingModified, FProp
 
 	FString PrimPath = ObjectsToWatch[ComponentBeingModified];
 
-	// Not all of our spawned components will have prim twins (e.g. HISM components for PointInstancers)
+	// Not all of our spawned components will have prim twins (e.g. ISM components for PointInstancers)
 	USceneComponent* PrimSceneComponent = Cast<USceneComponent>(ComponentBeingModified);
 	UUsdPrimTwin* UsdPrimTwin = GetRootPrimTwin()->Find(PrimPath);
 	if (UsdPrimTwin)
diff --git a/Engine/Plugins/Runtime/USDCore/Source/USDUtilities/Private/USDConversionUtils.cpp b/Engine/Plugins/Runtime/USDCore/Source/USDUtilities/Private/USDConversionUtils.cpp
--- a/Engine/Plugins/Runtime/USDCore/Source/USDUtilities/Private/USDConversionUtils.cpp
+++ b/Engine/Plugins/Runtime/USDCore/Source/USDUtilities/Private/USDConversionUtils.cpp
@@ -559,11 +559,11 @@ FString UsdUtils::GetSchemaNameForComponent(const USceneComponent& Component)
 	{
 		return TEXT("SkelRoot");
 	}
-	else if (Component.IsA<UHierarchicalInstancedStaticMeshComponent>())
+	else if (Component.IsA<UInstancedStaticMeshComponent>())
 	{
-		// The original HISM component becomes just a regular Xform prim, so that we can handle
+		// The original ISM component becomes just a regular Xform prim, so that we can handle
 		// its children correctly. We'll manually create a new child PointInstancer prim to it
-		// however, and convert the HISM data onto that prim.
+		// however, and convert the ISM data onto that prim.
 		return TEXT("Xform");
 	}
 	else if (const UStaticMeshComponent* StaticMeshComponent = Cast<UStaticMeshComponent>(&Component))
diff --git a/Engine/Plugins/Runtime/USDCore/Source/USDUtilities/Private/USDPrimConversion.cpp b/Engine/Plugins/Runtime/USDCore/Source/USDUtilities/Private/USDPrimConversion.cpp
--- a/Engine/Plugins/Runtime/USDCore/Source/USDUtilities/Private/USDPrimConversion.cpp
+++ b/Engine/Plugins/Runtime/USDCore/Source/USDUtilities/Private/USDPrimConversion.cpp
@@ -3282,13 +3282,18 @@ bool UnrealToUsd::ConvertHierarchicalInstancedStaticMeshComponent(
 	pxr::UsdPrim& UsdPrim,
 	double TimeCode
 )
+{
+	return ConvertInstancedStaticMeshComponent(HISMComponent, UsdPrim, TimeCode);
+}
+
+bool UnrealToUsd::ConvertInstancedStaticMeshComponent(const UInstancedStaticMeshComponent* ISMComponent, pxr::UsdPrim& UsdPrim, double TimeCode)
 {
 	using namespace pxr;
 
 	FScopedUsdAllocs Allocs;
 
 	UsdGeomPointInstancer PointInstancer{UsdPrim};
-	if (!PointInstancer || !HISMComponent)
+	if (!PointInstancer || !ISMComponent)
 	{
 		return false;
 	}
@@ -3301,13 +3306,13 @@ bool UnrealToUsd::ConvertHierarchicalInstancedStaticMeshComponent(
 	VtArray<GfQuath> Orientations;
 	VtArray<GfVec3f> Scales;
 
-	const int32 NumInstances = HISMComponent->GetInstanceCount();
+	const int32 NumInstances = ISMComponent->GetInstanceCount();
 	ProtoIndices.reserve(ProtoIndices.size() + NumInstances);
 	Positions.reserve(Positions.size() + NumInstances);
 	Orientations.reserve(Orientations.size() + NumInstances);
 	Scales.reserve(Scales.size() + NumInstances);
 
-	for (const FInstancedStaticMeshInstanceData& InstanceData : HISMComponent->PerInstanceSMData)
+	for (const FInstancedStaticMeshInstanceData& InstanceData : ISMComponent->PerInstanceSMData)
 	{
 		// Convert axes
 		FTransform UETransform{InstanceData.Transform};
@@ -3324,7 +3329,7 @@ bool UnrealToUsd::ConvertHierarchicalInstancedStaticMeshComponent(
 			Translation *= (UEMetersPerUnit / StageInfo.MetersPerUnit);
 		}
 
-		ProtoIndices.push_back(0);	  // We will always export a single prototype per PointInstancer, since HISM components handle only 1 mesh at a
+		ProtoIndices.push_back(0);	  // We will always export a single prototype per PointInstancer, since ISM components handle only 1 mesh at a
 									  // time
 		Positions.push_back(GfVec3f(Translation.X, Translation.Y, Translation.Z));
 		Orientations.push_back(GfQuath(Rotation.W, Rotation.X, Rotation.Y, Rotation.Z));
diff --git a/Engine/Plugins/Runtime/USDCore/Source/USDUtilities/Public/USDPrimConversion.h b/Engine/Plugins/Runtime/USDCore/Source/USDUtilities/Public/USDPrimConversion.h
--- a/Engine/Plugins/Runtime/USDCore/Source/USDUtilities/Public/USDPrimConversion.h
+++ b/Engine/Plugins/Runtime/USDCore/Source/USDUtilities/Public/USDPrimConversion.h
@@ -31,6 +31,7 @@ PXR_NAMESPACE_CLOSE_SCOPE
 class AInstancedFoliageActor;
 class UAudioComponent;
 class UCineCameraComponent;
+class UInstancedStaticMeshComponent;
 class UHierarchicalInstancedStaticMeshComponent;
 class ULevel;
 class UMeshComponent;
@@ -314,6 +315,12 @@ namespace UnrealToUsd
 		double TimeCode = UsdUtils::GetDefaultTimeCode()
 	);
 
+	USDUTILITIES_API bool ConvertInstancedStaticMeshComponent(
+		const UInstancedStaticMeshComponent* ISMComponent,
+		pxr::UsdPrim& UsdPrim,
+		double TimeCode = UsdUtils::GetDefaultTimeCode()
+	);
+
 	/**
 	 * Sets the material overrides as opinions on MeshPrim, taking care to figure out whether it needs to create LOD or subsection prims, and
 	 * which ones given the chosen export LOD range.
diff --git a/Engine/Plugins/Runtime/USDCore/Source/UnrealUSDWrapper/Public/UnrealUSDWrapper.h b/Engine/Plugins/Runtime/USDCore/Source/UnrealUSDWrapper/Public/UnrealUSDWrapper.h
--- a/Engine/Plugins/Runtime/USDCore/Source/UnrealUSDWrapper/Public/UnrealUSDWrapper.h
+++ b/Engine/Plugins/Runtime/USDCore/Source/UnrealUSDWrapper/Public/UnrealUSDWrapper.h
@@ -555,15 +555,15 @@ namespace UnrealIdentifiers
 	// The character used to separate property namespaces (usually just ':')
 	extern UNREALUSDWRAPPER_API FString UsdNamespaceDelimiter;
 
-	extern UNREALUSDWRAPPER_API const FName UniversalRenderContext; // Likely the empty string
+	extern UNREALUSDWRAPPER_API const FName UniversalRenderContext;	   // Likely the empty string
 	extern UNREALUSDWRAPPER_API const FName UnrealRenderContext;
 	extern UNREALUSDWRAPPER_API const FName MaterialXRenderContext;
 	extern UNREALUSDWRAPPER_API const FName MdlRenderContext;
 
-	extern UNREALUSDWRAPPER_API const FString UniversalRenderContextDisplayString; // "universal"
+	extern UNREALUSDWRAPPER_API const FString UniversalRenderContextDisplayString;	  // "universal"
 
 	// This marks a generated scene component as corresponding to a PointInstancer prim.
-	// It is mainly used to let us know to cleanup any HISM child component whenever we're cleaning up the scene component
+	// It is mainly used to let us know to cleanup any ISM/HISM child component whenever we're cleaning up the scene component
 	extern UNREALUSDWRAPPER_API const FName PointInstancerTag;
 
 }	 // namespace UnrealIdentifiers
