USD: Fix exporter generating invalid point instancers if a foliage static mesh fails to export.

Fix importer crashing when handling invalid point instancers.

diff --git a/Engine/Plugins/Importers/USDImporter/Source/USDSchemas/Private/USDInfoCache.cpp b/Engine/Plugins/Importers/USDImporter/Source/USDSchemas/Private/USDInfoCache.cpp
--- a/Engine/Plugins/Importers/USDImporter/Source/USDSchemas/Private/USDInfoCache.cpp
+++ b/Engine/Plugins/Importers/USDImporter/Source/USDSchemas/Private/USDInfoCache.cpp
@@ -1046,7 +1046,10 @@ namespace UE::USDInfoCache::Private
 					{
 						for (int ProtoIndex : ProtoIndicesArr)
 						{
-							OutVertexCount += PrototypeVertexCounts[ProtoIndex];
+							if (PrototypeVertexCounts.IsValidIndex(ProtoIndex))
+							{
+								OutVertexCount += PrototypeVertexCounts[ProtoIndex];
+							}
 						}
 					}
 				}
