USD: Make AUsdStageActor's subobjects into actual default subobjects.

diff --git a/Engine/Plugins/Importers/USDImporter/Source/USDStage/Private/USDPrimTwin.cpp b/Engine/Plugins/Importers/USDImporter/Source/USDStage/Private/USDPrimTwin.cpp
--- a/Engine/Plugins/Importers/USDImporter/Source/USDStage/Private/USDPrimTwin.cpp
+++ b/Engine/Plugins/Importers/USDImporter/Source/USDStage/Private/USDPrimTwin.cpp
@@ -195,3 +195,17 @@ USceneComponent* UUsdPrimTwin::GetSceneComponent() const
 
 	return nullptr;
 }
+
+void UUsdPrimTwin::Serialize(FArchive& Ar)
+{
+	Super::Serialize(Ar);
+
+	// Only when going into/out of PIE we'll duplicate over the prim twin hierarchy.
+	// When duplicating for other reasons our StageActor will LoadUsdStage(), and generate
+	// new prim twins anyway (see AUsdStageActor::PostDuplicate)
+	if (Ar.GetPortFlags() & PPF_DuplicateForPIE)
+	{
+		Ar << Children;
+		Ar << Parent;
+	}
+}
diff --git a/Engine/Plugins/Importers/USDImporter/Source/USDStage/Private/USDStageActor.cpp b/Engine/Plugins/Importers/USDImporter/Source/USDStage/Private/USDStageActor.cpp
--- a/Engine/Plugins/Importers/USDImporter/Source/USDStage/Private/USDStageActor.cpp
+++ b/Engine/Plugins/Importers/USDImporter/Source/USDStage/Private/USDStageActor.cpp
@@ -947,17 +947,16 @@ AUsdStageActor::AUsdStageActor()
 
 	RootComponent = SceneComponent;
 
-	// Note: We can't construct our RootUsdTwin as a default subobject here, it needs to be built on-demand.
-	// Even if we NewObject'd one it will work as a subobject in some contexts (maybe because the CDO will have a dedicated root twin?).
-	// As far as the engine is concerned, our prim twins are static assets like meshes or textures. However, they live on the transient
-	// package and we are the only strong reference to them, so the lifetime works out about the same, except we get to keep them during
-	// some transitions like reinstantiation.
-	// c.f. doc comment on FRecompilationTracker for more info.
-
-	// This is marked as a subobject now or else instances of blueprints that derive the AUsdStageActor don't get new transactors created
-	// for them
-	const FName UniqueName = MakeUniqueObjectName(this, UUsdTransactor::StaticClass(), TEXT("Transactor"));
-	Transactor = CreateDefaultSubobject<UUsdTransactor>(UniqueName);
+	const FName RootTwinName = MakeUniqueObjectName(this, UUsdPrimTwin::StaticClass(), TEXT("RootUsdTwin"));
+	const FName InfoCacheName = MakeUniqueObjectName(this, UUsdInfoCache::StaticClass(), TEXT("InfoCache"));
+	const FName PrimLinkCacheName = MakeUniqueObjectName(this, UUsdPrimLinkCache::StaticClass(), TEXT("LinkCache"));
+	const FName TransactorName = MakeUniqueObjectName(this, UUsdTransactor::StaticClass(), TEXT("Transactor"));
+	const bool bTransient = true;
+
+	RootUsdTwin = CreateDefaultSubobject<UUsdPrimTwin>(RootTwinName, bTransient);
+	UsdInfoCache = CreateDefaultSubobject<UUsdInfoCache>(InfoCacheName, bTransient);
+	PrimLinkCache = CreateDefaultSubobject<UUsdPrimLinkCache>(PrimLinkCacheName, bTransient);
+	Transactor = CreateDefaultSubobject<UUsdTransactor>(TransactorName, bTransient);
 	Transactor->Initialize(this);
 
 	// We never want to be without a valid BBoxCache or else we'll silently fail to compute bounds for all
@@ -3130,22 +3129,6 @@ void AUsdStageActor::OnBeginPIE(bool bIsSimulating)
 	UpdateSpawnedObjectsTransientFlag(bTransient);
 
 	bIsTransitioningIntoPIE = true;
-
-	// Take ownership of our RootTwin and pretend our entire prim tree is a subobject so that it's duplicated over with us into PIE
-	if (UUsdPrimTwin* RootTwin = GetRootPrimTwin())
-	{
-		RootTwin->Rename(nullptr, this);
-
-		if (FProperty* Prop = GetClass()->FindPropertyByName(GET_MEMBER_NAME_CHECKED(AUsdStageActor, RootUsdTwin)))
-		{
-			Prop->ClearPropertyFlags(CPF_Transient);
-		}
-
-		if (FProperty* Prop = UUsdPrimTwin::StaticClass()->FindPropertyByName(UUsdPrimTwin::GetChildrenPropertyName()))
-		{
-			Prop->ClearPropertyFlags(CPF_Transient);
-		}
-	}
 }
 
 void AUsdStageActor::OnPostPIEStarted(bool bIsSimulating)
@@ -3156,22 +3139,6 @@ void AUsdStageActor::OnPostPIEStarted(bool bIsSimulating)
 
 	bIsTransitioningIntoPIE = false;
 
-	// Put our RootTwin back on the transient package so that if our blueprint is compiled it doesn't get reconstructed with us
-	if (UUsdPrimTwin* RootTwin = GetRootPrimTwin())
-	{
-		RootTwin->Rename(nullptr, GetTransientPackage());
-
-		if (FProperty* Prop = GetClass()->FindPropertyByName(GET_MEMBER_NAME_CHECKED(AUsdStageActor, RootUsdTwin)))
-		{
-			Prop->SetPropertyFlags(CPF_Transient);
-		}
-
-		if (FProperty* Prop = UUsdPrimTwin::StaticClass()->FindPropertyByName(UUsdPrimTwin::GetChildrenPropertyName()))
-		{
-			Prop->SetPropertyFlags(CPF_Transient);
-		}
-	}
-
 	// Setup for the very first frame when we duplicate into PIE, or else we will display skeletal mesh components on their
 	// StartTimeCode state. We have to do this here (after duplicating) as we need the calls to FUsdSkelSkeletonTranslator::UpdateComponents
 	// to actually animate the components, and they will only be able to do anything after they have been registered (which
@@ -3197,17 +3164,6 @@ void AUsdStageActor::OnObjectsReplaced(const TMap<UObject*, UObject*>& ObjectRep
 	// and the spawned actors and components are still alive. We just need to move over any remaining non-property data
 	if (AUsdStageActor* NewActor = Cast<AUsdStageActor>(ObjectReplacementMap.FindRef(this)))
 	{
-		// If our BP has changes and we're going into PIE, we'll get automatically recompiled. Sadly OnBeginPIE will trigger
-		// before we're duplicated for the reinstantiation process, which is a problem because our prim twins will then be
-		// owned by us by the time we're duplicated, which will clear them. This handles that case, and just duplicates the prim
-		// twins from the old actor, which is what the reinstantiation process should have done instead anyway. Note that only
-		// later will the components and actors being pointed to by this duplicated prim twin be moved to the PIE world, so those
-		// references would be updated correctly.
-		if (RootUsdTwin && RootUsdTwin->GetOuter() == this)
-		{
-			NewActor->RootUsdTwin = DuplicateObject(RootUsdTwin, NewActor);
-		}
-
 		if (FRecompilationTracker::IsBeingCompiled(BP))
 		{
 			// Can't just move out of this one as TUsdStore expects its TOptional to always contain a value, and we may
@@ -3230,32 +3186,28 @@ void AUsdStageActor::OnObjectsReplaced(const TMap<UObject*, UObject*>& ObjectRep
 
 			NewActor->AssetCache = AssetCache;
 
-			// Move the built info cache over to the new actor
-			const TCHAR* NewName = nullptr;
-			if (UsdInfoCache)
+			// Move our subobjects to NewActor
+			auto MoveSubobject = [NewActor]<typename T>(TObjectPtr<T>& SourceObjectPtr, TObjectPtr<T>& DestObjectPtr)
 			{
-				if (NewActor->UsdInfoCache)
+				if (SourceObjectPtr)
 				{
-					NewActor->UsdInfoCache->Rename(NewName, GetTransientPackage());
-				}
+					const TCHAR* NewName = nullptr;
 
-				NewActor->UsdInfoCache = UsdInfoCache;
-				UsdInfoCache->Rename(NewName, NewActor);
-				UsdInfoCache = nullptr;
-			}
+					if (DestObjectPtr)
+					{
+						DestObjectPtr->Rename(NewName, GetTransientPackage());
+					}
 
-			// Move the built prim link cache over to the new actor
-			if (PrimLinkCache)
-			{
-				if (NewActor->PrimLinkCache)
-				{
-					NewActor->PrimLinkCache->Rename(NewName, GetTransientPackage());
-				}
+					DestObjectPtr = SourceObjectPtr;
+					SourceObjectPtr->Rename(NewName, NewActor);
 
-				NewActor->PrimLinkCache = PrimLinkCache;
-				PrimLinkCache->Rename(NewName, NewActor);
-				PrimLinkCache = nullptr;
-			}
+					SourceObjectPtr = nullptr;
+				}
+			};
+			MoveSubobject(RootUsdTwin, NewActor->RootUsdTwin);
+			MoveSubobject(UsdInfoCache, NewActor->UsdInfoCache);
+			MoveSubobject(PrimLinkCache, NewActor->PrimLinkCache);
+			// We can just keep NewActor's transactor though
 
 			NewActor->BBoxCache = BBoxCache;
 			BBoxCache = nullptr;
@@ -3362,20 +3314,6 @@ void AUsdStageActor::LoadUsdStage(bool bOpenIfNeeded)
 
 	UE::FUsdStage StageToLoad = GetUsdStage();
 
-	// Create Info and BBoxCache before calling ReloadAnimations as that is when the LevelSequenceHelper will also take
-	// a reference to them
-	if (!UsdInfoCache)
-	{
-		UObject* Outer = this;
-		UsdInfoCache = NewObject<UUsdInfoCache>(Outer, NAME_None, EObjectFlags::RF_Transient);
-	}
-
-	if (!PrimLinkCache)
-	{
-		UObject* Outer = this;
-		PrimLinkCache = NewObject<UUsdPrimLinkCache>(Outer, NAME_None, EObjectFlags::RF_Transient | EObjectFlags::RF_Transactional);
-	}
-
 	if (UsdInfoCache)
 	{
 		UsdInfoCache->GetInner().ResetTranslatedPrototypes();
@@ -3597,16 +3535,6 @@ void AUsdStageActor::RebuildInfoCacheFromStoredChanges()
 
 UUsdPrimTwin* AUsdStageActor::GetRootPrimTwin()
 {
-	if (!RootUsdTwin)
-	{
-		FScopedUnrealAllocs Allocs;
-
-		// Be careful not to give it a name, as there could be multiple of these on the transient package.
-		// It needs to be public or else FArchiveReplaceOrClearExternalReferences will reset our property
-		// whenever it is used from UEngine::CopyPropertiesForUnrelatedObjects for blueprint recompilation (if we're a blueprint class)
-		RootUsdTwin = NewObject<UUsdPrimTwin>(GetTransientPackage(), NAME_None, (DefaultObjFlag | RF_Public) & ~RF_Standalone);
-	}
-
 	return RootUsdTwin;
 }
 
@@ -4035,6 +3963,10 @@ void AUsdStageActor::PostDuplicate(bool bDuplicateForPIE)
 	}
 	else
 	{
+		// When duplicating for PIE the engine will duplicate all our spawned actors and components too.
+		// When duplicating directly (e.g. via Ctrl+D), it will not. This means we need to reload the stage in order
+		// to generate our own duplicate actors and components. It's probably possible to just traverse the attach
+		// hiearachy here and call DuplicateObject() on our spawns instead, but this should be more well-behaved
 		LoadUsdStage();
 	}
 }
@@ -4053,19 +3985,7 @@ void AUsdStageActor::Serialize(FArchive& Ar)
 		Ar << BlendShapesByPath;
 		Ar << MaterialToPrimvarToUVIndex;
 		Ar << bIsTransitioningIntoPIE;
-
-		if (!UsdInfoCache)
-		{
-			UObject* Outer = this;
-			UsdInfoCache = NewObject<UUsdInfoCache>(Outer, NAME_None, EObjectFlags::RF_Transient);
-		}
 		Ar << UsdInfoCache;
-
-		if (!PrimLinkCache)
-		{
-			UObject* Outer = this;
-			PrimLinkCache = NewObject<UUsdPrimLinkCache>(Outer, NAME_None, EObjectFlags::RF_Transient | EObjectFlags::RF_Transactional);
-		}
 		Ar << PrimLinkCache;
 	}
 
@@ -4150,22 +4070,6 @@ void AUsdStageActor::BeginDestroy()
 void AUsdStageActor::PostInitProperties()
 {
 	Super::PostInitProperties();
-
-	// Make sure we never keep a prim twin that was provided by our post-constructor initialization, as RootUsdTwin
-	// is not meant to be shared across actors.
-	//
-	// During OnBeginPIE/OnPostPIEStarted we will temporarily make our prim twins into fake subobjects,
-	// so that the PIE process duplicates all our spawns automatically. The annoying part is that if the stage actor
-	// is a spawnable, this actor with prim twin subobjects will be recorded as the "object archetype" for the
-	// spawnable in PIE, meaning that its RootUsdTwin property pointing at that one same prim twin will be copied over
-	// for every subsequent spawn of that spawnable within FObjectInitializer::InitProperties...
-	// The final outcome is that all the stage actor spawns will share the root prim twin, and when one of them clears it,
-	// all other stage actors lose their actors/components.
-	//
-	// Obviously we don't want any of this, so here we make sure that we always ignore those RootUsdTwin values.
-	// Ideally we'd make our RootUsdTwin into a default subobject so all this naturally goes away, but that will
-	// break the blueprint recompilation tricks so it's probably best left alone for now
-	RootUsdTwin = nullptr;
 }
 
 void AUsdStageActor::PostRegisterAllComponents()
diff --git a/Engine/Plugins/Importers/USDImporter/Source/USDStage/Public/USDPrimTwin.h b/Engine/Plugins/Importers/USDImporter/Source/USDStage/Public/USDPrimTwin.h
--- a/Engine/Plugins/Importers/USDImporter/Source/USDStage/Public/USDPrimTwin.h
+++ b/Engine/Plugins/Importers/USDImporter/Source/USDStage/Public/USDPrimTwin.h
@@ -34,13 +34,16 @@ public:
 	{
 		for (TMap<FString, TObjectPtr<UUsdPrimTwin>>::TIterator It = Children.CreateIterator(); It; ++It)
 		{
-			UUsdPrimTwin* Child = It->Value;
-
-			Func(*Child);
-
-			if (bRecursive)
+			// During some busy transitions it is possible that our children were already destroyed by
+			// the time we try clearing our own children
+			if (UUsdPrimTwin* Child = It->Value)
 			{
-				Child->Iterate(Func, bRecursive);
+				Func(*Child);
+
+				if (bRecursive)
+				{
+					Child->Iterate(Func, bRecursive);
+				}
 			}
 		}
 	}
@@ -50,6 +53,10 @@ public:
 
 	USceneComponent* GetSceneComponent() const;
 
+	// Begin UObject interface
+	void Serialize(FArchive& Ar) override;
+	// End UObject interface
+
 	DECLARE_EVENT_OneParam(UUsdPrimTwin, FOnUsdPrimTwinDestroyed, const UUsdPrimTwin&);
 	FOnUsdPrimTwinDestroyed OnDestroyed;
 
@@ -67,6 +74,8 @@ public:
 	TWeakObjectPtr<class USceneComponent> SceneComponent;
 
 private:
+	// Transient as we don't want to save this to disk, but we'll implement Serialize() to
+	// duplicate this over when we're being duplicated
 	UPROPERTY(Transient)
 	TMap<FString, TObjectPtr<UUsdPrimTwin>> Children;
 
diff --git a/Engine/Plugins/Importers/USDImporter/Source/USDStage/Public/USDStageActor.h b/Engine/Plugins/Importers/USDImporter/Source/USDStage/Public/USDStageActor.h
--- a/Engine/Plugins/Importers/USDImporter/Source/USDStage/Public/USDStageActor.h
+++ b/Engine/Plugins/Importers/USDImporter/Source/USDStage/Public/USDStageActor.h
@@ -62,10 +62,21 @@ public:
 	UPROPERTY(EditAnywhere, BlueprintReadOnly, Category = "USD")
 	TObjectPtr<UUsdAssetCache3> AssetCache;
 
-	UPROPERTY()
+	/**
+	 * 'Transient' as this shouldn't be saved and should be zero'd when loading (i.e. each stage actor should always have
+	 * a unique instance).
+	 *
+	 * 'Instanced' or else even though this is created as a DefaultSubObject the engine will still share the same object
+	 * between different stage actors in some scenarios (e.g. spawnables going into PIE), which we absolutely never want.
+	 *
+	 * In addition, we manually serialize this, PrimLinkCache and RootUsdTwin when going into PIE via ::Serialize
+	 * (as opposed to letting the engine leave us with reset property values), as we won't LoadUsdStage() when
+	 * duplicating for PIE, given that the engine will already duplicate our actor/component spawns anyway
+	 */
+	UPROPERTY(Transient, Instanced)
 	TObjectPtr<UUsdInfoCache> UsdInfoCache;
 
-	UPROPERTY()
+	UPROPERTY(Transient, Instanced)
 	TObjectPtr<UUsdPrimLinkCache> PrimLinkCache;
 
 	UPROPERTY()
@@ -473,19 +484,21 @@ protected:
 	UPROPERTY(VisibleAnywhere, Category = "USD", Transient)
 	TObjectPtr<ULevelSequence> LevelSequence;
 
-	// Note: This cannot be instanced: Read the comment in the constructor
-	UPROPERTY(Transient)
+	/*
+	 * c.f. UsdInfoCache UPROPERTY for property flag reasoning.
+	 */
+	UPROPERTY(Transient, Instanced)
 	TObjectPtr<UUsdPrimTwin> RootUsdTwin;
 
+	UPROPERTY(Transient, Instanced)
+	TObjectPtr<UUsdTransactor> Transactor;
+
 	UPROPERTY(Transient)
 	TSet<FString> PrimsToAnimate;
 
 	UPROPERTY(Transient)
 	TMap<TObjectPtr<UObject>, FString> ObjectsToWatch;
 
-	UPROPERTY(Transient)
-	TObjectPtr<UUsdTransactor> Transactor;
-
 	/**
 	 * USD can emit multiple notices within the same UE transaction. We accumulate these in here, and respond
 	 * to them only once, when the transaction is about to finish
