USD: Fix multiple issues when manipulating the stage properties relevant for the stage actor's LevelSequence.

Fix Sequencer refresh issues when reloading the stage or responding to notices;
Fix Sequencer work/view area resetting after some notices;
Fix work/view area not being set for layers with [0, 0] time code range;
Fix start/EndTimeCode changes not updating the playback range;
Fix changes to playback ranges not immediately refreshing the prim info panel;
Fix prim info panel displaying incorrect value of timeCodesPerSecond in some scenarios (i.e. not computing it with the fallbacks);
Fix edits to sublayer start/endTimeCodes not affecting the Sequence's playback range;
Add automated test to cover most of these issues.

diff --git a/Engine/Plugins/Importers/USDImporter/Source/USDStage/Private/USDLevelSequenceHelper.cpp b/Engine/Plugins/Importers/USDImporter/Source/USDStage/Private/USDLevelSequenceHelper.cpp
--- a/Engine/Plugins/Importers/USDImporter/Source/USDStage/Private/USDLevelSequenceHelper.cpp
+++ b/Engine/Plugins/Importers/USDImporter/Source/USDStage/Private/USDLevelSequenceHelper.cpp
@@ -943,7 +943,7 @@ private:
 	void UpdateLayerTimeInfoFromLayer(FLayerTimeInfo& LayerTimeInfo, const UE::FSdfLayer& Layer);
 
 	/** Updates MovieScene with LayerTimeInfo */
-	void UpdateMovieSceneTimeRanges(UMovieScene& MovieScene, const FLayerTimeInfo& LayerTimeInfo);
+	void UpdateMovieSceneTimeRanges(UMovieScene& MovieScene, const FLayerTimeInfo& LayerTimeInfo, bool bUpdateViewRanges = true);
 
 	double GetFramesPerSecond() const;
 	double GetTimeCodesPerSecond() const;
@@ -1182,8 +1182,6 @@ void FUsdLevelSequenceHelperImpl::AddReferencedObjects(FReferenceCollector& Coll
 
 void FUsdLevelSequenceHelperImpl::CreateLocalLayersSequences()
 {
-	Clear();
-
 	if (!UsdStage)
 	{
 		return;
@@ -1196,7 +1194,6 @@ void FUsdLevelSequenceHelperImpl::CreateLocalLayersSequences()
 
 	// Create main level sequence for root layer
 	MainLevelSequence = FindOrAddSequenceForLayer(RootLayer, RootLayer.GetIdentifier(), RootLayer.GetDisplayName());
-
 	if (!MainLevelSequence)
 	{
 		return;
@@ -3784,16 +3781,18 @@ void FUsdLevelSequenceHelperImpl::UpdateUsdLayerOffsetFromSection(const UMovieSc
 	UE_LOG(LogUsd, Verbose, TEXT("\tNew OffsetScale: %f, %f"), NewLayerOffset.Offset, NewLayerOffset.Scale);
 }
 
-void FUsdLevelSequenceHelperImpl::UpdateMovieSceneTimeRanges(UMovieScene& MovieScene, const FLayerTimeInfo& LayerTimeInfo)
+void FUsdLevelSequenceHelperImpl::UpdateMovieSceneTimeRanges(UMovieScene& MovieScene, const FLayerTimeInfo& LayerTimeInfo, bool bUpdateViewRanges)
 {
 	const double FramesPerSecond = GetFramesPerSecond();
 
+	double StartTimeCode = 0.0;
+	double EndTimeCode = 0.0f;
+	double TimeCodesPerSecond = 24.0;
 	if (LayerTimeInfo.IsAnimated())
 	{
-		double StartTimeCode = LayerTimeInfo.StartTimeCode.Get(0.0);
-		const double EndTimeCode = LayerTimeInfo.EndTimeCode.Get(0.0);
+		StartTimeCode = LayerTimeInfo.StartTimeCode.Get(0.0);
+		EndTimeCode = LayerTimeInfo.EndTimeCode.Get(0.0);
 
-		double TimeCodesPerSecond = 24.0;
 		if (UE::FSdfLayer Layer = UE::FSdfLayer::FindOrOpen(*LayerTimeInfo.Identifier))
 		{
 			TimeCodesPerSecond = Layer.GetTimeCodesPerSecond();
@@ -3817,6 +3816,12 @@ void FUsdLevelSequenceHelperImpl::UpdateMovieSceneTimeRanges(UMovieScene& MovieS
 		TRange<FFrameNumber> TimeRange = TRange<FFrameNumber>::Inclusive(StartFrame, EndFrame);
 
 		MovieScene.SetPlaybackRange(TimeRange);
+	}
+
+	// Do this even when "not animated" so that we always start with the playback range in view, even if the start/end are on top
+	// of each other
+	if (bUpdateViewRanges)
+	{
 		MovieScene.SetViewRange(StartTimeCode / TimeCodesPerSecond - 1.0f, 1.0f + EndTimeCode / TimeCodesPerSecond);
 		MovieScene.SetWorkingRange(StartTimeCode / TimeCodesPerSecond - 1.0f, 1.0f + EndTimeCode / TimeCodesPerSecond);
 	}
@@ -3969,38 +3974,60 @@ void FUsdLevelSequenceHelperImpl::OnUsdObjectsChanged(
 		return;
 	}
 
+	if (!MainLevelSequence)
+	{
+		return;
+	}
+
+	UE::FUsdStage ActiveStage = StageActorPtr->GetUsdStage();
+	if (!ActiveStage)
+	{
+		return;
+	}
+
+	TArray<UE::FSdfLayer> UsedLayers = ActiveStage.GetUsedLayers();
+
 	FScopedBlockMonitoringChangesForTransaction BlockMonitoring{*this};
 
-	for (const TPair<FString, TArray<UsdUtils::FSdfChangeListEntry>>& InfoChange : InfoChanges)
+	TArray<UE::FSdfLayer> UpdatedLayers;
+
+	auto IterateChanges = [this, UsedLayers](const UsdUtils::FObjectChangesByPath& Changes) -> bool
 	{
-		const FString& PrimPath = InfoChange.Key;
-		if (PrimPath == TEXT("/"))
+		for (const TPair<FString, TArray<UsdUtils::FSdfChangeListEntry>>& InfoChange : Changes)
 		{
-			for (const UsdUtils::FSdfChangeListEntry& ObjectChange : InfoChange.Value)
+			const FString& PrimPath = InfoChange.Key;
+			if (PrimPath == TEXT("/"))
 			{
-				for (const UsdUtils::FAttributeChange& AttributeChange : ObjectChange.AttributeChanges)
+				// Update info for all layers because OnUsdObjectsChanged doesn't specify which actual layer that changed
+				// TODO: Change this to also use the LayersDidChange notice instead (see comments on UE-222371).
+				for (const UE::FSdfLayer& UsedLayer : UsedLayers)
 				{
-					if (AttributeChange.PropertyName == TEXT("framesPerSecond") && MainLevelSequence)
+					FString Identifier = UsedLayer.GetIdentifier();
+
+					ULevelSequence* Sequence = LevelSequencesByIdentifier.FindRef(Identifier);
+					UMovieScene* MovieScene = Sequence ? Sequence->GetMovieScene() : nullptr;
+
+					FLayerTimeInfo* LayerTimeInfo = LayerTimeInfosByLayerIdentifier.Find(Identifier);
+
+					if (MovieScene && LayerTimeInfo)
 					{
-						UsdUtils::FConvertedVtValue ConvertedValue;
-						if (UsdToUnreal::ConvertValue(AttributeChange.NewValue, ConvertedValue))
-						{
-							if (ConvertedValue.Entries.Num() == 1 && ConvertedValue.Entries[0].Num() == 1)
-							{
-								if (double* NewValue = ConvertedValue.Entries[0][0].TryGet<double>())
-								{
-									if (UMovieScene* MovieScene = MainLevelSequence->GetMovieScene())
-									{
-										MovieScene->Modify();
-										MovieScene->SetDisplayRate(FFrameRate(*NewValue, 1));
-									}
-								}
-							}
-						}
+						UpdateLayerTimeInfoFromLayer(*LayerTimeInfo, UsedLayer);
+
+						// We should only change this when first creating the LevelSequence, not after every edit
+						const bool bUpdateViewRanges = false;
+						UpdateMovieSceneTimeRanges(*MovieScene, *LayerTimeInfo, bUpdateViewRanges);
 					}
 				}
+				return true;
 			}
 		}
+
+		return false;
+	};
+	bool bUpdated = IterateChanges(InfoChanges);
+	if (!bUpdated)
+	{
+		IterateChanges(ResyncChanges);
 	}
 }
 
@@ -4317,6 +4344,9 @@ void FUsdLevelSequenceHelperImpl::HandleMovieSceneChange(UMovieScene& MovieScene
 			}
 		}
 	}
+
+	// We may have changed things like playback ranges, so refresh the prim properties panel if relevant
+	StageActorPtr->OnPrimChanged.Broadcast(UE::FSdfPath::AbsoluteRootPath().GetString(), false);
 }
 
 void FUsdLevelSequenceHelperImpl::HandleSubSectionChange(UMovieSceneSubSection& Section)
diff --git a/Engine/Plugins/Importers/USDImporter/Source/USDStageEditorViewModels/Private/USDObjectFieldViewModel.cpp b/Engine/Plugins/Importers/USDImporter/Source/USDStageEditorViewModels/Private/USDObjectFieldViewModel.cpp
--- a/Engine/Plugins/Importers/USDImporter/Source/USDStageEditorViewModels/Private/USDObjectFieldViewModel.cpp
+++ b/Engine/Plugins/Importers/USDImporter/Source/USDStageEditorViewModels/Private/USDObjectFieldViewModel.cpp
@@ -396,23 +396,38 @@ void FUsdObjectFieldsViewModel::Refresh(const UE::FUsdStageWeak& InUsdStage, con
 
 			FScopedUsdAllocs UsdAllocs;
 
+			pxr::UsdStageRefPtr Stage{UsdStage};
+
 			std::vector<pxr::TfToken> TokenVector = pxr::SdfSchema::GetInstance().GetMetadataFields(pxr::SdfSpecType::SdfSpecTypePseudoRoot);
 			for (const pxr::TfToken& Token : TokenVector)
 			{
 				pxr::VtValue VtValue;
-				if (pxr::UsdStageRefPtr(UsdStage)->GetMetadata(Token, &VtValue))
-				{
-					FString FieldName = UsdToUnreal::ConvertToken(Token);
 
-					UsdUtils::FConvertedVtValue ConvertedValue;
-					if (!UsdToUnreal::ConvertValue(UE::FVtValue{VtValue}, ConvertedValue))
+				if (Token == pxr::SdfFieldKeys->TimeCodesPerSecond)
+				{
+					// There is some fallback logic here (e.g. using framesPerSecond), so don't fetch
+					// the raw metadata value generically as that will ignore it
+					VtValue = Stage->GetTimeCodesPerSecond();
+				}
+				else
+				{
+					bool bSuccess = Stage->GetMetadata(Token, &VtValue);
+					if (!bSuccess)
 					{
-						continue;
+						VtValue = pxr::VtValue{};
 					}
+				}
 
-					const bool bAttrReadOnly = ConvertedValue.bIsArrayValued;
-					CreateField(EObjectFieldType::Metadata, FieldName, ConvertedValue, bAttrReadOnly);
+				FString FieldName = UsdToUnreal::ConvertToken(Token);
+
+				UsdUtils::FConvertedVtValue ConvertedValue;
+				if (!UsdToUnreal::ConvertValue(UE::FVtValue{VtValue}, ConvertedValue))
+				{
+					continue;
 				}
+
+				const bool bAttrReadOnly = ConvertedValue.bIsArrayValued;
+				CreateField(EObjectFieldType::Metadata, FieldName, ConvertedValue, bAttrReadOnly);
 			}
 		}
 		// Show info about a prim
