USD: Fix stage importer "scene actors" having RF_Public, triggering crashes on WorldPartition levels.

diff --git a/Engine/Plugins/Importers/USDImporter/Source/USDStageImporter/Private/USDStageImporter.cpp b/Engine/Plugins/Importers/USDImporter/Source/USDStageImporter/Private/USDStageImporter.cpp
--- a/Engine/Plugins/Importers/USDImporter/Source/USDStageImporter/Private/USDStageImporter.cpp
+++ b/Engine/Plugins/Importers/USDImporter/Source/USDStageImporter/Private/USDStageImporter.cpp
@@ -199,7 +199,7 @@ namespace UE::USDStageImporter::Private
 		}
 
 		FActorSpawnParameters SpawnParameters;
-		SpawnParameters.ObjectFlags = ImportContext.ImportObjectFlags & ~RF_Standalone;
+		SpawnParameters.ObjectFlags = ImportContext.ImportObjectFlags & ~(RF_Standalone | RF_Public | RF_Transient);
 		SpawnParameters.OverrideLevel = Level;
 
 		// We always spawn another scene actor regardless of collision or whether the level already has one,
