USD: Fix crash when turning stage actor into a spawnable within it's own LevelSequence.

diff --git a/Engine/Plugins/Importers/USDImporter/Source/USDStage/Private/USDLevelSequenceHelper.cpp b/Engine/Plugins/Importers/USDImporter/Source/USDStage/Private/USDLevelSequenceHelper.cpp
--- a/Engine/Plugins/Importers/USDImporter/Source/USDStage/Private/USDLevelSequenceHelper.cpp
+++ b/Engine/Plugins/Importers/USDImporter/Source/USDStage/Private/USDLevelSequenceHelper.cpp
@@ -1149,19 +1149,6 @@ bool FUsdLevelSequenceHelperImpl::HasData() const
 
 void FUsdLevelSequenceHelperImpl::Clear()
 {
-	for (const TPair<FString, TObjectPtr<ULevelSequence>>& IdentifierAndSeq : LevelSequencesByIdentifier)
-	{
-		if (ULevelSequence* Seq = IdentifierAndSeq.Value.Get())
-		{
-#if WITH_EDITOR
-			if (!IsEngineExitRequested())
-			{
-				GEditor->GetEditorSubsystem<UAssetEditorSubsystem>()->CloseAllEditorsForAsset(Seq);
-			}
-#endif	  // WITH_EDITOR
-		}
-	}
-
 	MainLevelSequence = nullptr;
 	LevelSequencesByIdentifier.Empty();
 	IdentifierByLevelSequence.Empty();
