USD: Automatically resize subsequence sections to match the stage's playback range.

Automatically refresh time track to match playback range;
Make time track read-only by default.

diff --git a/Engine/Plugins/Importers/USDImporter/Source/USDStage/Private/USDLevelSequenceHelper.cpp b/Engine/Plugins/Importers/USDImporter/Source/USDStage/Private/USDLevelSequenceHelper.cpp
--- a/Engine/Plugins/Importers/USDImporter/Source/USDStage/Private/USDLevelSequenceHelper.cpp
+++ b/Engine/Plugins/Importers/USDImporter/Source/USDStage/Private/USDLevelSequenceHelper.cpp
@@ -866,6 +866,11 @@ private:
 	void CreateSubSequenceSection(ULevelSequence& Sequence, ULevelSequence& Subsequence);
 	void RemoveSubSequenceSection(ULevelSequence& Sequence, ULevelSequence& SubSequence);
 
+	/**
+	 * Adjusts all subsection ends, in all sequences, to match the stage's root layer's endTimeCode
+	 */
+	void SnapSubSectionEndsToPlaybackRange(bool bShowToast = false);
+
 	// Tracks handling
 private:
 	/** Creates a time track on the ULevelSequence corresponding to Info */
@@ -1644,7 +1649,19 @@ void FUsdLevelSequenceHelperImpl::CreateSubSequenceSection(ULevelSequence& Seque
 	}
 	else
 	{
-		SubSection = SubTrack->AddSequence(&SubSequence, SubSectionRange.GetLowerBoundValue(), SubSectionRange.Size<FFrameNumber>().Value);
+		// Always force our sections to be on separate rows.
+		//
+		// We compute the row ourselves instead of using AddSequence (that just passes INDEX_NONE for it), because internally
+		// UMovieSceneSubTrack::AddSequenceOnRow will let sections end up on the same row if they don't overlap. We may
+		// end up with zero-size sections in some cases though, and those *never* overlap, so they would end up bunched up on
+		// the same row (see UE-217625)
+		const int32 RowIndex = SubTrack->GetAllSections().Num();
+		SubSection = SubTrack->AddSequenceOnRow(	//
+			&SubSequence,
+			SubSectionRange.GetLowerBoundValue(),
+			SubSectionRange.Size<FFrameNumber>().Value,
+			RowIndex
+		);
 
 		UE_LOG(
 			LogUsd,
@@ -1695,6 +1712,136 @@ void FUsdLevelSequenceHelperImpl::RemoveSubSequenceSection(ULevelSequence& Seque
 	}
 }
 
+void FUsdLevelSequenceHelperImpl::SnapSubSectionEndsToPlaybackRange(bool bShowToast)
+{
+	// In USD, all endTimeCode values are ignored except those for the root layer. Even those don't "hard clip" the animation either,
+	// and are more like a suggestion of what the play range should be.
+	// In UE, subsection ranges do in fact clip the inner animation. We can't specify an unbounded range for them though, and
+	// making them max size is inconvenient UI-wise.
+	//
+	// Since there is no easy/efficient way of fetching a time code range for all animation contained "below a particular layer",
+	// the best way of reconciling that discrepancy that we can come up with so far is to always force the subsection ends to
+	// match whatever the stage's endTimeCode is, which is done here. This way, the only clipping that happens is one that is
+	// seen by the root layer's endTimeCode, which can similarly seen in other DCCs (like usdview) and is one which can be easily
+	// controlled (i.e. it's easy to just tweak the root layer's endTimeCode directly)
+
+	const double StageEndTimeSeconds = UsdStage.GetEndTimeCode() / UsdStage.GetTimeCodesPerSecond();
+
+	bool bChangedSomeSection = false;
+
+	for (const TPair<FString, TObjectPtr<ULevelSequence>>& IdentifierAndSeq : LevelSequencesByIdentifier)
+	{
+		ULevelSequence* Sequence = IdentifierAndSeq.Value.Get();
+		if (!Sequence)
+		{
+			continue;
+		}
+
+		UMovieScene* MovieScene = Sequence->GetMovieScene();
+		if (!MovieScene)
+		{
+			continue;
+		}
+
+		UMovieSceneSubTrack* SubTrack = MovieScene->FindTrack<UMovieSceneSubTrack>();
+		if (!SubTrack)
+		{
+			continue;
+		}
+
+		FFrameRate TickResolution = MovieScene->GetTickResolution();
+		const FFrameNumber StageEndFrame = UsdLevelSequenceHelperImpl::RoundAsFrameNumber(TickResolution, StageEndTimeSeconds);
+
+		for (UMovieSceneSection* Section : SubTrack->GetAllSections())
+		{
+			if (UMovieSceneSubSection* SubSection = Cast<UMovieSceneSubSection>(Section))
+			{
+				TRange<FFrameNumber> Range = SubSection->GetRange();
+
+				FFrameNumber NewUpper = FMath::Max(StageEndFrame, Range.GetLowerBoundValue());
+				if (NewUpper == Range.GetUpperBoundValue())
+				{
+					continue;
+				}
+
+				Range.SetUpperBoundValue(NewUpper);
+
+				SubSection->Modify();
+				SubSection->SetRange(Range);
+
+				bChangedSomeSection = true;
+			}
+		}
+	}
+
+	if (bChangedSomeSection && bShowToast)
+	{
+		const UUsdProjectSettings* Settings = GetDefault<UUsdProjectSettings>();
+		if (Settings && Settings->bShowSubsectionSnappingWarning)
+		{
+			const FText Text = LOCTEXT("SnappedSectionTitle", "USD: Subsections adjusted");
+
+			const FText SubText = LOCTEXT(
+				"SnappedSectionTitleSubText",
+				"In USD, start and endTimeCodes from any layer but the root layer are mostly ignored, while in UE a subsequence section range does clip the underlying animation.\n\nIn order to reconcile the behavior, subsequence sections on any of the StageActor's generated LevelSequences will automatically snap to the stage's playback range instead, which prevents any unwanted animation clipping."
+			);
+
+			UE_LOG(LogUsd, Warning, TEXT("%s"), *SubText.ToString().Replace(TEXT("\n\n"), TEXT(" ")));
+
+			static TWeakPtr<SNotificationItem> Notification;
+
+			FNotificationInfo Toast(Text);
+			Toast.SubText = SubText;
+			Toast.Image = FCoreStyle::Get().GetBrush(TEXT("MessageLog.Warning"));
+			Toast.CheckBoxText = LOCTEXT("DontAskAgain", "Don't prompt again");
+			Toast.bUseLargeFont = false;
+			Toast.bFireAndForget = false;
+			Toast.FadeOutDuration = 0.0f;
+			Toast.ExpireDuration = 0.0f;
+			Toast.bUseThrobber = false;
+			Toast.bUseSuccessFailIcons = false;
+			Toast.ButtonDetails.Emplace(
+				LOCTEXT("OverridenOpinionMessageOk", "Ok"),
+				FText::GetEmpty(),
+				FSimpleDelegate::CreateLambda(
+					[]()
+					{
+						if (TSharedPtr<SNotificationItem> PinnedNotification = Notification.Pin())
+						{
+							PinnedNotification->SetCompletionState(SNotificationItem::CS_Success);
+							PinnedNotification->ExpireAndFadeout();
+						}
+					}
+				)
+			);
+			// This is flipped because the default checkbox message is "Don't prompt again"
+			Toast.CheckBoxState = Settings->bShowSubsectionSnappingWarning ? ECheckBoxState::Unchecked : ECheckBoxState::Checked;
+			Toast.CheckBoxStateChanged = FOnCheckStateChanged::CreateStatic(
+				[](ECheckBoxState NewState)
+				{
+					if (UUsdProjectSettings* Settings = GetMutableDefault<UUsdProjectSettings>())
+					{
+						// This is flipped because the default checkbox message is "Don't prompt again"
+						Settings->bShowSubsectionSnappingWarning = NewState == ECheckBoxState::Unchecked;
+						Settings->SaveConfig();
+					}
+				}
+			);
+
+			// Only show one at a time
+			if (!Notification.IsValid())
+			{
+				Notification = FSlateNotificationManager::Get().AddNotification(Toast);
+			}
+
+			if (TSharedPtr<SNotificationItem> PinnedNotification = Notification.Pin())
+			{
+				PinnedNotification->SetCompletionState(SNotificationItem::CS_Pending);
+			}
+		}
+	}
+}
+
 void FUsdLevelSequenceHelperImpl::CreateTimeTrack(const FLayerTimeInfo& Info)
 {
 	ULevelSequence* Sequence = FindSequenceForIdentifier(Info.Identifier);
@@ -1732,7 +1879,9 @@ void FUsdLevelSequenceHelperImpl::CreateTimeTrack(const FLayerTimeInfo& Info)
 		MovieScene->SetEvaluationType(EMovieSceneEvaluationType::FrameLocked);
 	}
 
-	if (Info.IsAnimated())
+	// Always setup the time track even if the layer is "not animated" as we need this to refresh
+	// the start/end keyframes whenever we resize the playback range to being [0, 0] as will (which
+	// would be considered "not animated")
 	{
 		const double StartTimeCode = Info.StartTimeCode.Get(0.0);
 		const double EndTimeCode = Info.EndTimeCode.Get(0.0);
@@ -1762,6 +1911,11 @@ void FUsdLevelSequenceHelperImpl::CreateTimeTrack(const FLayerTimeInfo& Info)
 			FMovieSceneFloatChannel* TimeChannel = TimeSection->GetChannelProxy().GetChannel<FMovieSceneFloatChannel>(0);
 			TimeChannel->Set(FrameNumbers, FrameValues);
 
+			// It's probably for the best to always keep this "read-only" because all of our animations are coming from
+			// the actual parsed Sequencer tracks now, so if a user edits this it will likely have no effect on the animation
+			// that is actually visible
+			TimeSection->SetIsLocked(true);
+
 			RefreshSequencer();
 		}
 	}
@@ -3909,6 +4063,7 @@ void FUsdLevelSequenceHelperImpl::OnObjectTransacted(UObject* Object, const clas
 	else if (UMovieSceneSubSection* SubSection = Cast<UMovieSceneSubSection>(Object))
 	{
 		HandleSubSectionChange(*SubSection);
+		SnapSubSectionEndsToPlaybackRange();
 	}
 	else if (UMovieSceneTrack* Track = Cast<UMovieSceneTrack>(Object))
 	{
@@ -4014,7 +4169,13 @@ void FUsdLevelSequenceHelperImpl::OnUsdObjectsChanged(
 	bool bUpdated = IterateChanges(InfoChanges);
 	if (!bUpdated)
 	{
-		IterateChanges(ResyncChanges);
+		bUpdated = IterateChanges(ResyncChanges);
+	}
+
+	if (bUpdated)
+	{
+		SnapSubSectionEndsToPlaybackRange();
+		CreateTimeTrack(FindOrAddLayerTimeInfo(UsdStage.GetRootLayer()));
 	}
 }
 
@@ -4334,6 +4495,9 @@ void FUsdLevelSequenceHelperImpl::HandleMovieSceneChange(UMovieScene& MovieScene
 
 	// We may have changed things like playback ranges, so refresh the prim properties panel if relevant
 	StageActorPtr->OnPrimChanged.Broadcast(UE::FSdfPath::AbsoluteRootPath().GetString(), false);
+
+	const bool bShowToast = true;
+	SnapSubSectionEndsToPlaybackRange(bShowToast);
 }
 
 void FUsdLevelSequenceHelperImpl::HandleSubSectionChange(UMovieSceneSubSection& Section)
@@ -4583,12 +4747,21 @@ void FUsdLevelSequenceHelperImpl::HandleTrackChange(UMovieSceneTrack& Track, boo
 	// crashes (e.g. UE-215067)
 	UMovieScenePropertyTrack* PropertyTrack = Cast<UMovieScenePropertyTrack>(&Track);
 	const FName PropertyPath = PropertyTrack ? PropertyTrack->GetPropertyPath() : NAME_None;
-	if (BoundObject == StageActor.Get() && PropertyPath != GET_MEMBER_NAME_CHECKED(AUsdStageActor, Time))
+	if (BoundObject == StageActor.Get())
 	{
+		if (PropertyPath == GET_MEMBER_NAME_CHECKED(AUsdStageActor, Time))
+		{
+			// Discard any edits to the Time track
+			CreateTimeTrack(FindOrAddLayerTimeInfo(UsdStage.GetRootLayer()));
+		}
+		else
+		{
 #if WITH_EDITOR
-		UsdLevelSequenceHelperImpl::ShowStageActorPropertyTrackWarning(PropertyPath);
+			UsdLevelSequenceHelperImpl::ShowStageActorPropertyTrackWarning(PropertyPath);
 #endif	  // WITH_EDITOR
-		MovieScene->RemoveTrack(*PropertyTrack);
+			MovieScene->RemoveTrack(*PropertyTrack);
+		}
+
 		return;
 	}
 
diff --git a/Engine/Plugins/Runtime/USDCore/Source/USDClasses/Public/USDProjectSettings.h b/Engine/Plugins/Runtime/USDCore/Source/USDClasses/Public/USDProjectSettings.h
--- a/Engine/Plugins/Runtime/USDCore/Source/USDClasses/Public/USDProjectSettings.h
+++ b/Engine/Plugins/Runtime/USDCore/Source/USDClasses/Public/USDProjectSettings.h
@@ -108,6 +108,10 @@ public:
 	UPROPERTY(config, EditAnywhere, Category = "USD|Dialogs")
 	bool bShowTransformTrackOnCameraComponentWarning = true;
 
+	// Whether to show the warning dialog when snapping a subsequence section to the playback range
+	UPROPERTY(config, EditAnywhere, Category = "USD|Dialogs")
+	bool bShowSubsectionSnappingWarning = true;
+
 	// Whether to show the warning dialog when authoring new visibility tracks from Unreal
 	UPROPERTY(config, EditAnywhere, Category = "USD|Dialogs")
 	bool bShowInheritedVisibilityWarning = true;
