USD: Fix crash when calling AbandonFailedAsset from Async task chains.

diff --git a/Engine/Plugins/Importers/USDImporter/Source/USDSchemas/Private/USDTranslatorUtils.cpp b/Engine/Plugins/Importers/USDImporter/Source/USDSchemas/Private/USDTranslatorUtils.cpp
--- a/Engine/Plugins/Importers/USDImporter/Source/USDSchemas/Private/USDTranslatorUtils.cpp
+++ b/Engine/Plugins/Importers/USDImporter/Source/USDSchemas/Private/USDTranslatorUtils.cpp
@@ -6,6 +6,7 @@
 #include "USDPrimLinkCache.h"
 
 #include "AssetRegistry/AssetRegistryModule.h"
+#include "Containers/Ticker.h"
 
 void UsdUnreal::TranslatorUtils::AbandonFailedAsset(UObject* Asset, UUsdAssetCache3* AssetCache, FUsdPrimLinkCache* PrimLinkCache)
 {
@@ -14,15 +15,6 @@ void UsdUnreal::TranslatorUtils::AbandonFailedAsset(UObject* Asset, UUsdAssetCac
 		return;
 	}
 
-	// These come from the internals of ObjectTools::DeleteSingleObject
-	Asset->MarkPackageDirty();
-#if WITH_EDITOR
-	FAssetRegistryModule::AssetDeleted(Asset);
-#endif	  // WITH_EDITOR
-	Asset->ClearFlags(RF_Standalone | RF_Public);
-
-	Asset->MarkAsGarbage();
-
 	if (AssetCache)
 	{
 		const FString Hash = AssetCache->GetHashForAsset(Asset);
@@ -36,4 +28,24 @@ void UsdUnreal::TranslatorUtils::AbandonFailedAsset(UObject* Asset, UUsdAssetCac
 	{
 		PrimLinkCache->RemoveAllAssetPrimLinks(Asset);
 	}
+
+	// We can't call MarkPackageDirty() from an async thread, and sometimes we call AbandonFailedAsset() from schema translator task chains
+	TWeakObjectPtr<UObject> AssetPtr{Asset};
+	ExecuteOnGameThread(
+		UE_SOURCE_LOCATION,
+		[AssetPtr]()
+		{
+			if (UObject* Asset = AssetPtr.Get())
+			{
+				Asset->ClearFlags(RF_Standalone | RF_Public);
+
+				// These come from the internals of ObjectTools::DeleteSingleObject
+				Asset->MarkPackageDirty();
+#if WITH_EDITOR
+				FAssetRegistryModule::AssetDeleted(Asset);
+#endif	  // WITH_EDITOR
+				Asset->MarkAsGarbage();
+			}
+		}
+	);
 }
