USD: Add Actions menu items to select stage actor and open its levelsequence.

diff --git a/Engine/Plugins/Importers/USDImporter/Source/USDStage/Private/USDStageActor.cpp b/Engine/Plugins/Importers/USDImporter/Source/USDStage/Private/USDStageActor.cpp
--- a/Engine/Plugins/Importers/USDImporter/Source/USDStage/Private/USDStageActor.cpp
+++ b/Engine/Plugins/Importers/USDImporter/Source/USDStage/Private/USDStageActor.cpp
@@ -2917,7 +2917,7 @@ void AUsdStageActor::SetTime(float InTime)
 	Refresh();
 }
 
-ULevelSequence* AUsdStageActor::GetLevelSequence()
+ULevelSequence* AUsdStageActor::GetLevelSequence() const
 {
 	return LevelSequence;
 }
diff --git a/Engine/Plugins/Importers/USDImporter/Source/USDStage/Public/USDStageActor.h b/Engine/Plugins/Importers/USDImporter/Source/USDStage/Public/USDStageActor.h
--- a/Engine/Plugins/Importers/USDImporter/Source/USDStage/Public/USDStageActor.h
+++ b/Engine/Plugins/Importers/USDImporter/Source/USDStage/Public/USDStageActor.h
@@ -271,7 +271,7 @@ public:
 	USDSTAGE_API void SetTime(float InTime);
 
 	UFUNCTION(BlueprintCallable, Category = "USD", meta = (CallInEditor = "true"))
-	USDSTAGE_API ULevelSequence* GetLevelSequence();
+	USDSTAGE_API ULevelSequence* GetLevelSequence() const;
 
 	/**
 	 * Gets the transient component that was generated for a prim with a given prim path.
diff --git a/Engine/Plugins/Importers/USDImporter/Source/USDStageEditor/Private/SUSDStage.cpp b/Engine/Plugins/Importers/USDImporter/Source/USDStageEditor/Private/SUSDStage.cpp
--- a/Engine/Plugins/Importers/USDImporter/Source/USDStageEditor/Private/SUSDStage.cpp
+++ b/Engine/Plugins/Importers/USDImporter/Source/USDStageEditor/Private/SUSDStage.cpp
@@ -33,12 +33,14 @@
 #include "Framework/MultiBox/MultiBoxBuilder.h"
 #include "ISceneOutliner.h"
 #include "ISceneOutlinerColumn.h"
+#include "LevelSequence.h"
 #include "Modules/ModuleManager.h"
 #include "PropertyEditorModule.h"
 #include "SceneOutlinerModule.h"
 #include "ScopedTransaction.h"
 #include "Selection.h"
 #include "Styling/AppStyle.h"
+#include "Subsystems/AssetEditorSubsystem.h"
 #include "Widgets/Input/SButton.h"
 #include "Widgets/Input/SComboBox.h"
 #include "Widgets/Input/SComboButton.h"
@@ -726,7 +728,8 @@ TSharedRef<SWidget> SUsdStage::MakeActorPickerMenuContent()
 				}
 			),
 			bCanBeHidden,
-			FillSize};
+			FillSize
+		};
 
 		FSceneOutlinerInitializationOptions InitOptions;
 		InitOptions.bShowHeaderRow = true;	  // The header lets us resize/hide the columns if we want
@@ -1031,6 +1034,49 @@ void SUsdStage::FillActionsMenu(FMenuBuilder& MenuBuilder)
 		);
 	}
 	MenuBuilder.EndSection();
+
+	MenuBuilder.BeginSection("Stage Actor", LOCTEXT("Stage Actor", "Stage Actor"));
+	{
+		MenuBuilder.AddMenuEntry(
+			LOCTEXT("SelectStageActor", "Select stage actor"),
+			LOCTEXT("SelectStageActor_ToolTip", "Selects the currently attached USD Stage Actor"),
+			FSlateIcon(),
+			FUIAction(
+				FExecuteAction::CreateSP(this, &SUsdStage::ActionsSelectStageActor),
+				FCanExecuteAction::CreateLambda(
+					[this]()
+					{
+						return ViewModel.UsdStageActor.Get() != nullptr;
+					}
+				)
+			),
+			NAME_None,
+			EUserInterfaceActionType::Button
+		);
+
+		MenuBuilder.AddMenuEntry(
+			LOCTEXT("OpenLevelSequence", "Open level sequence"),
+			LOCTEXT("OpenLevelSequence_ToolTip", "Opens the generated LevelSequence asset for this stage"),
+			FSlateIcon(),
+			FUIAction(
+				FExecuteAction::CreateSP(this, &SUsdStage::ActionsOpenLevelSequence),
+				FCanExecuteAction::CreateLambda(
+					[this]()
+					{
+						if (const AUsdStageActor* StageActor = ViewModel.UsdStageActor.Get())
+						{
+							return static_cast<bool>(StageActor->GetUsdStage()) && (StageActor->GetLevelSequence() != nullptr);
+						}
+
+						return false;
+					}
+				)
+			),
+			NAME_None,
+			EUserInterfaceActionType::Button
+		);
+	}
+	MenuBuilder.EndSection();
 }
 
 void SUsdStage::FillOptionsMenu(FMenuBuilder& MenuBuilder)
@@ -1499,7 +1545,8 @@ void SUsdStage::FillMaterialPurposeSubMenu(FMenuBuilder& MenuBuilder)
 	MaterialPurposes = {
 		MakeShared<FString>(UnrealIdentifiers::MaterialAllPurpose),
 		MakeShared<FString>(UnrealIdentifiers::MaterialPreviewPurpose),
-		MakeShared<FString>(UnrealIdentifiers::MaterialFullPurpose)};
+		MakeShared<FString>(UnrealIdentifiers::MaterialFullPurpose)
+	};
 
 	// Add additional purposes from project settings
 	if (const UUsdProjectSettings* ProjectSettings = GetDefault<UUsdProjectSettings>())
@@ -1509,7 +1556,8 @@ void SUsdStage::FillMaterialPurposeSubMenu(FMenuBuilder& MenuBuilder)
 		TSet<FString> ExistingEntries = {
 			UnrealIdentifiers::MaterialAllPurpose,
 			UnrealIdentifiers::MaterialPreviewPurpose,
-			UnrealIdentifiers::MaterialFullPurpose};
+			UnrealIdentifiers::MaterialFullPurpose
+		};
 
 		for (const FName& AdditionalPurpose : ProjectSettings->AdditionalMaterialPurposes)
 		{
@@ -2774,6 +2822,48 @@ void SUsdStage::ActionsRegenerate()
 	}
 }
 
+void SUsdStage::ActionsSelectStageActor()
+{
+	AUsdStageActor* StageActor = ViewModel.UsdStageActor.Get();
+	if (!StageActor)
+	{
+		return;
+	}
+
+	// Don't use GEditor->SelectNone() as that will affect *every type of selection in the editor*,
+	// including even some UI menus, brushes, etc.
+	if (USelection* SelectedActors = GEditor->GetSelectedActors())
+	{
+		SelectedActors->DeselectAll(AActor::StaticClass());
+	}
+	if (USelection* SelectedComponents = GEditor->GetSelectedComponents())
+	{
+		SelectedComponents->DeselectAll(UActorComponent::StaticClass());
+	}
+
+	const bool bSelected = true;
+	const bool bNotifySelectionChanged = true;
+	GEditor->SelectActor(StageActor, bSelected, bNotifySelectionChanged);
+}
+
+void SUsdStage::ActionsOpenLevelSequence()
+{
+	AUsdStageActor* StageActor = ViewModel.UsdStageActor.Get();
+	if (!StageActor)
+	{
+		return;
+	}
+
+	ULevelSequence* LevelSequence = StageActor->GetLevelSequence();
+	if (!LevelSequence)
+	{
+		return;
+	}
+
+	UAssetEditorSubsystem* AssetEditorSubsystem = GEditor ? GEditor->GetEditorSubsystem<UAssetEditorSubsystem>() : nullptr;
+	AssetEditorSubsystem->OpenEditorForAsset(LevelSequence);
+}
+
 void SUsdStage::ExportSelectedLayers(const FString& OutputLayerOrDirectory)
 {
 	if (UsdLayersTreeView)
diff --git a/Engine/Plugins/Importers/USDImporter/Source/USDStageEditor/Private/SUSDStage.h b/Engine/Plugins/Importers/USDImporter/Source/USDStageEditor/Private/SUSDStage.h
--- a/Engine/Plugins/Importers/USDImporter/Source/USDStageEditor/Private/SUSDStage.h
+++ b/Engine/Plugins/Importers/USDImporter/Source/USDStageEditor/Private/SUSDStage.h
@@ -73,6 +73,8 @@ public:
 	void ActionsImportWithDialog();
 	void ActionsImport(const FString& OutputContentFolder, UUsdStageImportOptions* Options);
 	void ActionsRegenerate();
+	void ActionsSelectStageActor();
+	void ActionsOpenLevelSequence();
 	void ExportSelectedLayers(const FString& OutputLayerOrDirectory = {});
 
 protected:
diff --git a/Engine/Plugins/Importers/USDImporter/Source/USDStageEditor/USDStageEditor.Build.cs b/Engine/Plugins/Importers/USDImporter/Source/USDStageEditor/USDStageEditor.Build.cs
--- a/Engine/Plugins/Importers/USDImporter/Source/USDStageEditor/USDStageEditor.Build.cs
+++ b/Engine/Plugins/Importers/USDImporter/Source/USDStageEditor/USDStageEditor.Build.cs
@@ -16,10 +16,11 @@ namespace UnrealBuildTool.Rules
 					"CoreUObject",
 					"DesktopPlatform",
 					"DesktopWidgets",
-					"EditorFramework",	
+					"EditorFramework",
 					"Engine",
 					"InputCore",
 					"LevelEditor",
+					"LevelSequence",
 					"LiveLinkEditor", // For SLiveLinkSubjectRepresentationPicker
 					"LiveLinkInterface", // For the ULiveLinkAnimation/TransformRole classes
 					"Projects", // So that we can use the IPluginManager, required for our custom style
