USD: Fix more issues related to LevelSequence handling, sublayers, offsets and scales.

Fix subsection ranges sometimes ending beyond the stage's playback range;
Use StartFrameOffset to setup sublayer playback ranges according to their USD layers;
Fix section start times not being updated when sublayer offset/scales were updated via USD;
Fix reloading no longer recreating LevelSequences;
Fix LevelSequence repopulation triggering stage changes;
Add automated tests for sublayer/reference offset/scale cases.

diff --git a/Engine/Plugins/Importers/USDImporter/Source/USDStage/Private/USDLevelSequenceHelper.cpp b/Engine/Plugins/Importers/USDImporter/Source/USDStage/Private/USDLevelSequenceHelper.cpp
--- a/Engine/Plugins/Importers/USDImporter/Source/USDStage/Private/USDLevelSequenceHelper.cpp
+++ b/Engine/Plugins/Importers/USDImporter/Source/USDStage/Private/USDLevelSequenceHelper.cpp
@@ -869,7 +869,7 @@ private:
 	/**
 	 * Adjusts all subsection ends, in all sequences, to match the stage's root layer's endTimeCode
 	 */
-	void SnapSubSectionEndsToPlaybackRange(bool bShowToast = false);
+	void UpdateSubSectionTimeRanges(bool bShowToast = false);
 
 	// Tracks handling
 private:
@@ -1075,7 +1075,10 @@ ULevelSequence* FUsdLevelSequenceHelperImpl::Init(const UE::FUsdStage& InUsdStag
 {
 	UsdStage = InUsdStage;
 
+	Clear();
+
 	CreateLocalLayersSequences();
+
 	return MainLevelSequence;
 }
 
@@ -1621,9 +1624,7 @@ void FUsdLevelSequenceHelperImpl::CreateSubSequenceSection(ULevelSequence& Seque
 		const double StageEndTimeSeconds = UsdStage.GetEndTimeCode() / UsdStage.GetTimeCodesPerSecond();
 		const FFrameNumber StageEndFrame = UsdLevelSequenceHelperImpl::RoundAsFrameNumber(TickResolution, StageEndTimeSeconds);
 
-		// Max here because StartFrame can theoretically be larger than StageEndFrame, which would generate a range where the upper bound is smaller
-		// than the lower bound, which can trigger asserts
-		SubSectionRange = TRange<FFrameNumber>{StartFrame, FMath::Max(StageEndFrame, EndFrame)};
+		SubSectionRange = TRange<FFrameNumber>{StartFrame, FMath::Max(StartFrame, StageEndFrame)};
 	}
 	else
 	{
@@ -1676,6 +1677,15 @@ void FUsdLevelSequenceHelperImpl::CreateSubSequenceSection(ULevelSequence& Seque
 	const double TimeCodesPerSecondDifference = TimeCodesPerSecond / SubLayer.GetTimeCodesPerSecond();
 	SubSection->Parameters.TimeScale = FMath::IsNearlyZero(SubLayerOffset.Scale) ? 0.f : 1.f / (SubLayerOffset.Scale / TimeCodesPerSecondDifference);
 
+	// As far as the Sequencer is concerned, the subsection "starts" at it's playback range start (i.e. if the inner playback range is
+	// [5, 20] and we place the subsection at timeCode 14, it will try making it so that at outer timeCode 14 it plays the inner timeCode 5).
+	// To match USD composition, we instead want the subsection to always start at inner timeCode zero (i.e. at outer timeCode 14 it plays
+	// the inner timeCode 0). The only way of doing that is by specifying a StartFrameOffset for the subsection, that matches the playback range.
+	//
+	// These will be edited together by the Sequencer (so e.g. if a user drags the inner playback range it will also update the start frame offset),
+	// we just have to set them up properly the first time around it seems.
+	SubSection->Parameters.StartFrameOffset = -SubSequence.GetMovieScene()->GetPlaybackRange().GetLowerBoundValue();
+
 	if (MainLevelSequence)
 	{
 		UMovieSceneCompiledDataManager::CompileHierarchy(MainLevelSequence, &SequenceHierarchyCache, EMovieSceneServerClientMask::All);
@@ -1712,8 +1722,11 @@ void FUsdLevelSequenceHelperImpl::RemoveSubSequenceSection(ULevelSequence& Seque
 	}
 }
 
-void FUsdLevelSequenceHelperImpl::SnapSubSectionEndsToPlaybackRange(bool bShowToast)
+void FUsdLevelSequenceHelperImpl::UpdateSubSectionTimeRanges(bool bShowResizedSectionToast)
 {
+	// This is in charge of updating *all* our SubSection sizes, TimeScale, as well as their StartTimeOffsets according to the stage.
+	// This does not handle playback ranges, which are updated by UpdateMovieSceneTimeRanges.
+	//
 	// In USD, all endTimeCode values are ignored except those for the root layer. Even those don't "hard clip" the animation either,
 	// and are more like a suggestion of what the play range should be.
 	// In UE, subsection ranges do in fact clip the inner animation. We can't specify an unbounded range for them though, and
@@ -1727,7 +1740,7 @@ void FUsdLevelSequenceHelperImpl::SnapSubSectionEndsToPlaybackRange(bool bShowTo
 
 	const double StageEndTimeSeconds = UsdStage.GetEndTimeCode() / UsdStage.GetTimeCodesPerSecond();
 
-	bool bChangedSomeSection = false;
+	bool bChangedSectionSize = false;
 
 	for (const TPair<FString, TObjectPtr<ULevelSequence>>& IdentifierAndSeq : LevelSequencesByIdentifier)
 	{
@@ -1750,31 +1763,78 @@ void FUsdLevelSequenceHelperImpl::SnapSubSectionEndsToPlaybackRange(bool bShowTo
 		}
 
 		FFrameRate TickResolution = MovieScene->GetTickResolution();
+		double TimeCodesPerSecond = GetTimeCodesPerSecond();
 		const FFrameNumber StageEndFrame = UsdLevelSequenceHelperImpl::RoundAsFrameNumber(TickResolution, StageEndTimeSeconds);
 
+		FLayerTimeInfo* LayerTimeInfo = LayerTimeInfosByLayerIdentifier.Find(*IdentifierAndSeq.Key);
+
 		for (UMovieSceneSection* Section : SubTrack->GetAllSections())
 		{
 			if (UMovieSceneSubSection* SubSection = Cast<UMovieSceneSubSection>(Section))
 			{
-				TRange<FFrameNumber> Range = SubSection->GetRange();
+				UMovieSceneSequence* SubSequence = SubSection->GetSequence();
+				if (!SubSequence)
+				{
+					continue;
+				}
 
-				FFrameNumber NewUpper = FMath::Max(StageEndFrame, Range.GetLowerBoundValue());
-				if (NewUpper == Range.GetUpperBoundValue())
+				FString* SubLayerIdentifier = IdentifierByLevelSequence.Find(Cast<ULevelSequence>(SubSequence));
+				if (!SubLayerIdentifier)
 				{
 					continue;
 				}
 
-				Range.SetUpperBoundValue(NewUpper);
+				UE::FSdfLayer SubLayer = UE::FSdfLayer::FindOrOpen(**SubLayerIdentifier);
+				if (!SubLayer)
+				{
+					continue;
+				}
+
+				// Find our current sublayer offset for this sublayer
+				FLayerOffsetInfo* SubLayerOffset = Algo::FindByPredicate(
+					LayerTimeInfo->SubLayersOffsets,
+					[SubLayerIdentifier](const FLayerOffsetInfo& Other)
+					{
+						return Other.LayerIdentifier == *SubLayerIdentifier;
+					}
+				);
+				if (!SubLayerOffset)
+				{
+					continue;
+				}
 
 				SubSection->Modify();
+
+				// StartFrameOffset: Should match the sublayer's playback range start so that the inner sequence starts playing at its zero
+				SubSection->Parameters.StartFrameOffset = -SubSequence->GetMovieScene()->GetPlaybackRange().GetLowerBoundValue();
+
+				// TimeScale: Should match the sublayer's scale (or the inverse of it, rather)
+				const double TimeCodesPerSecondDifference = TimeCodesPerSecond / SubLayer.GetTimeCodesPerSecond();
+				SubSection->Parameters.TimeScale = FMath::IsNearlyZero(SubLayerOffset->LayerOffset.Scale)
+													   ? 0.f
+													   : 1.f / (SubLayerOffset->LayerOffset.Scale / TimeCodesPerSecondDifference);
+
+				// Section start/end: The start should match the sublayer's offset, and the end should always clip to the stage's endTimeCode
+				TRange<FFrameNumber> Range = SubSection->GetRange();
+				FFrameNumber NewLower = UsdLevelSequenceHelperImpl::RoundAsFrameNumber(
+					TickResolution,
+					SubLayerOffset->LayerOffset.Offset / TimeCodesPerSecond
+				);
+				FFrameNumber NewUpper = FMath::Max(StageEndFrame, NewLower);
+				if (NewLower == Range.GetLowerBoundValue() && NewUpper == Range.GetUpperBoundValue())
+				{
+					continue;
+				}
+				Range.SetLowerBoundValue(NewLower);
+				Range.SetUpperBoundValue(NewUpper);
 				SubSection->SetRange(Range);
 
-				bChangedSomeSection = true;
+				bChangedSectionSize = true;
 			}
 		}
 	}
 
-	if (bChangedSomeSection && bShowToast)
+	if (bChangedSectionSize && bShowResizedSectionToast)
 	{
 		const UUsdProjectSettings* Settings = GetDefault<UUsdProjectSettings>();
 		if (Settings && Settings->bShowSubsectionSnappingWarning)
@@ -3861,7 +3921,7 @@ void FUsdLevelSequenceHelperImpl::UpdateUsdLayerOffsetFromSection(const UMovieSc
 
 	UE::FSdfLayerOffset NewLayerOffset;
 	NewLayerOffset.Scale = FMath::IsNearlyZero(FixedPlayRate) ? 0.f : 1.f / FixedPlayRate;
-	NewLayerOffset.Offset = SubSectionStartTimeCode - SubStartTimeCode * NewLayerOffset.Scale;
+	NewLayerOffset.Offset = SubSectionStartTimeCode;
 
 	if (FMath::IsNearlyZero(NewLayerOffset.Offset))
 	{
@@ -3937,14 +3997,6 @@ void FUsdLevelSequenceHelperImpl::UpdateMovieSceneTimeRanges(UMovieScene& MovieS
 		if (UE::FSdfLayer Layer = UE::FSdfLayer::FindOrOpen(*LayerTimeInfo.Identifier))
 		{
 			TimeCodesPerSecond = Layer.GetTimeCodesPerSecond();
-
-			// When composing a sublayer that has startTimeCode 10 with an offset of 25 timecodes, USD will place the sublayer's
-			// time code 0 at the 25 mark. We want to mirror that behavior when composing our subsections but still leave the root layer's
-			// playback range to [startTimeCode, endTimeCode] as that's what we'd expect to see, and it doesn't affect composition
-			if (Layer != UsdStage.GetRootLayer())
-			{
-				StartTimeCode = 0.0;
-			}
 		}
 		else
 		{
@@ -4063,7 +4115,7 @@ void FUsdLevelSequenceHelperImpl::OnObjectTransacted(UObject* Object, const clas
 	else if (UMovieSceneSubSection* SubSection = Cast<UMovieSceneSubSection>(Object))
 	{
 		HandleSubSectionChange(*SubSection);
-		SnapSubSectionEndsToPlaybackRange();
+		UpdateSubSectionTimeRanges();
 	}
 	else if (UMovieSceneTrack* Track = Cast<UMovieSceneTrack>(Object))
 	{
@@ -4174,7 +4226,7 @@ void FUsdLevelSequenceHelperImpl::OnUsdObjectsChanged(
 
 	if (bUpdated)
 	{
-		SnapSubSectionEndsToPlaybackRange();
+		UpdateSubSectionTimeRanges();
 		CreateTimeTrack(FindOrAddLayerTimeInfo(UsdStage.GetRootLayer()));
 	}
 }
@@ -4497,7 +4549,7 @@ void FUsdLevelSequenceHelperImpl::HandleMovieSceneChange(UMovieScene& MovieScene
 	StageActorPtr->OnPrimChanged.Broadcast(UE::FSdfPath::AbsoluteRootPath().GetString(), false);
 
 	const bool bShowToast = true;
-	SnapSubSectionEndsToPlaybackRange(bShowToast);
+	UpdateSubSectionTimeRanges(bShowToast);
 }
 
 void FUsdLevelSequenceHelperImpl::HandleSubSectionChange(UMovieSceneSubSection& Section)
diff --git a/Engine/Plugins/Importers/USDImporter/Source/USDStage/Private/USDStageActor.cpp b/Engine/Plugins/Importers/USDImporter/Source/USDStage/Private/USDStageActor.cpp
--- a/Engine/Plugins/Importers/USDImporter/Source/USDStage/Private/USDStageActor.cpp
+++ b/Engine/Plugins/Importers/USDImporter/Source/USDStage/Private/USDStageActor.cpp
@@ -3559,6 +3559,13 @@ void AUsdStageActor::RegenerateLevelSequence()
 {
 	TRACE_CPUPROFILER_EVENT_SCOPE(AUsdStageActor::RegenerateLevelSequence);
 
+	// Don't check for full authority here because even if we can't write back to the stage (i.e. during PIE) we still
+	// want to listen to it and have valid level sequences
+	if (IsTemplate())
+	{
+		return;
+	}
+
 	// If we're using some property editor that can trigger a stage reload (like the Nanite threshold spinbox),
 	// applying a value may trigger RegenerateLevelSequence -> Can trigger asset editors to open/close/change focus ->
 	// -> Can trigger focus to drop from the property editors -> Can cause the values to be applied from the
@@ -3580,31 +3587,32 @@ void AUsdStageActor::RegenerateLevelSequence()
 		return;
 	}
 
+	// We will update our levelsequence objects (tracks, moviescene, sections, etc.) due to these changes.
+	// We definitely don't want to write anything back to USD when these objects change though.
+	FScopedBlockMonitoringChangesForTransaction BlockMonitoring{LevelSequenceHelper};
+
 	// Don't check for full authority here because even if we can't write back to the stage (i.e. during PIE) we still
 	// want to listen to it and have valid level sequences
-	if (!IsTemplate())
-	{
 #if WITH_EDITOR
-		UAssetEditorSubsystem* AssetEditorSubsystem = GEditor ? GEditor->GetEditorSubsystem<UAssetEditorSubsystem>() : nullptr;
-		const bool bSequencerWasOpened = AssetEditorSubsystem && AssetEditorSubsystem->FindEditorsForAssetAndSubObjects(LevelSequence).Num() > 0;
+	UAssetEditorSubsystem* AssetEditorSubsystem = GEditor ? GEditor->GetEditorSubsystem<UAssetEditorSubsystem>() : nullptr;
+	const bool bSequencerWasOpened = AssetEditorSubsystem && AssetEditorSubsystem->FindEditorsForAssetAndSubObjects(LevelSequence).Num() > 0;
 #endif	  // WITH_EDITOR
 
-		// We need to guarantee we'll record our change of LevelSequence into the transaction, as Init() will create a new one
-		const bool bMarkDirty = false;
-		Modify(bMarkDirty);
+	// We need to guarantee we'll record our change of LevelSequence into the transaction, as Init() will create a new one
+	const bool bMarkDirty = false;
+	Modify(bMarkDirty);
 
-		LevelSequence = LevelSequenceHelper.Init(CurrentStage);
-		LevelSequenceHelper.BindToUsdStageActor(this);
+	LevelSequence = LevelSequenceHelper.Init(CurrentStage);
+	LevelSequenceHelper.BindToUsdStageActor(this);
 
 #if WITH_EDITOR
-		if (bSequencerWasOpened && LevelSequence && AssetEditorSubsystem)
-		{
-			// Open the LevelSequence editor via the UAssetEditorSubsystem, otherwise calls to CloseAllEditorsForAsset
-			// won't be able to close it
-			AssetEditorSubsystem->OpenEditorForAsset(LevelSequence);
-		}
-#endif	  // WITH_EDITOR
+	if (bSequencerWasOpened && LevelSequence && AssetEditorSubsystem)
+	{
+		// Open the LevelSequence editor via the UAssetEditorSubsystem, otherwise calls to CloseAllEditorsForAsset
+		// won't be able to close it
+		AssetEditorSubsystem->OpenEditorForAsset(LevelSequence);
 	}
+#endif	  // WITH_EDITOR
 }
 
 void AUsdStageActor::RepopulateLevelSequence()
@@ -3617,6 +3625,10 @@ void AUsdStageActor::RepopulateLevelSequence()
 	);
 	TranslationContext->bIsJustRepopulatingLevelSequence = true;
 
+	// We will update our levelsequence objects (tracks, moviescene, sections, etc.) due to these changes.
+	// We definitely don't want to write anything back to USD when these objects change though.
+	FScopedBlockMonitoringChangesForTransaction BlockMonitoring{LevelSequenceHelper};
+
 	const bool bIsResync = false;
 	UpdatePrim(UE::FSdfPath::AbsoluteRootPath(), bIsResync, *TranslationContext);
 }
