USD: Fix broken build after cl 37872184

diff --git a/Engine/Plugins/Importers/USDImporter/Source/USDStage/Private/USDLevelSequenceHelper.cpp b/Engine/Plugins/Importers/USDImporter/Source/USDStage/Private/USDLevelSequenceHelper.cpp
--- a/Engine/Plugins/Importers/USDImporter/Source/USDStage/Private/USDLevelSequenceHelper.cpp
+++ b/Engine/Plugins/Importers/USDImporter/Source/USDStage/Private/USDLevelSequenceHelper.cpp
@@ -1815,7 +1815,7 @@ void FUsdLevelSequenceHelperImpl::UpdateSubSectionTimeRanges(bool bShowResizedSe
 
 				// Section start/end: The start should match the sublayer's offset, and the end should always clip to the stage's endTimeCode
 				TRange<FFrameNumber> Range = SubSection->GetRange();
-				FFrameNumber NewLower = UsdLevelSequenceHelperImpl::RoundAsFrameNumber(
+				FFrameNumber NewLower = UE::USDLevelSequenceHelper::Private::RoundAsFrameNumber(
 					TickResolution,
 					SubLayerOffset->LayerOffset.Offset / TimeCodesPerSecond
 				);
