USD: Make the LoadAll/LoadNone payload affect the existing stage and payloads.

diff --git a/Engine/Plugins/Importers/USDImporter/Source/USDStageEditor/Private/SUSDStage.cpp b/Engine/Plugins/Importers/USDImporter/Source/USDStageEditor/Private/SUSDStage.cpp
--- a/Engine/Plugins/Importers/USDImporter/Source/USDStageEditor/Private/SUSDStage.cpp
+++ b/Engine/Plugins/Importers/USDImporter/Source/USDStageEditor/Private/SUSDStage.cpp
@@ -1346,7 +1346,7 @@ void SUsdStage::FillPayloadsSubMenu(FMenuBuilder& MenuBuilder)
 {
 	MenuBuilder.AddMenuEntry(
 		LOCTEXT("LoadAll", "Load all"),
-		LOCTEXT("LoadAll_ToolTip", "Loads all payloads initially"),
+		LOCTEXT("LoadAll_ToolTip", "Loads all payloads, and make new payloads start out loaded"),
 		FSlateIcon(),
 		FUIAction(
 			FExecuteAction::CreateLambda(
@@ -1355,14 +1355,15 @@ void SUsdStage::FillPayloadsSubMenu(FMenuBuilder& MenuBuilder)
 					if (AUsdStageActor* StageActor = GetStageActorOrCDO())
 					{
 						FScopedTransaction Transaction(FText::Format(
-							LOCTEXT("SetLoadAllTransaction", "Set USD stage actor '{0}' actor to load all payloads initially"),
+							LOCTEXT("SetLoadAllTransaction", "Set USD stage actor '{0}' actor to load all payloads"),
 							FText::FromString(StageActor->GetActorLabel())
 						));
 
 						// c.f. comment in SUsdStage::FillCollapsingSubMenu
 						TGuardValue<bool> MaintainSelectionGuard(bUpdatingViewportSelection, true);
 
-						StageActor->SetInitialLoadSet(EUsdInitialLoadSet::LoadAll);
+						ViewModel.SetLoadAllRule();
+
 						if (StageActor->IsTemplate())
 						{
 							StageActor->SaveConfig();
@@ -1376,7 +1377,10 @@ void SUsdStage::FillPayloadsSubMenu(FMenuBuilder& MenuBuilder)
 				{
 					if (AUsdStageActor* StageActor = GetStageActorOrCDO())
 					{
-						return StageActor->InitialLoadSet == EUsdInitialLoadSet::LoadAll;
+						if (!StageActor->GetUsdStage() || ViewModel.HasDefaultLoadRule())
+						{
+							return StageActor->InitialLoadSet == EUsdInitialLoadSet::LoadAll;
+						}
 					}
 					return false;
 				}
@@ -1388,7 +1392,7 @@ void SUsdStage::FillPayloadsSubMenu(FMenuBuilder& MenuBuilder)
 
 	MenuBuilder.AddMenuEntry(
 		LOCTEXT("LoadNone", "Load none"),
-		LOCTEXT("LoadNone_ToolTip", "Don't load any payload initially"),
+		LOCTEXT("LoadNone_ToolTip", "Unload all payloads, and make new payloads start out unloaded"),
 		FSlateIcon(),
 		FUIAction(
 			FExecuteAction::CreateLambda(
@@ -1397,14 +1401,15 @@ void SUsdStage::FillPayloadsSubMenu(FMenuBuilder& MenuBuilder)
 					if (AUsdStageActor* StageActor = GetStageActorOrCDO())
 					{
 						FScopedTransaction Transaction(FText::Format(
-							LOCTEXT("SetLoadNoneTransaction", "Set USD stage actor '{0}' actor to load no payloads initially"),
+							LOCTEXT("SetLoadNoneTransaction", "Set USD stage actor '{0}' actor to load no payloads"),
 							FText::FromString(StageActor->GetActorLabel())
 						));
 
 						// c.f. comment in SUsdStage::FillCollapsingSubMenu
 						TGuardValue<bool> MaintainSelectionGuard(bUpdatingViewportSelection, true);
 
-						StageActor->SetInitialLoadSet(EUsdInitialLoadSet::LoadNone);
+						ViewModel.SetLoadNoneRule();
+
 						if (StageActor->IsTemplate())
 						{
 							StageActor->SaveConfig();
@@ -1418,7 +1423,10 @@ void SUsdStage::FillPayloadsSubMenu(FMenuBuilder& MenuBuilder)
 				{
 					if (AUsdStageActor* StageActor = GetStageActorOrCDO())
 					{
-						return StageActor->InitialLoadSet == EUsdInitialLoadSet::LoadNone;
+						if (!StageActor->GetUsdStage() || ViewModel.HasDefaultLoadRule())
+						{
+							return StageActor->InitialLoadSet == EUsdInitialLoadSet::LoadNone;
+						}
 					}
 					return false;
 				}
diff --git a/Engine/Plugins/Importers/USDImporter/Source/USDStageEditorViewModels/Private/USDStageViewModel.cpp b/Engine/Plugins/Importers/USDImporter/Source/USDStageEditorViewModels/Private/USDStageViewModel.cpp
--- a/Engine/Plugins/Importers/USDImporter/Source/USDStageEditorViewModels/Private/USDStageViewModel.cpp
+++ b/Engine/Plugins/Importers/USDImporter/Source/USDStageEditorViewModels/Private/USDStageViewModel.cpp
@@ -96,6 +96,81 @@ void FUsdStageViewModel::OpenStage(const TCHAR* FilePath)
 	UsdUtils::ShowErrorsAndStopMonitoring();
 }
 
+void FUsdStageViewModel::SetLoadAllRule()
+{
+#if USE_USD_SDK
+	if (!UsdStageActor.IsValid())
+	{
+		return;
+	}
+
+	// This will set what happens to payloads on any other future stages this stage actor opens
+	UsdStageActor->SetInitialLoadSet(EUsdInitialLoadSet::LoadAll);
+
+	// This will toggle all current payloads, and set what happens to any future payloads that
+	// are created on the currently opened stage
+	if (UE::FUsdStage Stage = UsdStageActor->GetOrOpenUsdStage())
+	{
+		FScopedUsdAllocs Allocs;
+
+		pxr::SdfChangeBlock Block;
+
+		pxr::UsdStageRefPtr UsdStage = pxr::UsdStageRefPtr(Stage);
+		UsdStage->SetLoadRules(pxr::UsdStageLoadRules::LoadAll());
+	}
+#endif	  // USE_USD_SDK
+}
+
+void FUsdStageViewModel::SetLoadNoneRule()
+{
+#if USE_USD_SDK
+	if (!UsdStageActor.IsValid())
+	{
+		return;
+	}
+
+	UsdStageActor->SetInitialLoadSet(EUsdInitialLoadSet::LoadNone);
+
+	// This will toggle all current payloads, and set what happens to any future payloads that
+	// are created on the currently opened stage
+	if (UE::FUsdStage Stage = UsdStageActor->GetOrOpenUsdStage())
+	{
+		FScopedUsdAllocs Allocs;
+
+		pxr::SdfChangeBlock Block;
+
+		pxr::UsdStageRefPtr UsdStage = pxr::UsdStageRefPtr(Stage);
+		UsdStage->SetLoadRules(pxr::UsdStageLoadRules::LoadNone());
+	}
+#endif	  // USE_USD_SDK
+}
+
+bool FUsdStageViewModel::HasDefaultLoadRule()
+{
+#if USE_USD_SDK
+	if (!UsdStageActor.IsValid())
+	{
+		return true;
+	}
+
+	// This will toggle all current payloads, and set what happens to any future payloads that
+	// are created on the currently opened stage
+	if (UE::FUsdStage Stage = UsdStageActor->GetOrOpenUsdStage())
+	{
+		FScopedUsdAllocs Allocs;
+
+		pxr::UsdStageRefPtr UsdStage = pxr::UsdStageRefPtr(Stage);
+
+		pxr::UsdStageLoadRules Rules = UsdStage->GetLoadRules();
+		Rules.Minimize();
+
+		return Rules == pxr::UsdStageLoadRules::LoadNone() || Rules == pxr::UsdStageLoadRules::LoadAll();
+	}
+#endif	  // USE_USD_SDK
+
+	return true;
+}
+
 void FUsdStageViewModel::ReloadStage()
 {
 	if (!UsdStageActor.IsValid())
diff --git a/Engine/Plugins/Importers/USDImporter/Source/USDStageEditorViewModels/Public/USDStageViewModel.h b/Engine/Plugins/Importers/USDImporter/Source/USDStageEditorViewModels/Public/USDStageViewModel.h
--- a/Engine/Plugins/Importers/USDImporter/Source/USDStageEditorViewModels/Public/USDStageViewModel.h
+++ b/Engine/Plugins/Importers/USDImporter/Source/USDStageEditorViewModels/Public/USDStageViewModel.h
@@ -14,6 +14,12 @@ class USDSTAGEEDITORVIEWMODELS_API FUsdStageViewModel
 public:
 	void NewStage();
 	void OpenStage(const TCHAR* FilePath);
+	void SetLoadAllRule();
+	void SetLoadNoneRule();
+	// Returns false if we have a payload in the current stage that has been manually loaded or unloaded. Returns true otherwise.
+	// The intent is to get the LoadAll/LoadNone buttons to *both* show as unchecked in case the user manually toggled any,
+	// to indicate that picking either LoadAll or LoadNone would trigger some change.
+	bool HasDefaultLoadRule();
 	void ReloadStage();
 	void ResetStage();
 	void CloseStage();
