USD: Allow undo/redoing apiSchema/reference/payload changes.

diff --git a/Engine/Plugins/Importers/USDImporter/Source/USDStage/Private/USDStageActor.cpp b/Engine/Plugins/Importers/USDImporter/Source/USDStage/Private/USDStageActor.cpp
--- a/Engine/Plugins/Importers/USDImporter/Source/USDStage/Private/USDStageActor.cpp
+++ b/Engine/Plugins/Importers/USDImporter/Source/USDStage/Private/USDStageActor.cpp
@@ -1435,10 +1435,57 @@ void AUsdStageActor::HandleAccumulatedNotices()
 
 		for (const TPair<FString, TArray<UsdUtils::FSdfChangeListEntry>>& InfoChange : AccumulatedInfoChanges)
 		{
-			UE::FSdfPath PrimPath = UE::FSdfPath(*InfoChange.Key).StripAllVariantSelections();
+			UE::FSdfPath ObjectPath{*InfoChange.Key};
 
-			// Upgrade some info changes into resync changes
+			UE::FSdfPath PrimPath = ObjectPath.GetAbsoluteRootOrPrimPath().StripAllVariantSelections();
+
+			FString PropertyName;
+			if (ObjectPath.IsPropertyPath())
+			{
+				PropertyName = ObjectPath.GetName();
+			}
+
+			// Upgrade these to resync so that the prim twins are regenerated, which clears all the existing
+			// animation tracks and adds new ones, automatically re-baking to control rig
 			bool bIsResync = false;
+			static const TSet<FString> PrimResyncProperties = {
+				*UsdToUnreal::ConvertToken(UnrealIdentifiers::UnrealControlRigPath),
+				*UsdToUnreal::ConvertToken(UnrealIdentifiers::UnrealUseFKControlRig),
+				*UsdToUnreal::ConvertToken(UnrealIdentifiers::UnrealControlRigReduceKeys),
+				*UsdToUnreal::ConvertToken(UnrealIdentifiers::UnrealControlRigReductionTolerance),
+				// For now we need to do this as we need to refresh the material slot info on the info cache if these
+				// update... this is of course way too aggressive, although it's unlikely people will be manually editing these.
+				// TODO: More nuanced info cache updates
+				UnrealIdentifiers::PrimvarsDisplayColor,
+				UnrealIdentifiers::PrimvarsDisplayOpacity,
+				UnrealIdentifiers::DoubleSided,
+				// When we change these UsdGeomModelAPI attributes we may need to create a new component type
+				// for the prim (as it may now need/stop needing an alternate draw mode component)
+				UnrealIdentifiers::ModelDrawMode,
+				UnrealIdentifiers::ModelApplyDrawMode,
+				// Collapsing changes need to resync to build the new collapsed/uncollapsed meshes
+				*UsdToUnreal::ConvertToken(UnrealIdentifiers::UnrealCollapsingAttr),
+				// Physics collision attribute change needs to rebuild the collision shape
+				*UsdToUnreal::ConvertToken(pxr::UsdPhysicsTokens->physicsCollisionEnabled),
+				*UsdToUnreal::ConvertToken(pxr::UsdPhysicsTokens->physicsApproximation),
+				// We put all of the audio info directly on the section, so if any of these change then we need to re-add the prim
+				// to the sequencer, and potentially generate a new audio asset
+				*UsdToUnreal::ConvertToken(pxr::UsdMediaTokens->filePath),
+				*UsdToUnreal::ConvertToken(pxr::UsdMediaTokens->auralMode),
+				*UsdToUnreal::ConvertToken(pxr::UsdMediaTokens->playbackMode),
+				*UsdToUnreal::ConvertToken(pxr::UsdMediaTokens->startTime),
+				*UsdToUnreal::ConvertToken(pxr::UsdMediaTokens->endTime),
+				*UsdToUnreal::ConvertToken(pxr::UsdMediaTokens->mediaOffset),
+				*UsdToUnreal::ConvertToken(pxr::UsdMediaTokens->gain)
+			};
+			if (PrimPath.IsAbsoluteRootOrPrimPath() && PrimResyncProperties.Contains(PropertyName))
+			{
+				bIsResync = true;
+				bHasResync = true;
+				bNeedsAnimationReload = true;
+			}
+
+			// Upgrade some info changes into resync changes
 			for (const UsdUtils::FSdfChangeListEntry& ObjectChange : InfoChange.Value)
 			{
 				// This is in charge of resyncing components (and so calling CreateComponents instead of just
@@ -1466,67 +1513,27 @@ void AUsdStageActor::HandleAccumulatedNotices()
 					continue;
 				}
 
-				for (const UsdUtils::FAttributeChange& AttributeChange : ObjectChange.AttributeChanges)
+				// Some stage info should trigger some resyncs because they should trigger reparsing of geometry
+				if (PrimPath.IsAbsoluteRootPath())
 				{
-					static const TSet<FString> StageResyncProperties = {TEXT("metersPerUnit"), TEXT("upAxis")};
-
-					// Upgrade these to resync so that the prim twins are regenerated, which clears all the existing
-					// animation tracks and adds new ones, automatically re-baking to control rig
-					static const TSet<FString> PrimResyncProperties = {
-						*UsdToUnreal::ConvertToken(UnrealIdentifiers::UnrealControlRigPath),
-						*UsdToUnreal::ConvertToken(UnrealIdentifiers::UnrealUseFKControlRig),
-						*UsdToUnreal::ConvertToken(UnrealIdentifiers::UnrealControlRigReduceKeys),
-						*UsdToUnreal::ConvertToken(UnrealIdentifiers::UnrealControlRigReductionTolerance),
-						// For now we need to do this as we need to refresh the material slot info on the info cache if these
-						// update... this is of course way too aggressive, although it's unlikely people will be manually editing these.
-						// TODO: More nuanced info cache updates
-						UnrealIdentifiers::PrimvarsDisplayColor,
-						UnrealIdentifiers::PrimvarsDisplayOpacity,
-						UnrealIdentifiers::DoubleSided,
-						// When we change these UsdGeomModelAPI attributes we may need to create a new component type
-						// for the prim (as it may now need/stop needing an alternate draw mode component)
-						UnrealIdentifiers::ModelDrawMode,
-						UnrealIdentifiers::ModelApplyDrawMode,
-						// Collapsing changes need to resync to build the new collapsed/uncollapsed meshes
-						*UsdToUnreal::ConvertToken(UnrealIdentifiers::UnrealCollapsingAttr),
-						// Physics collision attribute change needs to rebuild the collision shape
-						*UsdToUnreal::ConvertToken(pxr::UsdPhysicsTokens->physicsCollisionEnabled),
-						*UsdToUnreal::ConvertToken(pxr::UsdPhysicsTokens->physicsApproximation)};
-
-					// Some stage info should trigger some resyncs because they should trigger reparsing of geometry
-					if ((PrimPath.IsAbsoluteRootPath() && StageResyncProperties.Contains(AttributeChange.PropertyName))
-						|| PrimResyncProperties.Contains(AttributeChange.PropertyName))
+					for (const UsdUtils::FAttributeChange& FieldChange : ObjectChange.FieldChanges)
 					{
-						bIsResync = true;
-						bHasResync = true;
-						break;
+						static const TSet<FString> StageResyncProperties = {TEXT("metersPerUnit"), TEXT("upAxis")};
+						if (StageResyncProperties.Contains(FieldChange.Field))
+						{
+							bIsResync = true;
+							bHasResync = true;
+							break;
+						}
 					}
 
-					// We put all of the audio info directly on the section, so if any of these change then we need to re-add the prim
-					// to the sequencer, and potentially generate a new audio asset
-					static const TSet<FString> AudioSectionProperties = {
-						*UsdToUnreal::ConvertToken(pxr::UsdMediaTokens->filePath),
-						*UsdToUnreal::ConvertToken(pxr::UsdMediaTokens->auralMode),
-						*UsdToUnreal::ConvertToken(pxr::UsdMediaTokens->playbackMode),
-						*UsdToUnreal::ConvertToken(pxr::UsdMediaTokens->startTime),
-						*UsdToUnreal::ConvertToken(pxr::UsdMediaTokens->endTime),
-						*UsdToUnreal::ConvertToken(pxr::UsdMediaTokens->mediaOffset),
-						*UsdToUnreal::ConvertToken(pxr::UsdMediaTokens->gain)};
-					if (AudioSectionProperties.Contains(AttributeChange.PropertyName))
+					// Any sublayer change (even offsets) means we need to regenerate our LevelSequence to add (or shift)
+					// the corresponding subsequences
+					if (ObjectChange.SubLayerChanges.Num() > 0)
 					{
-						bIsResync = true;
-						bHasResync = true;
-						break;
+						bNeedsAnimationReload = true;
 					}
 				}
-
-				// Any sublayer change (even offsets) means we need to regenerate our LevelSequence to add (or shift)
-				// the corresponding subsequences
-				if (PrimPath.IsAbsoluteRootPath() && ObjectChange.SubLayerChanges.Num() > 0)
-				{
-					bNeedsAnimationReload = true;
-					break;
-				}
 			}
 
 			// We may need the full spec path with variant selections later, but for traversal and retrieving prims from the stage we always need
@@ -1534,14 +1541,14 @@ void AUsdStageActor::HandleAccumulatedNotices()
 			// GetPrimAtPath("/Root/Child")), and USD sometimes emits changes with the variant selection path (like during renames).
 			SortedPrimsChangedList.Add(PrimPath, bIsResync);
 		}
-		// Do Resyncs after so that they overwrite pure info changes if we have any
-		for (const TPair<FString, TArray<UsdUtils::FSdfChangeListEntry>>& ResyncChange : AccumulatedResyncChanges)
-		{
-			UE::FSdfPath PrimPath = UE::FSdfPath(*ResyncChange.Key).StripAllVariantSelections();
+	}
+	// Do Resyncs after so that they overwrite pure info changes if we have any
+	for (const TPair<FString, TArray<UsdUtils::FSdfChangeListEntry>>& ResyncChange : AccumulatedResyncChanges)
+	{
+		UE::FSdfPath PrimPath = UE::FSdfPath(*ResyncChange.Key).GetAbsoluteRootOrPrimPath().StripAllVariantSelections();
 
-			const bool bIsResync = true;
-			SortedPrimsChangedList.Add(PrimPath, bIsResync);
-		}
+		const bool bIsResync = true;
+		SortedPrimsChangedList.Add(PrimPath, bIsResync);
 	}
 
 	// During PIE, the PIE and the editor world will respond to notices. We have to prevent any PIE
@@ -4392,7 +4399,8 @@ void AUsdStageActor::OnObjectPropertyChanged(UObject* ObjectBeingModified, FProp
 	const static TSet<FName> TransformProperties = {
 		USceneComponent::GetRelativeLocationPropertyName(),
 		USceneComponent::GetRelativeRotationPropertyName(),
-		USceneComponent::GetRelativeScale3DPropertyName()};
+		USceneComponent::GetRelativeScale3DPropertyName()
+	};
 	const bool bIsTransformChange = TransformProperties.Contains(PropertyChangedEvent.GetPropertyName());
 
 	// When we change an actor property that is just a mirror of a component property (e.g. light intensity, or camera aperture)
diff --git a/Engine/Plugins/Importers/USDImporter/Source/USDStage/Private/USDTransactor.cpp b/Engine/Plugins/Importers/USDImporter/Source/USDStage/Private/USDTransactor.cpp
--- a/Engine/Plugins/Importers/USDImporter/Source/USDStage/Private/USDTransactor.cpp
+++ b/Engine/Plugins/Importers/USDImporter/Source/USDStage/Private/USDTransactor.cpp
@@ -35,8 +35,6 @@ namespace UsdUtils
 	// Objects adapted from USDListener, since we use slightly different data here
 	struct FTransactorAttributeChange
 	{
-		FString PropertyName;
-		bool bIsRelationship = false;
 		FString Field;
 		FString AttributeTypeName;	   // Full SdfValueTypeName of the attribute (e.g. normal3f, bool, texCoord3d, float2) so that we can undo/redo
 									   // attribute creation
@@ -49,15 +47,24 @@ namespace UsdUtils
 
 	struct FTransactorObjectChange
 	{
-		TArray<FTransactorAttributeChange> AttributeChanges;
+		TArray<FTransactorAttributeChange> FieldChanges;
 		FPrimChangeFlags Flags;
 		FString PrimTypeName;
 		TArray<FString> PrimAppliedSchemas;
 		FString OldPath;
 	};
 
+	enum class EObjectType
+	{
+		Prim,
+		Attribute,
+		Relationship
+	};
+
 	struct FTransactorRecordedEdit
 	{
+		EObjectType ObjectType;	   // Describes what object that ObjectPath is pointing to, and also what the "Field" values inside the ObjectChanges
+								   // mean (e.g. if ObjectPath describes an Attribute, we know the ObjectChanges are attribute metadata changes)
 		FString ObjectPath;
 		TArray<FTransactorObjectChange> ObjectChanges;
 	};
@@ -80,10 +87,8 @@ namespace UsdUtils
 
 FArchive& operator<<(FArchive& Ar, UsdUtils::FTransactorAttributeChange& Change)
 {
-	Ar << Change.PropertyName;
 	Ar << Change.Field;
 	Ar << Change.AttributeTypeName;
-	Ar << Change.bIsRelationship;
 	Ar << Change.OldValue;
 	Ar << Change.NewValue;
 	Ar << Change.TimeSamples;
@@ -114,7 +119,7 @@ FArchive& operator<<(FArchive& Ar, UsdUtils::FPrimChangeFlags& Flags)
 
 FArchive& operator<<(FArchive& Ar, UsdUtils::FTransactorObjectChange& Change)
 {
-	Ar << Change.AttributeChanges;
+	Ar << Change.FieldChanges;
 	Ar << Change.Flags;
 	Ar << Change.PrimTypeName;
 	Ar << Change.PrimAppliedSchemas;
@@ -124,6 +129,7 @@ FArchive& operator<<(FArchive& Ar, UsdUtils::FTransactorObjectChange& Change)
 
 FArchive& operator<<(FArchive& Ar, UsdUtils::FTransactorRecordedEdit& Edit)
 {
+	Ar << Edit.ObjectType;
 	Ar << Edit.ObjectPath;
 	Ar << Edit.ObjectChanges;
 	return Ar;
@@ -149,11 +155,21 @@ namespace UsdUtils
 
 		for (const TPair<FString, TArray<FSdfChangeListEntry>>& ChangesByPrimPath : InChanges)
 		{
-			const FString& PrimPath = ChangesByPrimPath.Key;
+			const FString& ObjectPath = ChangesByPrimPath.Key;
 			const TArray<FSdfChangeListEntry>& Changes = ChangesByPrimPath.Value;
 
+			UE::FSdfPath UEObjectPath{*ObjectPath};
+
+			UE::FSdfPath PrimPath = UEObjectPath.GetAbsoluteRootOrPrimPath();
+			UE::FUsdPrim Prim = InStage.GetPrimAtPath(PrimPath);
+
+			FString PropertyName = UEObjectPath.IsPropertyPath() ? UEObjectPath.GetName() : FString{};
+
 			FTransactorRecordedEdit& Edit = InOutEdits.Edits.Emplace_GetRef();
-			Edit.ObjectPath = PrimPath;
+			Edit.ObjectPath = ObjectPath;
+			Edit.ObjectType = UEObjectPath.IsAbsoluteRootOrPrimPath()		  ? EObjectType::Prim
+							  : (Prim && Prim.HasRelationship(*PropertyName)) ? EObjectType::Relationship
+																			  : EObjectType::Attribute;
 
 			for (const FSdfChangeListEntry& Change : Changes)
 			{
@@ -161,11 +177,12 @@ namespace UsdUtils
 				ConvertedChange.Flags = Change.Flags;
 				ConvertedChange.OldPath = Change.OldPath;
 
-				UE::FUsdPrim Prim = InStage.GetPrimAtPath(UE::FSdfPath{*PrimPath});
-
+				// Adding/removing prim --> Record prim info
 				if (Prim
-					&& (ConvertedChange.Flags.bDidAddInertPrim || ConvertedChange.Flags.bDidAddNonInertPrim
-						|| ConvertedChange.Flags.bDidRemoveInertPrim || ConvertedChange.Flags.bDidRemoveNonInertPrim))
+					&& (ConvertedChange.Flags.bDidAddInertPrim ||		//
+						ConvertedChange.Flags.bDidAddNonInertPrim ||	//
+						ConvertedChange.Flags.bDidRemoveInertPrim ||	//
+						ConvertedChange.Flags.bDidRemoveNonInertPrim))
 				{
 					ConvertedChange.PrimTypeName = Prim.GetTypeName().ToString();
 
@@ -181,7 +198,7 @@ namespace UsdUtils
 						Log,
 						TEXT("Recorded the %s of prim '%s' with TypeName '%s' and PrimAppliedSchemas [%s]"),
 						(ConvertedChange.Flags.bDidAddInertPrim || ConvertedChange.Flags.bDidAddNonInertPrim) ? TEXT("addition") : TEXT("removal"),
-						*PrimPath,
+						*ObjectPath,
 						*ConvertedChange.PrimTypeName,
 						*FString::Join(ConvertedChange.PrimAppliedSchemas, TEXT(", "))
 					);
@@ -193,59 +210,51 @@ namespace UsdUtils
 					// We expect not to find the prim if the change says it was just removed though
 					if (!ConvertedChange.Flags.bDidRemoveInertPrim && !ConvertedChange.Flags.bDidRemoveNonInertPrim)
 					{
-						UE_LOG(LogUsd, Warning, TEXT("Failed to find prim at path '%s' when serializing object changes in transactor"), *PrimPath);
+						UE_LOG(LogUsd, Warning, TEXT("Failed to find prim at path '%s' when serializing object changes in transactor"), *ObjectPath);
 					}
 					continue;
 				}
 
-				for (const FAttributeChange& AttributeChange : Change.AttributeChanges)
+				for (const FFieldChange& FieldChange : Change.FieldChanges)
 				{
-					// We won't be able to do anything with this, so just ignore it
-					if (AttributeChange.PropertyName.IsEmpty())
-					{
-						continue;
-					}
-
-					FTransactorAttributeChange& ConvertedAttributeChange = ConvertedChange.AttributeChanges.Emplace_GetRef();
-					ConvertedAttributeChange.PropertyName = AttributeChange.PropertyName;
-					ConvertedAttributeChange.Field = AttributeChange.Field;
+					FTransactorAttributeChange& ConvertedAttributeChange = ConvertedChange.FieldChanges.Emplace_GetRef();
+					ConvertedAttributeChange.Field = FieldChange.Field;
 
 					FConvertedVtValue ConvertedOldValue;
-					if (UsdToUnreal::ConvertValue(AttributeChange.OldValue, ConvertedOldValue))
+					if (UsdToUnreal::ConvertValue(FieldChange.OldValue, ConvertedOldValue))
 					{
 						ConvertedAttributeChange.OldValue = ConvertedOldValue;
 					}
 
 					FConvertedVtValue ConvertedNewValue;
-					if (UsdToUnreal::ConvertValue(AttributeChange.NewValue, ConvertedNewValue))
+					if (UsdToUnreal::ConvertValue(FieldChange.NewValue, ConvertedNewValue))
 					{
 						ConvertedAttributeChange.NewValue = ConvertedNewValue;
 					}
 
 					// We likely won't be able to fetch the attribute if it was removed, so try deducing the typename from the value just so that we
 					// have *something*
-					if (!AttributeChange.OldValue.IsEmpty()
-						&& (ConvertedChange.Flags.bDidRemoveProperty || ConvertedChange.Flags.bDidRemovePropertyWithOnlyRequiredFields))
+					if (!FieldChange.OldValue.IsEmpty()
+						&& (ConvertedChange.Flags.bDidRemoveProperty ||	   //
+							ConvertedChange.Flags.bDidRemovePropertyWithOnlyRequiredFields))
 					{
-						ConvertedAttributeChange.AttributeTypeName = UsdUtils::GetImpliedTypeName(AttributeChange.OldValue);
+						ConvertedAttributeChange.AttributeTypeName = UsdUtils::GetImpliedTypeName(FieldChange.OldValue);
 						UE_LOG(
 							LogUsd,
 							Log,
 							TEXT("Recording the removal of properties is not fully supported: Using underlying type '%s' for record of attribute "
 								 "'%s' of prim '%s', as we don't have access to the attribute's role"),
 							*ConvertedAttributeChange.AttributeTypeName,
-							*ConvertedAttributeChange.PropertyName,
-							*PrimPath
+							*PropertyName,
+							*ObjectPath
 						);
 					}
 
-					// If the prim is the root, this is a stage info change, and PropertyName is actually a metadata key, so
-					// it's not possible (or required) to know its SdfValueTypeName to e.g. undo/redo the creation of the property
-					if (!Prim.IsPseudoRoot() && ConvertedAttributeChange.PropertyName != TEXT("kind")
-						&& !ConvertedAttributeChange.PropertyName.IsEmpty() && !ConvertedChange.Flags.bDidRemoveProperty
-						&& !ConvertedChange.Flags.bDidRemovePropertyWithOnlyRequiredFields)
+					// Record attribute typename/timeSamples if we can find it
+					if (Edit.ObjectType == EObjectType::Attribute && !Prim.IsPseudoRoot() && !PropertyName.IsEmpty()
+						&& !ConvertedChange.Flags.bDidRemoveProperty && !ConvertedChange.Flags.bDidRemovePropertyWithOnlyRequiredFields)
 					{
-						if (UE::FUsdAttribute Attribute = Prim.GetAttribute(*ConvertedAttributeChange.PropertyName))
+						if (UE::FUsdAttribute Attribute = Prim.GetAttribute(*PropertyName))
 						{
 							ConvertedAttributeChange.AttributeTypeName = Attribute.GetTypeName().ToString();
 
@@ -266,18 +275,14 @@ namespace UsdUtils
 								}
 							}
 						}
-						else if (UE::FUsdRelationship Relationship = Prim.GetRelationship(*ConvertedAttributeChange.PropertyName))
-						{
-							ConvertedAttributeChange.bIsRelationship = true;
-						}
 						else
 						{
 							UE_LOG(
 								LogUsd,
 								Warning,
 								TEXT("Failed to find attribute '%s' for prim at path '%s' when serializing object changes in transactor"),
-								*ConvertedAttributeChange.PropertyName,
-								*PrimPath
+								*PropertyName,
+								*PrimPath.GetString()
 							);
 						}
 					}
@@ -365,20 +370,21 @@ namespace UsdUtils
 		return false;
 	}
 
-	bool ApplyAttributeTimeSamples(const FTransactorAttributeChange& AttributeChange, UE::FUsdPrim& Prim)
+	bool ApplyAttributeTimeSamples(const FTransactorAttributeChange& AttributeChange, const UE::FSdfPath& ObjectPath, UE::FUsdPrim& Prim)
 	{
-		if (!Prim || AttributeChange.TimeSamples.Num() == 0 || AttributeChange.TimeSamples.Num() != AttributeChange.TimeValues.Num()
-			|| AttributeChange.Field != TEXT("timeSamples"))
+		if (!Prim || AttributeChange.TimeSamples.Num() == 0 || AttributeChange.TimeSamples.Num() != AttributeChange.TimeValues.Num())
 		{
 			return false;
 		}
 
+		FString PropertyName = ObjectPath.GetName();
+
 		// Try Getting first because we shouldn't trust our AttributeTypeName to always just CreateAttribute, as it may be just deduced from a value
 		// and be different
-		UE::FUsdAttribute Attribute = Prim.GetAttribute(*AttributeChange.PropertyName);
+		UE::FUsdAttribute Attribute = Prim.GetAttribute(*PropertyName);
 		if (!Attribute)
 		{
-			Attribute = Prim.CreateAttribute(*AttributeChange.PropertyName, *AttributeChange.AttributeTypeName);
+			Attribute = Prim.CreateAttribute(*PropertyName, *AttributeChange.AttributeTypeName);
 		}
 		if (!Attribute)
 		{
@@ -394,7 +400,7 @@ namespace UsdUtils
 			Log,
 			TEXT("Applying '%d' timeSamples for attribute '%s' of prim '%s'"),
 			AttributeChange.TimeSamples.Num(),
-			*AttributeChange.PropertyName,
+			*PropertyName,
 			*Prim.GetPrimPath().GetString()
 		);
 
@@ -412,7 +418,7 @@ namespace UsdUtils
 						TEXT("Failed to apply value '%s' at timesample '%f' for attribute '%s' of prim '%s'"),
 						*UsdUtils::Stringify(Value),
 						AttributeChange.TimeSamples[Index],
-						*AttributeChange.PropertyName,
+						*PropertyName,
 						*Prim.GetPrimPath().GetString()
 					);
 					bSuccess = false;
@@ -425,7 +431,7 @@ namespace UsdUtils
 					Warning,
 					TEXT("Failed to convert value for timesample '%f' for attribute '%s' of prim '%s'"),
 					AttributeChange.TimeSamples[Index],
-					*AttributeChange.PropertyName,
+					*PropertyName,
 					*Prim.GetPrimPath().GetString()
 				);
 				bSuccess = false;
@@ -436,14 +442,14 @@ namespace UsdUtils
 	}
 
 	bool ApplyAttributeChange(
-		const FString& PropertyName,
-		const FString& Field,
-		const FString& AttributeTypeName,
-		bool bIsRelationship,
-		bool bRemoveProperty,
-		const FConvertedVtValue& Value,
-		UE::FUsdPrim& Prim,
-		TOptional<double> Time = {}
+		const UE::FUsdPrim& Prim,			 // This is always the leafmost, relevant prim. Regardless of ObjectType
+		const UE::FSdfPath& ObjectPath,		 // Path to the object containing the field. Could be a property or prim
+		const EObjectType ObjectType,		 // Describes what the ObjectPath is pointing to
+		const FString& Field,				 // The actual field name. "default" for attr default values, or "kind"/"payload" for metadata, etc.
+		const FString& AttributeTypeName,	 // If ObjectType == EObjectType::Attribute, this describes the type name of the attribute's values
+		bool bRemoveProperty,				 // Whether to remove the property instead of just applying values
+		const FConvertedVtValue& Value,		 // The value to apply to the field
+		TOptional<double> Time = {}			 // The time at which to apply/clear the value to the field
 	)
 	{
 		if (!Prim)
@@ -452,75 +458,76 @@ namespace UsdUtils
 		}
 
 		bool bCreated = false;
+
+		FString PropertyName = ObjectPath.IsPropertyPath() ? ObjectPath.GetName() : FString{};
 		UE::FUsdAttribute Attribute;
 		UE::FUsdRelationship Relationship;
 
-		if (PropertyName != TEXT("kind"))	 // Kind is prim metadata, not an attribute
+		if (bRemoveProperty)
 		{
-			if (bRemoveProperty)
+			if (ObjectType == EObjectType::Relationship)
 			{
-				if (bIsRelationship)
+				PropertyName = ObjectPath.GetName();
+				Relationship = Prim.GetRelationship(*PropertyName);
+				if (!Relationship)
 				{
-					Relationship = Prim.GetRelationship(*PropertyName);
-					if (!Relationship)
-					{
-						return true;
-					}
+					return true;
 				}
-				else
+			}
+			else if (ObjectType == EObjectType::Attribute)
+			{
+				PropertyName = ObjectPath.GetName();
+				Attribute = Prim.GetAttribute(*PropertyName);
+				if (!Attribute)
 				{
-					Attribute = Prim.GetAttribute(*PropertyName);
-					if (!Attribute)
-					{
-						return true;
-					}
+					return true;
 				}
 			}
-			else
+		}
+		else
+		{
+			if (ObjectType == EObjectType::Relationship)
 			{
-				if (bIsRelationship)
+				bool bHadRelationship = Prim.HasRelationship(*PropertyName);
+				Relationship = Prim.CreateRelationship(*PropertyName);
+				if (!Relationship)
 				{
-					bool bHadRelationship = Prim.HasRelationship(*PropertyName);
-					Relationship = Prim.CreateRelationship(*PropertyName);
-					if (!Relationship)
+					UE_LOG(
+						LogUsd,
+						Warning,
+						TEXT("Failed to create relationship '%s' with for prim '%s'"),
+						*PropertyName,
+						*Prim.GetPrimPath().GetString()
+					);
+
+					return false;
+				}
+
+				bCreated = !bHadRelationship;
+			}
+			else if (ObjectType == EObjectType::Attribute)
+			{
+				bool bHadAttr = Prim.HasAttribute(*PropertyName);
+				Attribute = Prim.CreateAttribute(*PropertyName, *AttributeTypeName);
+				if (!Attribute)
+				{
+					// We expect to fail to create an attribute if we have no typename here (e.g. undo remove property)
+					if (AttributeTypeName.IsEmpty())
 					{
 						UE_LOG(
 							LogUsd,
 							Warning,
-							TEXT("Failed to create relationship '%s' with for prim '%s'"),
+							TEXT("Failed to create attribute '%s' with typename '%s' for prim '%s'"),
 							*PropertyName,
+							*AttributeTypeName,
 							*Prim.GetPrimPath().GetString()
 						);
-
-						return false;
 					}
 
-					bCreated = !bHadRelationship;
+					return false;
 				}
-				else
-				{
-					bool bHadAttr = Prim.HasAttribute(*PropertyName);
-					Attribute = Prim.CreateAttribute(*PropertyName, *AttributeTypeName);
-					if (!Attribute)
-					{
-						// We expect to fail to create an attribute if we have no typename here (e.g. undo remove property)
-						if (AttributeTypeName.IsEmpty())
-						{
-							UE_LOG(
-								LogUsd,
-								Warning,
-								TEXT("Failed to create attribute '%s' with typename '%s' for prim '%s'"),
-								*PropertyName,
-								*AttributeTypeName,
-								*Prim.GetPrimPath().GetString()
-							);
-						}
-
-						return false;
-					}
 
-					bCreated = !bHadAttr;
-				}
+				bCreated = !bHadAttr;
 			}
 		}
 
@@ -530,9 +537,9 @@ namespace UsdUtils
 			UE_LOG(
 				LogUsd,
 				Warning,
-				TEXT("Failed to convert VtValue back to USD when applying it to property '%s' of prim '%s'"),
-				*PropertyName,
-				*Prim.GetPrimPath().GetString()
+				TEXT("Failed to convert VtValue back to USD when applying it to object '%s' field '%s'"),
+				*ObjectPath.GetString(),
+				*Field
 			);
 			return false;
 		}
@@ -540,48 +547,26 @@ namespace UsdUtils
 		UE_LOG(
 			LogUsd,
 			Log,
-			TEXT("%s property '%s' (typename '%s'), field '%s' of prim '%s' with value '%s' at time '%s'"),
+			TEXT("%s object '%s' (typename '%s'), field '%s' with value '%s' at time '%s'"),
 			bCreated				 ? TEXT("Creating")
 			: WrapperValue.IsEmpty() ? bRemoveProperty ? TEXT("Removing") : TEXT("Clearing")
 									 : TEXT("Setting"),
-			*PropertyName,
+			*ObjectPath.GetString(),
 			*AttributeTypeName,
 			*Field,
-			*Prim.GetPrimPath().GetString(),
 			WrapperValue.IsEmpty() ? TEXT("<empty>") : *UsdUtils::Stringify(WrapperValue),
 			Time.IsSet() ? *LexToString(Time.GetValue()) : TEXT("<unset>")
 		);
 
-		// If we just want to remove the property don't really bother doing anything else
-		if (bRemoveProperty)
+		// If we just want to remove the property don't really bother doing anything else.
+		// Note: We never get this flag set to true when dealing with metadata, only properties (attributes and relationships)
+		if (bRemoveProperty && !PropertyName.IsEmpty())
 		{
 			Prim.RemoveProperty(*PropertyName);
 			return true;
 		}
 
-		if (PropertyName == TEXT("kind"))
-		{
-#if USE_USD_SDK
-			if (Value.Entries.Num() == 1 && Value.Entries[0].Num() == 1)
-			{
-				if (const FString* KindString = Value.Entries[0][0].TryGet<FString>())
-				{
-					if (!IUsdPrim::SetKind(Prim, UnrealToUsd::ConvertToken(**KindString).Get()))
-					{
-						UE_LOG(LogUsd, Warning, TEXT("Failed to set Kind '%s' for prim '%s'"), **KindString, *Prim.GetPrimPath().GetString());
-					}
-				}
-			}
-			else
-			{
-				if (!IUsdPrim::ClearKind(Prim))
-				{
-					UE_LOG(LogUsd, Warning, TEXT("Failed to clear Kind for prim '%s'"), *Prim.GetPrimPath().GetString());
-				}
-			}
-#endif	  // USE_USD_SDK
-		}
-		else if (Field == TEXT("default"))
+		if (Field == TEXT("default"))
 		{
 			if (WrapperValue.IsEmpty())
 			{
@@ -600,7 +585,7 @@ namespace UsdUtils
 			}
 		}
 		// This seems to be the field name for the actual value in pxr:UsdRelationship
-		else if (Field == TEXT("targetPaths"))
+		else if (ObjectType == EObjectType::Relationship && Field == TEXT("targetPaths"))
 		{
 			if (WrapperValue.IsEmpty())
 			{
@@ -638,7 +623,7 @@ namespace UsdUtils
 				}
 			}
 		}
-		else	// variability, colorSpace, etc.
+		else	// Other metadata fields (on prims or properties)
 		{
 			if (WrapperValue.IsEmpty())
 			{
@@ -650,6 +635,10 @@ namespace UsdUtils
 				{
 					Relationship.ClearMetadata(*Field);
 				}
+				else
+				{
+					Prim.ClearMetadata(*Field);
+				}
 			}
 			else
 			{
@@ -661,15 +650,19 @@ namespace UsdUtils
 				{
 					Relationship.SetMetadata(*Field, WrapperValue);
 				}
+				else
+				{
+					Prim.SetMetadata(*Field, WrapperValue);
+				}
 			}
 		}
 
 		return true;
 	}
 
-	bool ApplyStageMetadataChange(const FString& PropertyName, const FConvertedVtValue& Value, UE::FUsdStage& Stage)
+	bool ApplyStageMetadataChange(const FString& FieldName, const FConvertedVtValue& Value, UE::FUsdStage& Stage)
 	{
-		if (!Stage || PropertyName.IsEmpty())
+		if (!Stage || FieldName.IsEmpty())
 		{
 			return false;
 		}
@@ -677,22 +670,22 @@ namespace UsdUtils
 		UE::FVtValue WrapperValue;
 		if (!UnrealToUsd::ConvertValue(Value, WrapperValue))
 		{
-			UE_LOG(LogUsd, Warning, TEXT("Failed to convert VtValue back to USD when applying it to stage metadata field '%s'"), *PropertyName);
+			UE_LOG(LogUsd, Warning, TEXT("Failed to convert VtValue back to USD when applying it to stage metadata field '%s'"), *FieldName);
 			return false;
 		}
 
-		UE_LOG(LogUsd, Log, TEXT("Setting stage metadata '%s', with value '%s'"), *PropertyName, *UsdUtils::Stringify(WrapperValue));
+		UE_LOG(LogUsd, Log, TEXT("Setting stage metadata '%s', with value '%s'"), *FieldName, *UsdUtils::Stringify(WrapperValue));
 
 		UE::FSdfLayer OldEditTarget = Stage.GetEditTarget();
 		Stage.SetEditTarget(Stage.GetRootLayer());	  // Stage metadata always needs to be set at the root layer
 		{
 			if (WrapperValue.IsEmpty())
 			{
-				Stage.ClearMetadata(*PropertyName);
+				Stage.ClearMetadata(*FieldName);
 			}
 			else
 			{
-				Stage.SetMetadata(*PropertyName, WrapperValue);
+				Stage.SetMetadata(*FieldName, WrapperValue);
 			}
 		}
 		Stage.SetEditTarget(OldEditTarget);
@@ -780,10 +773,10 @@ namespace UsdUtils
 				{
 					for (const FTransactorObjectChange& PrimChange : Edit.ObjectChanges)
 					{
-						for (const FTransactorAttributeChange& AttributeChange : PrimChange.AttributeChanges)
+						for (const FTransactorAttributeChange& AttributeChange : PrimChange.FieldChanges)
 						{
 							ApplyStageMetadataChange(
-								AttributeChange.PropertyName,
+								AttributeChange.Field,
 								Direction == EApplicationDirection::Forward ? AttributeChange.NewValue : AttributeChange.OldValue,
 								Stage
 							);
@@ -794,7 +787,8 @@ namespace UsdUtils
 				{
 					for (const FTransactorObjectChange& PrimChange : Edit.ObjectChanges)
 					{
-						UE::FSdfPath PrimPath = UE::FSdfPath{*Edit.ObjectPath};
+						UE::FSdfPath UEObjectPath = UE::FSdfPath{*Edit.ObjectPath};
+						UE::FSdfPath PrimPath = UEObjectPath.GetAbsoluteRootOrPrimPath();
 						if (ApplyPrimChange(PrimPath, PrimChange, Stage, Direction))
 						{
 							// If we managed to apply a prim change, we know there aren't any other attribute changes in the same ObjectChange
@@ -818,23 +812,23 @@ namespace UsdUtils
 
 						if (PrimChange.Flags.bDidChangeAttributeTimeSamples)
 						{
-							for (const FTransactorAttributeChange& AttributeChange : PrimChange.AttributeChanges)
+							for (const FTransactorAttributeChange& AttributeChange : PrimChange.FieldChanges)
 							{
-								ApplyAttributeTimeSamples(AttributeChange, Prim);
+								ApplyAttributeTimeSamples(AttributeChange, UEObjectPath, Prim);
 							}
 						}
 						else
 						{
-							for (const FTransactorAttributeChange& AttributeChange : PrimChange.AttributeChanges)
+							for (const FTransactorAttributeChange& FieldChange : PrimChange.FieldChanges)
 							{
 								ApplyAttributeChange(
-									AttributeChange.PropertyName,
-									AttributeChange.Field,
-									AttributeChange.AttributeTypeName,
-									AttributeChange.bIsRelationship,
+									Prim,
+									UEObjectPath,
+									Edit.ObjectType,
+									FieldChange.Field,
+									FieldChange.AttributeTypeName,
 									bShouldRemove,
-									Direction == EApplicationDirection::Forward ? AttributeChange.NewValue : AttributeChange.OldValue,
-									Prim
+									Direction == EApplicationDirection::Forward ? FieldChange.NewValue : FieldChange.OldValue
 								);
 							}
 						}
diff --git a/Engine/Plugins/Runtime/USDCore/Source/USDUtilities/Private/USDValueConversion.cpp b/Engine/Plugins/Runtime/USDCore/Source/USDUtilities/Private/USDValueConversion.cpp
--- a/Engine/Plugins/Runtime/USDCore/Source/USDUtilities/Private/USDValueConversion.cpp
+++ b/Engine/Plugins/Runtime/USDCore/Source/USDUtilities/Private/USDValueConversion.cpp
@@ -13,6 +13,8 @@
 #include "USDIncludesStart.h"
 #include "pxr/base/tf/stringUtils.h"
 #include "pxr/base/vt/value.h"
+#include "pxr/usd/sdf/payload.h"
+#include "pxr/usd/sdf/reference.h"
 #include "pxr/usd/sdf/schema.h"
 #include "pxr/usd/sdf/types.h"
 #include "USDIncludesEnd.h"
@@ -27,7 +29,39 @@ namespace UE::USDValueConversion::Private
 	const static char EscapeChar = '\\';
 	const static char AssetDelimiterChar = '@';
 	const static TCHAR DoubleQuoteTCHAR = '\"';
-}
+
+	// Use a control character here because USD shouldn't normally let this be used for actual SdfPaths
+	const static FString ReferencePayloadDelimiter = TEXT("\x1c");
+
+#if USE_USD_SDK
+	// Functions instead of just variables to make sure we don't run into any bizarre issues with these static variables being
+	// initialized before the TfTypes are registered (not sure if that is possible, but we don't need to find out)
+	template<typename T>
+	const FString& GetSdfTypeName()
+	{
+		const static FString TypeName = *UsdToUnreal::ConvertString(pxr::TfType::Find<T>().GetTypeName());
+		return TypeName;
+	}
+
+	enum EListOpListIndex
+	{
+		Explicit = 0,
+		Added,
+		Deleted,
+		Ordered,
+		Prepended,
+		Appended,
+
+		Num
+	};
+	static_assert((int)EListOpListIndex::Explicit == (int)pxr::SdfListOpType::SdfListOpTypeExplicit);
+	static_assert((int)EListOpListIndex::Added == (int)pxr::SdfListOpType::SdfListOpTypeAdded);
+	static_assert((int)EListOpListIndex::Deleted == (int)pxr::SdfListOpType::SdfListOpTypeDeleted);
+	static_assert((int)EListOpListIndex::Ordered == (int)pxr::SdfListOpType::SdfListOpTypeOrdered);
+	static_assert((int)EListOpListIndex::Prepended == (int)pxr::SdfListOpType::SdfListOpTypePrepended);
+	static_assert((int)EListOpListIndex::Appended == (int)pxr::SdfListOpType::SdfListOpTypeAppended);
+#endif	  // USE_USD_SDK
+}	 // namespace UE::USDValueConversion::Private
 
 #if USE_USD_SDK
 namespace UsdToUnrealImpl
@@ -198,6 +232,90 @@ namespace UsdToUnrealImpl
 			}
 		);
 	}
+
+	template<typename UEType, typename USDType, typename ConverterFunc>
+	void ConvertListOpValue(const pxr::VtValue& InValue, UsdUtils::FConvertedVtValue& OutValue, ConverterFunc Converter)
+	{
+		using namespace pxr;
+		using namespace UsdUtils;
+		using namespace UE::USDValueConversion::Private;
+
+		OutValue.bIsArrayValued = true;
+		OutValue.Entries.SetNum(EListOpListIndex::Num);
+
+		const SdfListOp<USDType>& Val = InValue.UncheckedGet<SdfListOp<USDType>>();
+
+		// Explicit ops only have explicit items
+		if (Val.IsExplicit())
+		{
+			TArray<FConvertedVtValueComponent>& ExplicitItems = OutValue.Entries[EListOpListIndex::Explicit];
+
+			for (const USDType& Item : Val.GetExplicitItems())
+			{
+				ExplicitItems.Add(UsdUtils::FConvertedVtValueComponent(TInPlaceType<UEType>(), Converter(Item)));
+			}
+		}
+		// Non explicit ops may have items in all other lists
+		else
+		{
+			TArray<FConvertedVtValueComponent>& AddedItems = OutValue.Entries[EListOpListIndex::Added];
+			TArray<FConvertedVtValueComponent>& PrependedItems = OutValue.Entries[EListOpListIndex::Prepended];
+			TArray<FConvertedVtValueComponent>& AppendedItems = OutValue.Entries[EListOpListIndex::Appended];
+			TArray<FConvertedVtValueComponent>& DeletedItems = OutValue.Entries[EListOpListIndex::Deleted];
+			TArray<FConvertedVtValueComponent>& OrderedItems = OutValue.Entries[EListOpListIndex::Ordered];
+
+			for (const USDType& Item : Val.GetAddedItems())
+			{
+				FString Converted = Converter(Item);
+				AddedItems.Add(UsdUtils::FConvertedVtValueComponent(TInPlaceType<UEType>(), Converter(Item)));
+			}
+			for (const USDType& Item : Val.GetPrependedItems())
+			{
+				FString Converted = Converter(Item);
+				PrependedItems.Add(UsdUtils::FConvertedVtValueComponent(TInPlaceType<UEType>(), Converter(Item)));
+			}
+			for (const USDType& Item : Val.GetAppendedItems())
+			{
+				FString Converted = Converter(Item);
+				AppendedItems.Add(UsdUtils::FConvertedVtValueComponent(TInPlaceType<UEType>(), Converter(Item)));
+			}
+			for (const USDType& Item : Val.GetDeletedItems())
+			{
+				FString Converted = Converter(Item);
+				DeletedItems.Add(UsdUtils::FConvertedVtValueComponent(TInPlaceType<UEType>(), Converter(Item)));
+			}
+			for (const USDType& Item : Val.GetOrderedItems())
+			{
+				FString Converted = Converter(Item);
+				OrderedItems.Add(UsdUtils::FConvertedVtValueComponent(TInPlaceType<UEType>(), Converter(Item)));
+			}
+		}
+	}
+
+	template<typename SdfRefOrPayload>
+	FString StringifySdfReferenceOrPayload(const SdfRefOrPayload& Reference)
+	{
+		using namespace UE::USDValueConversion::Private;
+
+		FString AssetPath = UsdToUnreal::ConvertString(Reference.GetAssetPath());
+		FString PrimPath = UsdToUnreal::ConvertPath(Reference.GetPrimPath());
+		double Offset = Reference.GetLayerOffset().GetOffset();
+		double Scale = Reference.GetLayerOffset().GetScale();
+
+		// TODO: Apparently there's a full CustomData VtDictionary inside the SdfReference that we should probably also stringify.
+		// This doesn't seem to be used that often (and we never author it anyway), so we'll skip it for now
+
+		return FString::Printf(
+			TEXT("%s%s%s%s%f%s%f"),
+			*AssetPath,
+			*ReferencePayloadDelimiter,
+			*PrimPath,
+			*ReferencePayloadDelimiter,
+			Offset,
+			*ReferencePayloadDelimiter,
+			Scale
+		);
+	}
 }	 // namespace UsdToUnrealImpl
 #endif	  // USE_USD_SDK
 
@@ -209,6 +327,7 @@ namespace UsdToUnreal
 		using namespace pxr;
 		using namespace UsdUtils;
 		using namespace UsdToUnrealImpl;
+		using namespace UE::USDValueConversion::Private;
 
 		OutValue.SourceType = UsdUtils::EUsdBasicDataTypes::None;
 		OutValue.Entries.Reset();
@@ -389,15 +508,58 @@ namespace UsdToUnreal
 		else if (UnderlyingType.IsA<SdfListOp<SdfPath>>())
 		{
 			OutValue.SourceType = EUsdBasicDataTypes::String;
-			OutValue.bIsArrayValued = true;
+			OutValue.ExtraTypeInfo = GetSdfTypeName<SdfListOp<SdfPath>>();
 
-			const SdfListOp<SdfPath>& Val = UsdValue.UncheckedGet<SdfListOp<SdfPath>>();
-			const std::vector<SdfPath>& ItemVec = Val.GetExplicitItems();
-			OutValue.Entries.Reserve(ItemVec.size());
-			for (const pxr::SdfPath& Item : ItemVec)
-			{
-				OutValue.Entries.Add({UsdUtils::FConvertedVtValueComponent(TInPlaceType<FString>(), UsdToUnreal::ConvertPath(Item))});
-			}
+			ConvertListOpValue<FString, SdfPath, FString(const pxr::SdfPath&)>(	   //
+				UsdValue,
+				OutValue,
+				UsdToUnreal::ConvertPath
+			);
+		}
+		else if (UnderlyingType.IsA<SdfListOp<std::string>>())
+		{
+			OutValue.SourceType = EUsdBasicDataTypes::String;
+			OutValue.ExtraTypeInfo = GetSdfTypeName<SdfListOp<std::string>>();
+
+			ConvertListOpValue<FString, std::string, FString(const std::string&)>(	  //
+				UsdValue,
+				OutValue,
+				UsdToUnreal::ConvertString
+			);
+		}
+		else if (UnderlyingType.IsA<SdfListOp<SdfReference>>())
+		{
+			OutValue.SourceType = EUsdBasicDataTypes::String;
+			OutValue.ExtraTypeInfo = GetSdfTypeName<SdfListOp<SdfReference>>();
+
+			ConvertListOpValue<FString, SdfReference>(	  //
+				UsdValue,
+				OutValue,
+				StringifySdfReferenceOrPayload<SdfReference>
+			);
+		}
+		else if (UnderlyingType.IsA<SdfListOp<SdfPayload>>())
+		{
+			OutValue.SourceType = EUsdBasicDataTypes::String;
+			OutValue.ExtraTypeInfo = GetSdfTypeName<SdfListOp<SdfPayload>>();
+
+			ConvertListOpValue<FString, SdfPayload>(	  //
+				UsdValue,
+				OutValue,
+				StringifySdfReferenceOrPayload<SdfPayload>
+			);
+		}
+		// Type most commonly used for the value of the apiSchemas prim metadata
+		else if (UnderlyingType.IsA<SdfListOp<TfToken>>())
+		{
+			OutValue.SourceType = EUsdBasicDataTypes::Token;
+			OutValue.ExtraTypeInfo = GetSdfTypeName<SdfListOp<TfToken>>();
+
+			ConvertListOpValue<FString, TfToken, FString(const pxr::TfToken&)>(	   //
+				UsdValue,
+				OutValue,
+				UsdToUnreal::ConvertToken
+			);
 		}
 		// These types should only appear within metadata, and don't support arrays.
 		// There are more of them (e.g. pxr/usd/usd/crateDataTypes.h), but these are the most common.
@@ -601,6 +763,95 @@ namespace UnrealToUsdImpl
 			}
 		);
 	}
+
+	template<typename UEElementType, typename USDElementType, typename ConverterFunc>
+	void ConvertListOpValue(const UsdUtils::FConvertedVtValue& InValue, pxr::VtValue& OutValue, ConverterFunc Converter)
+	{
+		using namespace pxr;
+		using namespace UsdUtils;
+		using namespace UE::USDValueConversion::Private;
+
+		if (InValue.Entries.Num() == EListOpListIndex::Num)
+		{
+			SdfListOp<USDElementType> NewList;
+
+			bool bAddedItems = false;
+			for (int32 ListIndex = 0; ListIndex < EListOpListIndex::Num; ++ListIndex)
+			{
+				const TArray<FConvertedVtValueComponent>& UEItems = InValue.Entries[ListIndex];
+
+				std::vector<USDElementType> UsdItems;
+				UsdItems.reserve(UEItems.Num());
+
+				for (const FConvertedVtValueComponent& UEItem : UEItems)
+				{
+					if (const UEElementType* Value = UEItem.TryGet<UEElementType>())
+					{
+						UsdItems.push_back(Converter(*Value));
+					}
+				}
+
+				// It's very important that we set ExplicitItems first, and that we only set if they have at
+				// least one item in the list. This because SdfListOp has essentially two "modes":
+				// Explicit and not explicit. Any time we set explicit items and it wasn't in "explicit mode" yet,
+				// it switches to explicit mode and dumps all previous items. Analogously, if we set any of the
+				// non-explicit items and it was previously in "explicit mode", it dumps everything again and
+				// switches to "non-explicit mode". So starting with explicit and only calling SetItems when we
+				// have something prevents us from dumping items when we didn't mean to
+				if (UsdItems.size() > 0)
+				{
+					NewList.SetItems(UsdItems, (SdfListOpType)ListIndex);
+					bAddedItems = true;
+				}
+			}
+
+			if (bAddedItems)
+			{
+				OutValue = NewList;
+			}
+		}
+	}
+
+	template<typename SdfRefOrPayload>
+	SdfRefOrPayload UnstringifySdfReferenceOrPayload(const FString& String)
+	{
+		using namespace UE::USDValueConversion::Private;
+
+		TArray<FString> Elements;
+		const bool bCullEmpty = false;
+		int32 NumElements = String.ParseIntoArray(Elements, *ReferencePayloadDelimiter, bCullEmpty);
+
+		SdfRefOrPayload Result;
+		if (NumElements == 4)
+		{
+			if (Elements[0].Len() > 0)
+			{
+				Result.SetAssetPath(UnrealToUsd::ConvertString(*Elements[0]).Get());
+			}
+
+			if (Elements[1].Len() > 0)
+			{
+				Result.SetPrimPath(UnrealToUsd::ConvertPath(*Elements[1]).Get());
+			}
+
+			pxr::SdfLayerOffset LayerOffset;
+			{
+				double Offset = 0.0f;
+				if (LexTryParseString(Offset, *Elements[2]))
+				{
+					LayerOffset.SetOffset(Offset);
+				}
+
+				double Scale = 1.0f;
+				if (LexTryParseString(Scale, *Elements[3]))
+				{
+					LayerOffset.SetScale(Scale);
+				}
+			}
+			Result.SetLayerOffset(LayerOffset);
+		}
+		return Result;
+	}
 }	 // namespace UnrealToUsdImpl
 #endif	  // USE_USD_SDK
 
@@ -612,6 +863,7 @@ namespace UnrealToUsd
 		using namespace pxr;
 		using namespace UsdUtils;
 		using namespace UnrealToUsdImpl;
+		using namespace UE::USDValueConversion::Private;
 
 		// Always USD Allocs here because we may end up allocating arrays for these compound types,
 		// and those arrays can be moved into the UE::FVtValue's pxr::VtValue and, when set in an
@@ -653,10 +905,65 @@ namespace UnrealToUsd
 				ConvertSimpleValue<double, SdfTimeCode>(InValue, UsdValue);
 				break;
 			case EUsdBasicDataTypes::String:
-				ConvertSimpleValue<FString, std::string>(InValue, UsdValue);
+				if (InValue.ExtraTypeInfo == GetSdfTypeName<SdfListOp<std::string>>())
+				{
+					ConvertListOpValue<FString, std::string>(
+						InValue,
+						UsdValue,
+						[](const FString& Item) -> std::string
+						{
+							return UnrealToUsd::ConvertString(*Item).Get();
+						}
+					);
+				}
+				else if (InValue.ExtraTypeInfo == GetSdfTypeName<SdfListOp<SdfPath>>())
+				{
+					ConvertListOpValue<FString, SdfPath>(
+						InValue,
+						UsdValue,
+						[](const FString& Item) -> SdfPath
+						{
+							return UnrealToUsd::ConvertPath(*Item).Get();
+						}
+					);
+				}
+				else if (InValue.ExtraTypeInfo == GetSdfTypeName<SdfListOp<SdfReference>>())
+				{
+					ConvertListOpValue<FString, SdfReference>(	  //
+						InValue,
+						UsdValue,
+						UnstringifySdfReferenceOrPayload<SdfReference>
+					);
+				}
+				else if (InValue.ExtraTypeInfo == GetSdfTypeName<SdfListOp<SdfPayload>>())
+				{
+					ConvertListOpValue<FString, SdfPayload>(	//
+						InValue,
+						UsdValue,
+						UnstringifySdfReferenceOrPayload<SdfPayload>
+					);
+				}
+				else
+				{
+					ConvertSimpleValue<FString, std::string>(InValue, UsdValue);
+				}
 				break;
 			case EUsdBasicDataTypes::Token:
-				ConvertSimpleValue<FString, TfToken>(InValue, UsdValue);
+				if (InValue.ExtraTypeInfo == GetSdfTypeName<SdfListOp<TfToken>>())
+				{
+					ConvertListOpValue<FString, TfToken>(
+						InValue,
+						UsdValue,
+						[](const FString& Item) -> TfToken
+						{
+							return UnrealToUsd::ConvertToken(*Item).Get();
+						}
+					);
+				}
+				else
+				{
+					ConvertSimpleValue<FString, TfToken>(InValue, UsdValue);
+				}
 				break;
 			case EUsdBasicDataTypes::Asset:
 				ConvertSimpleValue<FString, SdfAssetPath>(InValue, UsdValue);
@@ -732,6 +1039,7 @@ FArchive& operator<<(FArchive& Ar, UsdUtils::FConvertedVtValue& Struct)
 {
 	Ar << Struct.Entries;
 	Ar << Struct.SourceType;
+	Ar << Struct.ExtraTypeInfo;
 	Ar << Struct.bIsArrayValued;
 	Ar << Struct.bIsEmpty;
 	return Ar;
@@ -1022,7 +1330,8 @@ namespace UsdUtils
 		const pxr::SdfListOp<pxr::TfToken>& TokenList = Value.UncheckedGet<pxr::SdfListOp<pxr::TfToken>>();
 
 		FString Result = UE::USDValueConversion::Private::TokenListStringStart;
-		for (const pxr::TfToken& Token : TokenList.GetExplicitItems())
+		TUsdStore<std::vector<pxr::TfToken>> TokenListItems = TokenList.GetAppliedItems();
+		for (const pxr::TfToken& Token : TokenListItems.Get())
 		{
 			const std::string& String = Token.GetString();
 
@@ -2202,7 +2511,8 @@ namespace UE::USDValueConversion::Private
 				FlattenedValues[FlattenedIndex + 0],
 				FlattenedValues[FlattenedIndex + 1],
 				FlattenedValues[FlattenedIndex + 2],
-				FlattenedValues[FlattenedIndex + 3]};
+				FlattenedValues[FlattenedIndex + 3]
+			};
 		}
 
 		return true;
@@ -2281,7 +2591,7 @@ namespace UsdUtils
 		using namespace UE::USDValueConversion::Private;
 
 		static const TMap<FString, TFunction<bool(const FString&, pxr::VtValue&)>> Unstringifiers = {
-	// Basic datatypes
+			// Basic datatypes
 			{TEXT("bool"), WrapInVtValue<bool, UnstringifyInner<bool>>},
 			{TEXT("uchar"), WrapInVtValue<uint8_t, UnstringifyInner<uint8_t>>},
 			{TEXT("int"), WrapInVtValue<int32_t, UnstringifyInner<int32_t>>},
@@ -2314,7 +2624,7 @@ namespace UsdUtils
 			{TEXT("half4"), WrapInVtValue<pxr::GfVec4h, UnstringifyVec<pxr::GfVec4h>>},
 			{TEXT("int4"), WrapInVtValue<pxr::GfVec4i, UnstringifyVec<pxr::GfVec4i>>},
 
- // Array versions of the ones above
+			// Array versions of the ones above
 			{TEXT("bool[]"), WrapInVtValue<pxr::VtArray<bool>, UnstringifyBasicArray<bool>>},
 			{TEXT("uchar[]"), WrapInVtValue<pxr::VtArray<uint8_t>, UnstringifyBasicArray<uint8_t>>},
 			{TEXT("int[]"), WrapInVtValue<pxr::VtArray<int32_t>, UnstringifyBasicArray<int32_t>>},
@@ -2347,9 +2657,9 @@ namespace UsdUtils
 			{TEXT("half4[]"), WrapInVtValue<pxr::VtArray<pxr::GfVec4h>, UnstringifyNumberArray<pxr::GfVec4h, 4>>},
 			{TEXT("int4[]"), WrapInVtValue<pxr::VtArray<pxr::GfVec4i>, UnstringifyNumberArray<pxr::GfVec4i, 4>>},
 
- // // Exotic types found in some scenarios
+			// Exotic types found in some scenarios
 			{TEXT("SdfListOp<TfToken>"), WrapInVtValue<pxr::SdfListOp<pxr::TfToken>, UnstringifyListOpTokens>}, // This is used for apiSchemas
-	// metadata
+																												   // metadata
 		};
 
 		if (const TFunction<bool(const FString&, pxr::VtValue&)>* FoundUnstringifier = Unstringifiers.Find(TypeName))
@@ -2967,7 +3277,7 @@ namespace UsdUtils
 		pxr::SdfListOp<pxr::TfToken> InnerValue;
 		if (UnstringifyListOpTokens(String, InnerValue))
 		{
-			const std::vector<pxr::TfToken>& ItemVector = InnerValue.GetExplicitItems();
+			std::vector<pxr::TfToken> ItemVector = InnerValue.GetAppliedItems();
 			OutValue.Empty(ItemVector.size());
 			for (const pxr::TfToken& InnerElement : ItemVector)
 			{
diff --git a/Engine/Plugins/Runtime/USDCore/Source/USDUtilities/Public/USDValueConversion.h b/Engine/Plugins/Runtime/USDCore/Source/USDUtilities/Public/USDValueConversion.h
--- a/Engine/Plugins/Runtime/USDCore/Source/USDUtilities/Public/USDValueConversion.h
+++ b/Engine/Plugins/Runtime/USDCore/Source/USDUtilities/Public/USDValueConversion.h
@@ -95,6 +95,8 @@ namespace UsdUtils
 		TArray<FConvertedVtValueEntry> Entries;
 
 		EUsdBasicDataTypes SourceType = EUsdBasicDataTypes::None;
+		FString ExtraTypeInfo;	  // Used so that we can serialize something like SdfListOp<Token> as basically a token array, but still get the same
+								  // type when converting back on the other end.
 		bool bIsArrayValued = false;
 		bool bIsEmpty = false;	  // Helps differentiating between empty arrays and having actually no value
 	};
diff --git a/Engine/Plugins/Runtime/USDCore/Source/UnrealUSDWrapper/Private/USDListener.cpp b/Engine/Plugins/Runtime/USDCore/Source/UnrealUSDWrapper/Private/USDListener.cpp
--- a/Engine/Plugins/Runtime/USDCore/Source/UnrealUSDWrapper/Private/USDListener.cpp
+++ b/Engine/Plugins/Runtime/USDCore/Source/UnrealUSDWrapper/Private/USDListener.cpp
@@ -141,15 +141,15 @@ namespace UE::USDListener::Private
 		UE_LOG(LogUsd, Log, TEXT("%sConverted ChangeListEntry:"), *Indent);
 
 		UE_LOG(LogUsd, Log, TEXT("%s\tAttribute changes:"), *Indent);
-		for (const UsdUtils::FAttributeChange& AttributeChange : Entry.AttributeChanges)
+		for (const UsdUtils::FAttributeChange& AttributeChange : Entry.FieldChanges)
 		{
 			UE_LOG(
 				LogUsd,
 				Log,
-				TEXT("%s\t\t'%s' field '%s': From '%s' to '%s'"),
+				TEXT("%s\t\tfield '%s' typename '%s': From '%s' to '%s'"),
 				*Indent,
-				*AttributeChange.PropertyName,
 				*AttributeChange.Field,
+				AttributeChange.OldValue.IsEmpty() ? *AttributeChange.NewValue.GetTypeName() : *AttributeChange.OldValue.GetTypeName(),
 				UTF8_TO_TCHAR(pxr::TfStringify(AttributeChange.OldValue.GetUsdValue()).c_str()),
 				UTF8_TO_TCHAR(pxr::TfStringify(AttributeChange.NewValue.GetUsdValue()).c_str())
 			);
@@ -363,42 +363,22 @@ namespace UE::USDListener::Private
 		for (const std::pair<TfToken, std::pair<VtValue, VtValue>>& AttributeChange : Entry.infoChanged)	// Note: infoChanged here is just a naming
 																											// conflict, there's no "resyncChanged"
 		{
-			UsdUtils::FAttributeChange& ConvertedAttributeChange = OutConverted.AttributeChanges.Emplace_GetRef();
-
-			const FString FieldToken = UTF8_TO_TCHAR(AttributeChange.first.GetString().c_str());
-
-			// For regular properties (most common case) FullFieldPath will already carry the property and FieldToken will be "default",
-			// "timeSamples", "variability", etc.
-			if (ChangedObject.IsPropertyPath())
-			{
-				ConvertedAttributeChange.PropertyName = UTF8_TO_TCHAR(ChangedObject.GetName().c_str());
-				ConvertedAttributeChange.Field = FieldToken;
-			}
-			// For stage properties it seems like USD likes to send just "/" as the FullFieldPath, and the actual property name in FieldToken (e.g.
-			// "metersPerUnit", "timeCodesPerSecond", etc.) We send "default" here for consistency
-			else if (ChangedObject.IsAbsoluteRootPath() || FieldToken == FString(TEXT("kind")))
-			{
-				ConvertedAttributeChange.PropertyName = FieldToken;
-				ConvertedAttributeChange.Field = TEXT("default");
-			}
-
+			UsdUtils::FAttributeChange& ConvertedAttributeChange = OutConverted.FieldChanges.Emplace_GetRef();
+			ConvertedAttributeChange.Field = UTF8_TO_TCHAR(AttributeChange.first.GetString().c_str());
 			ConvertedAttributeChange.OldValue = UE::FVtValue{AttributeChange.second.first};
 			ConvertedAttributeChange.NewValue = UE::FVtValue{AttributeChange.second.second};
 		}
 
-		// Some notices (like creating/removing a property) don't have any actual infoChanged entries, so we need to create a fake one in here in
-		// order to convey the PropertyName, if applicable
+		// Some notices (like creating/removing a property) don't have any actual infoChanged entries, so we create one in here for convenience
 		if (Entry.infoChanged.size() == 0
-			&& (Entry.flags.didAddProperty || Entry.flags.didAddPropertyWithOnlyRequiredFields || Entry.flags.didRemoveProperty
-				|| Entry.flags.didRemovePropertyWithOnlyRequiredFields || Entry.flags.didChangeAttributeTimeSamples))
+			&& (Entry.flags.didAddProperty ||							  //
+				Entry.flags.didAddPropertyWithOnlyRequiredFields ||		  //
+				Entry.flags.didRemoveProperty ||						  //
+				Entry.flags.didRemovePropertyWithOnlyRequiredFields ||	  //
+				Entry.flags.didChangeAttributeTimeSamples))
 		{
-			UsdUtils::FAttributeChange& ConvertedAttributeChange = OutConverted.AttributeChanges.Emplace_GetRef();
+			UsdUtils::FAttributeChange& ConvertedAttributeChange = OutConverted.FieldChanges.Emplace_GetRef();
 			ConvertedAttributeChange.Field = Entry.flags.didChangeAttributeTimeSamples ? TEXT("timeSamples") : TEXT("default");
-
-			if (ChangedObject.IsPropertyPath())
-			{
-				ConvertedAttributeChange.PropertyName = UTF8_TO_TCHAR(ChangedObject.GetName().c_str());
-			}
 		}
 
 		// These should be packed just the same, but just in case we do member by member here instead of memcopying it over
@@ -464,17 +444,14 @@ namespace UE::USDListener::Private
 
 		for (UsdNotice::ObjectsChanged::PathRange::const_iterator It = PathRange.begin(); It != PathRange.end(); ++It)
 		{
-			const FString PrimPath = UTF8_TO_TCHAR(It->GetAbsoluteRootOrPrimPath().GetAsString().c_str());
-
-			if (pxr::UsdPrim::IsPathInPrototype(UE::FSdfPath(*PrimPath)))
+			if (pxr::UsdPrim::IsPathInPrototype(It->GetAbsoluteRootOrPrimPath()))
 			{
 				continue;
 			}
 
-			// Something like "/Root/Prim.some_field", or "/"
-			const FString FullFieldPath = UTF8_TO_TCHAR(It->GetAsString().c_str());
-
-			TArray<UsdUtils::FSdfChangeListEntry>& ConvertedChanges = OutChanges.FindOrAdd(PrimPath);
+			// This may be a prim path, but also just a property path in case we're changing a property value or its metadata
+			const FString ObjectPath = UTF8_TO_TCHAR(It->GetAsString().c_str());
+			TArray<UsdUtils::FSdfChangeListEntry>& ConvertedChanges = OutChanges.FindOrAdd(ObjectPath);
 
 			// Changes may be empty, but we should still pass along this overall notice because sending a root
 			// resync notice with no actual change item inside is how USD signals that a layer has been added/removed/resynced
@@ -528,9 +505,12 @@ namespace UE::USDListener::Private
 		bool bReloadedContent = false;
 		for (TPair<FString, TArray<UsdUtils::FSdfChangeListEntry>>& InfoPair : OutInfoChanges)
 		{
-			const FString& PrimPath = InfoPair.Key;
+			const FString& ObjectPath = InfoPair.Key;
 			TArray<UsdUtils::FSdfChangeListEntry>* AnalogueResyncChanges = nullptr;
 
+			UE::FSdfPath UEPath{*ObjectPath};
+			FString ObjectName = UEPath.GetName();
+
 			for (TArray<UsdUtils::FSdfChangeListEntry>::TIterator ChangeIt = InfoPair.Value.CreateIterator(); ChangeIt; ++ChangeIt)
 			{
 				bool bUpgrade = false;
@@ -547,9 +527,9 @@ namespace UE::USDListener::Private
 				// the mesh will be completely regenerated however
 				if (!bUpgrade)
 				{
-					for (UsdUtils::FAttributeChange& Change : ChangeIt->AttributeChanges)
+					for (UsdUtils::FAttributeChange& Change : ChangeIt->FieldChanges)
 					{
-						if (Change.PropertyName == UTF8_TO_TCHAR(pxr::UsdGeomTokens->visibility.GetString().c_str()))
+						if (ObjectName == UTF8_TO_TCHAR(pxr::UsdGeomTokens->visibility.GetString().c_str()))
 						{
 							bUpgrade = true;
 							break;
@@ -561,7 +541,7 @@ namespace UE::USDListener::Private
 				{
 					if (!AnalogueResyncChanges)
 					{
-						AnalogueResyncChanges = &OutResyncChanges.FindOrAdd(PrimPath);
+						AnalogueResyncChanges = &OutResyncChanges.FindOrAdd(ObjectPath);
 					}
 
 					AnalogueResyncChanges->Add(*ChangeIt);
diff --git a/Engine/Plugins/Runtime/USDCore/Source/UnrealUSDWrapper/Private/UsdWrappers/SdfPath.cpp b/Engine/Plugins/Runtime/USDCore/Source/UnrealUSDWrapper/Private/UsdWrappers/SdfPath.cpp
--- a/Engine/Plugins/Runtime/USDCore/Source/UnrealUSDWrapper/Private/UsdWrappers/SdfPath.cpp
+++ b/Engine/Plugins/Runtime/USDCore/Source/UnrealUSDWrapper/Private/UsdWrappers/SdfPath.cpp
@@ -43,7 +43,7 @@ namespace UE
 			TUsdStore<pxr::SdfPath> PxrSdfPath;
 #endif	  // #if USE_USD_SDK
 		};
-	}	  // namespace Internal
+	}	 // namespace Internal
 
 	FSdfPath::FSdfPath()
 	{
@@ -247,6 +247,33 @@ namespace UE
 #endif	  // #if USE_USD_SDK
 	}
 
+	bool FSdfPath::IsPropertyPath() const
+	{
+#if USE_USD_SDK
+		return Impl->PxrSdfPath.Get().IsPropertyPath();
+#else
+		return false;
+#endif	  // #if USE_USD_SDK
+	}
+
+	bool FSdfPath::IsRelationalAttributePath() const
+	{
+#if USE_USD_SDK
+		return Impl->PxrSdfPath.Get().IsRelationalAttributePath();
+#else
+		return false;
+#endif	  // #if USE_USD_SDK
+	}
+
+	bool FSdfPath::IsTargetPath() const
+	{
+#if USE_USD_SDK
+		return Impl->PxrSdfPath.Get().IsTargetPath();
+#else
+		return false;
+#endif	  // #if USE_USD_SDK
+	}
+
 	FString FSdfPath::GetName() const
 	{
 #if USE_USD_SDK
diff --git a/Engine/Plugins/Runtime/USDCore/Source/UnrealUSDWrapper/Private/UsdWrappers/UsdPrim.cpp b/Engine/Plugins/Runtime/USDCore/Source/UnrealUSDWrapper/Private/UsdWrappers/UsdPrim.cpp
--- a/Engine/Plugins/Runtime/USDCore/Source/UnrealUSDWrapper/Private/UsdWrappers/UsdPrim.cpp
+++ b/Engine/Plugins/Runtime/USDCore/Source/UnrealUSDWrapper/Private/UsdWrappers/UsdPrim.cpp
@@ -11,6 +11,7 @@
 #include "UsdWrappers/UsdRelationship.h"
 #include "UsdWrappers/UsdStage.h"
 #include "UsdWrappers/UsdVariantSets.h"
+#include "UsdWrappers/VtValue.h"
 
 #if USE_USD_SDK
 #include "USDIncludesStart.h"
@@ -54,7 +55,7 @@ namespace UE
 			TUsdStore<pxr::UsdPrim> PxrUsdPrim;
 #endif	  // #if USE_USD_SDK
 		};
-	}	  // namespace Internal
+	}	 // namespace Internal
 
 	FUsdPrim::FUsdPrim()
 	{
@@ -156,6 +157,42 @@ namespace UE
 	}
 #endif	  // #if USE_USD_SDK
 
+	bool FUsdPrim::GetMetadata(const TCHAR* Key, UE::FVtValue& Value) const
+	{
+#if USE_USD_SDK
+		return Impl->PxrUsdPrim.Get().GetMetadata(pxr::TfToken{TCHAR_TO_UTF8(Key)}, &Value.GetUsdValue());
+#else
+		return false;
+#endif	  // #if USE_USD_SDK
+	}
+
+	bool FUsdPrim::HasMetadata(const TCHAR* Key) const
+	{
+#if USE_USD_SDK
+		return Impl->PxrUsdPrim.Get().HasMetadata(pxr::TfToken{TCHAR_TO_UTF8(Key)});
+#else
+		return false;
+#endif	  // #if USE_USD_SDK
+	}
+
+	bool FUsdPrim::SetMetadata(const TCHAR* Key, const UE::FVtValue& Value) const
+	{
+#if USE_USD_SDK
+		return Impl->PxrUsdPrim.Get().SetMetadata(pxr::TfToken{TCHAR_TO_UTF8(Key)}, Value.GetUsdValue());
+#else
+		return false;
+#endif	  // #if USE_USD_SDK
+	}
+
+	bool FUsdPrim::ClearMetadata(const TCHAR* Key) const
+	{
+#if USE_USD_SDK
+		return Impl->PxrUsdPrim.Get().ClearMetadata(pxr::TfToken{TCHAR_TO_UTF8(Key)});
+#else
+		return false;
+#endif	  // #if USE_USD_SDK
+	}
+
 	bool FUsdPrim::SetSpecifier(ESdfSpecifier Specifier)
 	{
 #if USE_USD_SDK
diff --git a/Engine/Plugins/Runtime/USDCore/Source/UnrealUSDWrapper/Public/USDListener.h b/Engine/Plugins/Runtime/USDCore/Source/UnrealUSDWrapper/Public/USDListener.h
--- a/Engine/Plugins/Runtime/USDCore/Source/UnrealUSDWrapper/Public/USDListener.h
+++ b/Engine/Plugins/Runtime/USDCore/Source/UnrealUSDWrapper/Public/USDListener.h
@@ -56,13 +56,13 @@ namespace UsdUtils
 	 * Analogous to pxr::SdfChangeList::Entry::InfoChange, describes a change to an attribute.
 	 * Here we break off PropertyName and Field for simplicity
 	 */
-	struct FAttributeChange
+	struct FFieldChange
 	{
-		FString PropertyName;	  // metersPerUnit, kind, upAxis, etc.
-		FString Field;			  // default, variability, timeSamples, etc.
-		UE::FVtValue OldValue;	  // Can be empty when we create a new attribute opinion
-		UE::FVtValue NewValue;	  // Can be empty when we clear an existing attribute opinion
+		FString Field;						  // default, variability, timeSamples, etc.
+		UE::FVtValue OldValue;				  // Can be empty when we create a new attribute opinion
+		UE::FVtValue NewValue;				  // Can be empty when we clear an existing attribute opinion
 	};
+	using FAttributeChange = FFieldChange;	  // Renamed in 5.6 as this can refer to a relationship, or even prim metadata
 
 	/** Analogous to pxr::SdfChangeList::SubLayerChangeType, describes a change to a sublayer */
 	enum ESubLayerChangeType
@@ -75,7 +75,7 @@ namespace UsdUtils
 	/** Analogous to pxr::SdfChangeList::Entry, describes a generic change to an object */
 	struct FSdfChangeListEntry
 	{
-		TArray<FAttributeChange> AttributeChanges;
+		TArray<FFieldChange> FieldChanges;
 		FPrimChangeFlags Flags;
 		FString OldPath;		  // Empty if Flags.bDidRename is not set
 		FString OldIdentifier;	  // Empty if Flags.bDIdChangeIdentifier is not set
@@ -88,9 +88,7 @@ namespace UsdUtils
 	using FLayerToSdfChangeList = TArray<TPair<UE::FSdfLayerWeak, FSdfChangeList>>;
 
 	/**
-	 * Describes USD object changes by object path.
-	 * The only difference to the USD data structures is that we store them by object path here for convenience, so
-	 * the key can be "/" to signify a stage change, or something like "/MyRoot/SomePrim" to indicate a particular prim change
+	 * Describes USD object changes by object path (an object can be a prim, a property, etc.)
 	 */
 	using FObjectChangesByPath = TMap<FString, TArray<FSdfChangeListEntry>>;
 }	 // namespace UsdUtils
diff --git a/Engine/Plugins/Runtime/USDCore/Source/UnrealUSDWrapper/Public/UsdWrappers/SdfPath.h b/Engine/Plugins/Runtime/USDCore/Source/UnrealUSDWrapper/Public/UsdWrappers/SdfPath.h
--- a/Engine/Plugins/Runtime/USDCore/Source/UnrealUSDWrapper/Public/UsdWrappers/SdfPath.h
+++ b/Engine/Plugins/Runtime/USDCore/Source/UnrealUSDWrapper/Public/UsdWrappers/SdfPath.h
@@ -68,6 +68,9 @@ namespace UE
 		bool IsPrimPath() const;
 		bool IsAbsoluteRootOrPrimPath() const;
 		bool IsPrimPropertyPath() const;
+		bool IsPropertyPath() const;
+		bool IsRelationalAttributePath() const;
+		bool IsTargetPath() const;
 		FString GetName() const;
 		FString GetElementString() const;
 		FSdfPath GetAbsoluteRootOrPrimPath() const;
diff --git a/Engine/Plugins/Runtime/USDCore/Source/UnrealUSDWrapper/Public/UsdWrappers/UsdPrim.h b/Engine/Plugins/Runtime/USDCore/Source/UnrealUSDWrapper/Public/UsdWrappers/UsdPrim.h
--- a/Engine/Plugins/Runtime/USDCore/Source/UnrealUSDWrapper/Public/UsdWrappers/UsdPrim.h
+++ b/Engine/Plugins/Runtime/USDCore/Source/UnrealUSDWrapper/Public/UsdWrappers/UsdPrim.h
@@ -29,6 +29,7 @@ namespace UE
 	class FUsdVariantSet;
 	class FUsdVariantSets;
 	class FUsdRelationship;
+	class FVtValue;
 
 	namespace Internal
 	{
@@ -76,7 +77,14 @@ namespace UE
 		operator const pxr::UsdPrim&() const;
 #endif	  // #if USE_USD_SDK
 
-		  // Wrapped pxr::UsdPrim functions, refer to the USD SDK documentation
+		  // Wrapped pxr::UsdObject functions, refer to the USD SDK documentation
+	public:
+		bool GetMetadata(const TCHAR* Key, UE::FVtValue& Value) const;
+		bool HasMetadata(const TCHAR* Key) const;
+		bool SetMetadata(const TCHAR* Key, const UE::FVtValue& Value) const;
+		bool ClearMetadata(const TCHAR* Key) const;
+
+		// Wrapped pxr::UsdPrim functions, refer to the USD SDK documentation
 	public:
 		bool SetSpecifier(ESdfSpecifier Specifier);
 
