USD: Fix stage actor losing connection to its spawned actors and components after recompiling its blueprint.

diff --git a/Engine/Plugins/Importers/USDImporter/Source/USDStage/Private/USDPrimTwin.cpp b/Engine/Plugins/Importers/USDImporter/Source/USDStage/Private/USDPrimTwin.cpp
--- a/Engine/Plugins/Importers/USDImporter/Source/USDStage/Private/USDPrimTwin.cpp
+++ b/Engine/Plugins/Importers/USDImporter/Source/USDStage/Private/USDPrimTwin.cpp
@@ -200,10 +200,11 @@ void UUsdPrimTwin::Serialize(FArchive& Ar)
 {
 	Super::Serialize(Ar);
 
-	// Only when going into/out of PIE we'll duplicate over the prim twin hierarchy.
-	// When duplicating for other reasons our StageActor will LoadUsdStage(), and generate
-	// new prim twins anyway (see AUsdStageActor::PostDuplicate)
-	if (Ar.GetPortFlags() & PPF_DuplicateForPIE)
+	// Keep properties transient as we never want them to be saved to disk, but we need to
+	// duplicate them here so that whenever we duplicate the root twin we also bring our
+	// prim twin hierarchy with it (happens when going into PIE, or when calling DuplicateObject inside
+	// AUsdStageActor::OnObjectsReplaced)
+	if (Ar.GetPortFlags() & PPF_Duplicate)
 	{
 		Ar << Children;
 		Ar << Parent;
diff --git a/Engine/Plugins/Importers/USDImporter/Source/USDStage/Private/USDStageActor.cpp b/Engine/Plugins/Importers/USDImporter/Source/USDStage/Private/USDStageActor.cpp
--- a/Engine/Plugins/Importers/USDImporter/Source/USDStage/Private/USDStageActor.cpp
+++ b/Engine/Plugins/Importers/USDImporter/Source/USDStage/Private/USDStageActor.cpp
@@ -3193,27 +3193,22 @@ void AUsdStageActor::OnObjectsReplaced(const TMap<UObject*, UObject*>& ObjectRep
 
 			NewActor->AssetCache = AssetCache;
 
-			// Move our subobjects to NewActor
-			auto MoveSubobject = [NewActor]<typename T>(TObjectPtr<T>& SourceObjectPtr, TObjectPtr<T>& DestObjectPtr)
-			{
-				if (SourceObjectPtr)
-				{
-					const TCHAR* NewName = nullptr;
-
-					if (DestObjectPtr)
-					{
-						DestObjectPtr->Rename(NewName, GetTransientPackage());
-					}
-
-					DestObjectPtr = SourceObjectPtr;
-					SourceObjectPtr->Rename(NewName, NewActor);
+			// We used to just move our subobjects to the new actor here, but from the outside that doesn't look very good:
+			// We'd have the new instance of the actor pointing at the subobjects of the old instance...
+			// Even though we used to rename/reparent the subobjects properly, the code that calls this function
+			// (UEngine::CopyPropertiesForUnrelatedObjects) builds a map from remapped old objects to new objects before it calls us.
+			// Later on, the engine will use stuff like FindAndReplaceReferences to always remap from old subobjects to new subobjects,
+			// undoing any UObject swapping we did here.
+			//
+			// Even if it were possible to influence/prepare for CopyPropertiesForUnrelatedObjects, it's probably a bad idea anyway:
+			// We just want to be a "well behaved UObject" here and not go against the grain. So instead we will now
+			// duplicate/stomp the new objects using the old objects. It seems slightly more aggressive, but from the outside
+			// looking in it should seem better behaved as we're just keeping our new subobject. It should be slower of course,
+			// but recompiling a blueprint deriving from the stage actor is not exactly a hot path
+			NewActor->RootUsdTwin = DuplicateObject(RootUsdTwin, NewActor, NewActor->RootUsdTwin.GetFName());
+			NewActor->UsdInfoCache = DuplicateObject(UsdInfoCache, NewActor, NewActor->UsdInfoCache.GetFName());
+			NewActor->PrimLinkCache = DuplicateObject(PrimLinkCache, NewActor, NewActor->PrimLinkCache.GetFName());
 
-					SourceObjectPtr = nullptr;
-				}
-			};
-			MoveSubobject(RootUsdTwin, NewActor->RootUsdTwin);
-			MoveSubobject(UsdInfoCache, NewActor->UsdInfoCache);
-			MoveSubobject(PrimLinkCache, NewActor->PrimLinkCache);
 			// We can just keep NewActor's transactor though
 
 			NewActor->BBoxCache = BBoxCache;
