USD: Fix crash when parsing skeletal meshes with no mesh data.

diff --git a/Engine/Plugins/Importers/USDImporter/Source/USDSchemas/Private/USDSkelSkeletonTranslator.cpp b/Engine/Plugins/Importers/USDImporter/Source/USDSchemas/Private/USDSkelSkeletonTranslator.cpp
--- a/Engine/Plugins/Importers/USDImporter/Source/USDSchemas/Private/USDSkelSkeletonTranslator.cpp
+++ b/Engine/Plugins/Importers/USDImporter/Source/USDSchemas/Private/USDSkelSkeletonTranslator.cpp
@@ -849,7 +849,10 @@ namespace UsdSkelSkeletonTranslatorImpl
 		}
 		if (UserData)
 		{
-			LODIndexToAssignments[0].PrimvarToUVIndex = UserData->PrimvarToUVIndex;
+			if (LODIndexToAssignments.Num() > 0)
+			{
+				LODIndexToAssignments[0].PrimvarToUVIndex = UserData->PrimvarToUVIndex;
+			}
 		}
 		else
 		{
@@ -1664,7 +1667,11 @@ namespace UsdSkelSkeletonTranslatorImpl
 										TEXT("Failed to create PhysicsAsset for skeletal mesh '%s'"),
 										*SkeletalMesh->GetPathName()
 									);
-									UsdUnreal::TranslatorUtils::AbandonFailedAsset(PhysicsAsset, Context->UsdAssetCache.Get(), Context->PrimLinkCache);
+									UsdUnreal::TranslatorUtils::AbandonFailedAsset(
+										PhysicsAsset,
+										Context->UsdAssetCache.Get(),
+										Context->PrimLinkCache
+									);
 								}
 							}
 						}
