USD: Keep track of previously added schema names and show those on the Add Schema list.

diff --git a/Engine/Plugins/Importers/USDImporter/Source/USDStageEditor/Private/SUSDStageTreeView.cpp b/Engine/Plugins/Importers/USDImporter/Source/USDStageEditor/Private/SUSDStageTreeView.cpp
--- a/Engine/Plugins/Importers/USDImporter/Source/USDStageEditor/Private/SUSDStageTreeView.cpp
+++ b/Engine/Plugins/Importers/USDImporter/Source/USDStageEditor/Private/SUSDStageTreeView.cpp
@@ -11,6 +11,7 @@
 #include "USDMemory.h"
 #include "USDOptionsWindow.h"
 #include "USDPrimViewModel.h"
+#include "USDProjectSettings.h"
 #include "USDReferenceOptions.h"
 #include "USDTypesConversion.h"
 #include "UsdWrappers/SdfChangeBlock.h"
@@ -1774,7 +1775,13 @@ void SUsdStageTreeView::FillCollapsingSubmenu(FMenuBuilder& MenuBuilder)
 
 void SUsdStageTreeView::FillAddSchemaSubmenu(FMenuBuilder& MenuBuilder)
 {
-	static TArray<TSharedPtr<FString>> AllKnownSchemas = {
+	const UUsdProjectSettings* ProjectSettings = GetDefault<UUsdProjectSettings>();
+	if (!ProjectSettings)
+	{
+		return;
+	}
+
+	static TArray<TSharedPtr<FString>> DefaultSchemas = {
 		MakeShared<FString>(UsdToUnreal::ConvertToken(UnrealIdentifiers::ControlRigAPI)),
 		MakeShared<FString>(UsdToUnreal::ConvertToken(UnrealIdentifiers::GroomAPI)),
 		MakeShared<FString>(UsdToUnreal::ConvertToken(UnrealIdentifiers::GroomBindingAPI)),
@@ -1786,17 +1793,49 @@ void SUsdStageTreeView::FillAddSchemaSubmenu(FMenuBuilder& MenuBuilder)
 		MakeShared<FString>(UsdToUnreal::ConvertToken(pxr::UsdSkelTokens->SkelBindingAPI)),
 	};
 
-	// Only show on the list the schemas that can be applied to the selected prim
 	static TArray<TSharedPtr<FString>> AllowedKnownSchemas;
-	AllowedKnownSchemas.Reset(AllKnownSchemas.Num());
-	for (const TSharedPtr<FString>& KnownSchema : AllKnownSchemas)
+	AllowedKnownSchemas.Reset(DefaultSchemas.Num() + ProjectSettings->AdditionalCustomSchemaNames.Num());
+
+	static TSet<FString> SeenSchemas;
+	SeenSchemas.Reset();
+	SeenSchemas.Reserve(AllowedKnownSchemas.Num());
+
+	// Add default list of known schemas
+	for (const TSharedPtr<FString>& KnownSchema : DefaultSchemas)
 	{
+		SeenSchemas.Add(*KnownSchema);
+
+		// Only show on the list the schemas that can be applied to the selected prim
 		if (CanApplySchema(**KnownSchema))
 		{
 			AllowedKnownSchemas.Add(KnownSchema);
 		}
 	}
 
+	// Add additional list of custom schemas
+	for (const FString& CustomSchema : ProjectSettings->AdditionalCustomSchemaNames)
+	{
+		if (SeenSchemas.Contains(CustomSchema))
+		{
+			continue;
+		}
+		SeenSchemas.Add(*CustomSchema);
+
+		// Only show on the list the schemas that can be applied to the selected prim
+		if (CanApplySchema(*CustomSchema))
+		{
+			AllowedKnownSchemas.Add(MakeShared<FString>(CustomSchema));
+		}
+	}
+
+	// Sort list for a consistent order
+	AllowedKnownSchemas.Sort(
+		[](const TSharedPtr<FString>& LHS, const TSharedPtr<FString>& RHS)
+		{
+			return *LHS < *RHS;
+		}
+	);
+
 	static TSharedPtr<FString> CurrentKnownSchema = nullptr;
 	if (AllowedKnownSchemas.Num() > 0)
 	{
@@ -1876,10 +1915,23 @@ void SUsdStageTreeView::FillAddSchemaSubmenu(FMenuBuilder& MenuBuilder)
 			})
 			.OnClicked_Lambda([this]() -> FReply
 			{
-				FName SelectedSchema = *GetCurrentSchemaNameText().ToString();
-				if (CanApplySchema(SelectedSchema))
+				FText SelectedSchemaText = GetCurrentSchemaNameText();
+				FString SelectedSchemaString = SelectedSchemaText.ToString();
+				FName SelectedSchemaName = *SelectedSchemaString;
+
+				if (CanApplySchema(SelectedSchemaName))
+				{
+					OnApplySchema(SelectedSchemaName);
+				}
+
+				if (!SeenSchemas.Contains(SelectedSchemaString))
 				{
-					OnApplySchema(SelectedSchema);
+					UUsdProjectSettings* ProjectSettings = GetMutableDefault<UUsdProjectSettings>();
+					if (ProjectSettings)
+					{
+						ProjectSettings->AdditionalCustomSchemaNames.AddUnique(SelectedSchemaText.ToString());
+						ProjectSettings->SaveConfig();
+					}
 				}
 
 				return FReply::Handled();
diff --git a/Engine/Plugins/Runtime/USDCore/Source/USDClasses/Public/USDProjectSettings.h b/Engine/Plugins/Runtime/USDCore/Source/USDClasses/Public/USDProjectSettings.h
--- a/Engine/Plugins/Runtime/USDCore/Source/USDClasses/Public/USDProjectSettings.h
+++ b/Engine/Plugins/Runtime/USDCore/Source/USDClasses/Public/USDProjectSettings.h
@@ -64,6 +64,10 @@ public:
 	UPROPERTY(config, EditAnywhere, Category = USD)
 	TArray<FName> AdditionalMaterialPurposes;
 
+	// apiSchema names to show on the "Add Schema" right-click option, in addition to the standard common schemas
+	UPROPERTY(config, EditAnywhere, Category = USD)
+	TArray<FString> AdditionalCustomSchemaNames;
+
 	// Whether to show on the output log messages, warnings and errors reported directly by the USD SDK
 	UPROPERTY(config, EditAnywhere, Category = USD)
 	bool bLogUsdSdkErrors = true;
