USD: Fix ensure when referenced .mdl material does not exist.

diff --git a/Engine/Plugins/Importers/USDImporter/Source/USDSchemas/Private/Custom/MDLUSDShadeMaterialTranslator.cpp b/Engine/Plugins/Importers/USDImporter/Source/USDSchemas/Private/Custom/MDLUSDShadeMaterialTranslator.cpp
--- a/Engine/Plugins/Importers/USDImporter/Source/USDSchemas/Private/Custom/MDLUSDShadeMaterialTranslator.cpp
+++ b/Engine/Plugins/Importers/USDImporter/Source/USDSchemas/Private/Custom/MDLUSDShadeMaterialTranslator.cpp
@@ -131,6 +131,7 @@ void FMdlUsdShadeMaterialTranslator::CreateAssets()
 		}
 	}
 
+	const FString MdlOrigAssetPath = UsdToUnreal::ConvertString(SurfaceSourceAssetPath.GetAssetPath());
 	const FString MdlAbsoluteAssetPath = UsdToUnreal::ConvertString(SurfaceSourceAssetPath.GetResolvedPath());
 
 	const FString MdlModuleName = [&MdlAbsoluteAssetPath, &MdlRootPath]()
@@ -173,7 +174,8 @@ void FMdlUsdShadeMaterialTranslator::CreateAssets()
 			MaterialHash,
 			MdlFullName,
 			Context->ObjectFlags,
-			[&MdlSearchPath, &MdlModuleName, &MdlDefinitionName](UPackage* Outer, FName SanitizedName, EObjectFlags FlagsToUse) -> UMaterialInterface*
+			[&MdlSearchPath, &MdlModuleName, &MdlDefinitionName, &MdlOrigAssetPath](UPackage* Outer, FName SanitizedName, EObjectFlags FlagsToUse)
+				-> UMaterialInterface*
 			{
 				// Add the USD root as a search path for MDL
 				FMdlMaterialImporter::FScopedSearchPath UsdDirMdlSearchPath(MdlSearchPath);
@@ -188,8 +190,11 @@ void FMdlUsdShadeMaterialTranslator::CreateAssets()
 					MdlDefinitionName,
 					*ImportOptions.Get()
 				);
-				if (!ensure(ReferenceMaterial))
+				if (!ReferenceMaterial)
 				{
+					FString FullPath = FPaths::Combine(MdlSearchPath, MdlOrigAssetPath);
+					FPaths::NormalizeFilename(FullPath);
+					UE_LOG(LogUsd, Warning, TEXT("Failed to load MDL material from file '%s'. Does the file exist?"), *FullPath);
 					return nullptr;
 				}
 
