USD: Fix ensure/error spam when force-deleting the asset cache used by a stage actor.

diff --git a/Engine/Plugins/Importers/USDImporter/Source/USDStage/Private/USDStageActor.cpp b/Engine/Plugins/Importers/USDImporter/Source/USDStage/Private/USDStageActor.cpp
--- a/Engine/Plugins/Importers/USDImporter/Source/USDStage/Private/USDStageActor.cpp
+++ b/Engine/Plugins/Importers/USDImporter/Source/USDStage/Private/USDStageActor.cpp
@@ -2581,32 +2581,6 @@ void AUsdStageActor::SetUsdAssetCache(UUsdAssetCache3* NewCache)
 	}
 
 	AssetCache = NewCache;
-
-	// We can't have no cache while we have a stage loaded, so fetch a new one from the IUsdClassesModule
-	// instead, as the intent may have been to just have the actor not point at the previous cache anymore.
-	if (!AssetCache && UsdStage)
-	{
-		FNotificationInfo Toast(LOCTEXT("MustHaveCache", "Must have an Asset Cache"));
-		Toast.SubText = LOCTEXT(
-			"MustHaveCache_Subtext",
-			"The Stage Actor must always have an Asset Cache while a stage is loaded, so a temporary cache will be created.\n\nClose the stage "
-			"before clearing the cache if you wish to clear this property."
-		);
-		Toast.Image = FCoreStyle::Get().GetBrush(TEXT("MessageLog.Warning"));
-		Toast.bUseLargeFont = false;
-		Toast.bFireAndForget = true;
-		Toast.FadeOutDuration = 1.0f;
-		Toast.ExpireDuration = 12.0f;
-		Toast.bUseThrobber = false;
-		Toast.bUseSuccessFailIcons = false;
-		FSlateNotificationManager::Get().AddNotification(Toast);
-
-		SetupAssetCacheIfNeeded();
-	}
-
-	// We used to trigger a reload here, forcing new assets to cached into the new cache.
-	// There's no longer any real need for it? If the assets are used by components the object
-	// references will be enough to prevent the old assets from being deleted anyway
 }
 
 void AUsdStageActor::SetInitialLoadSet(EUsdInitialLoadSet NewLoadSet)
@@ -3507,6 +3481,20 @@ void AUsdStageActor::SetupAssetCacheIfNeeded()
 {
 	if (!AssetCache)
 	{
+		FNotificationInfo Toast(LOCTEXT("MustHaveCache3", "Must have an Asset Cache"));
+		Toast.SubText = LOCTEXT(
+			"MustHaveCache3_Subtext",
+			"The Stage Actor must always have an Asset Cache while a stage is loaded, so a new cache will be created."
+		);
+		Toast.Image = FCoreStyle::Get().GetBrush(TEXT("MessageLog.Warning"));
+		Toast.bUseLargeFont = false;
+		Toast.bFireAndForget = true;
+		Toast.FadeOutDuration = 1.0f;
+		Toast.ExpireDuration = 8.0f;
+		Toast.bUseThrobber = false;
+		Toast.bUseSuccessFailIcons = false;
+		FSlateNotificationManager::Get().AddNotification(Toast);
+
 		TGuardValue<ITransaction*> SuppressTransaction{GUndo, nullptr};
 		AssetCache = IUsdClassesModule::GetAssetCacheForProject();
 	}
