USD: Fix warning about UnrealUSDWrapper module not being loaded when cooking.

diff --git a/Engine/Plugins/Importers/USDImporter/Source/USDStageImporter/Private/USDStageAssetImportFactory.cpp b/Engine/Plugins/Importers/USDImporter/Source/USDStageImporter/Private/USDStageAssetImportFactory.cpp
--- a/Engine/Plugins/Importers/USDImporter/Source/USDStageImporter/Private/USDStageAssetImportFactory.cpp
+++ b/Engine/Plugins/Importers/USDImporter/Source/USDStageImporter/Private/USDStageAssetImportFactory.cpp
@@ -16,6 +16,7 @@
 #include "Engine/SkeletalMesh.h"
 #include "Engine/StaticMesh.h"
 #include "Misc/Paths.h"
+#include "Modules/ModuleManager.h"
 
 #define LOCTEXT_NAMESPACE "USDStageAssetImportFactory"
 
@@ -34,6 +35,7 @@ UUsdStageAssetImportFactory::UUsdStageAssetImportFactory(const FObjectInitialize
 	bEditorImport = true;
 	bText = false;
 
+	FModuleManager::Get().LoadModuleChecked(TEXT("UnrealUSDWrapper"));
 	UnrealUSDWrapper::AddUsdImportFileFormatDescriptions(Formats);
 }
 
diff --git a/Engine/Plugins/Importers/USDImporter/Source/USDStageImporter/Private/USDStageImportFactory.cpp b/Engine/Plugins/Importers/USDImporter/Source/USDStageImporter/Private/USDStageImportFactory.cpp
--- a/Engine/Plugins/Importers/USDImporter/Source/USDStageImporter/Private/USDStageImportFactory.cpp
+++ b/Engine/Plugins/Importers/USDImporter/Source/USDStageImporter/Private/USDStageImportFactory.cpp
@@ -9,6 +9,7 @@
 
 #include "AssetImportTask.h"
 #include "Editor.h"
+#include "Modules/ModuleManager.h"
 
 #define LOCTEXT_NAMESPACE "USDImportFactory"
 
@@ -24,6 +25,7 @@ UUsdStageImportFactory::UUsdStageImportFactory(const FObjectInitializer& ObjectI
 	bEditorImport = true;
 	bText = false;
 
+	FModuleManager::Get().LoadModuleChecked(TEXT("UnrealUSDWrapper"));
 	UnrealUSDWrapper::AddUsdImportFileFormatDescriptions(Formats);
 }
 
@@ -66,7 +68,7 @@ UObject* UUsdStageImportFactory::FactoryCreateFile(
 		GEditor->RedrawLevelEditingViewports();
 
 		ImportedObject = ImportContext.ImportedAsset ? ToRawPtr(ImportContext.ImportedAsset) : Cast<UObject>(ImportContext.SceneActor);
-		
+
 		ImportContext.ImportedAssets.Remove(ImportedObject);
 		AdditionalImportedObjects = ImportContext.ImportedAssets;
 	}
diff --git a/Engine/Plugins/Interchange/Runtime/Source/Import/Private/Usd/InterchangeUsdTranslator.cpp b/Engine/Plugins/Interchange/Runtime/Source/Import/Private/Usd/InterchangeUsdTranslator.cpp
--- a/Engine/Plugins/Interchange/Runtime/Source/Import/Private/Usd/InterchangeUsdTranslator.cpp
+++ b/Engine/Plugins/Interchange/Runtime/Source/Import/Private/Usd/InterchangeUsdTranslator.cpp
@@ -42,6 +42,7 @@
 #include "InterchangeTexture2DNode.h"
 #include "InterchangeTranslatorHelper.h"
 #include "Mesh/InterchangeMeshPayload.h"
+#include "Modules/ModuleManager.h"
 #include "MovieSceneSection.h"
 #include "Rendering/SkeletalMeshLODImporterData.h"
 #include "StaticMeshAttributes.h"
@@ -2455,6 +2456,7 @@ TArray<FString> UInterchangeUSDTranslator::GetSupportedFormats() const
 	TArray<FString> Extensions;
 	if (GInterchangeEnableUSDImport)
 	{
+		FModuleManager::Get().LoadModuleChecked(TEXT("UnrealUSDWrapper"));
 		UnrealUSDWrapper::AddUsdImportFileFormatDescriptions(Extensions);
 	}
 	return Extensions;
