An Error occurred while cancelling the Reimport Process for static mesh asset
- Usd Stage Import Context initialization returns result as Cancelled instead of Failed.

#rnx

diff --git a/Engine/Plugins/Importers/USDImporter/Source/USDStageImporter/Private/USDStageAssetImportFactory.cpp b/Engine/Plugins/Importers/USDImporter/Source/USDStageImporter/Private/USDStageAssetImportFactory.cpp
--- a/Engine/Plugins/Importers/USDImporter/Source/USDStageImporter/Private/USDStageAssetImportFactory.cpp
+++ b/Engine/Plugins/Importers/USDImporter/Source/USDStageImporter/Private/USDStageAssetImportFactory.cpp
@@ -213,11 +213,11 @@ EReimportResult::Type UUsdStageAssetImportFactory::Reimport(UObject* Obj)
 		FUsdLogManager::LogMessage(
 			EMessageSeverity::Error,
 			FText::Format(
-				LOCTEXT("ReimportErrorNoContext", "Failed to initialize reimport context for asset '{0}'!"),
+				LOCTEXT("ReimportErrorNoContext", "Failed to initialize re-import context for asset '{0}'!"),
 				FText::FromName(Obj->GetFName())
 			)
 		);
-		return EReimportResult::Failed;
+		return EReimportResult::Cancelled;
 	}
 
 	ImportContext.PackagePath = Obj->GetOutermost()->GetPathName();
