USD: Fix LevelSequenceHelperImpls subscribing to editor events when they didn't need to.

diff --git a/Engine/Plugins/Importers/USDImporter/Source/USDStage/Private/USDLevelSequenceHelper.cpp b/Engine/Plugins/Importers/USDImporter/Source/USDStage/Private/USDLevelSequenceHelper.cpp
--- a/Engine/Plugins/Importers/USDImporter/Source/USDStage/Private/USDLevelSequenceHelper.cpp
+++ b/Engine/Plugins/Importers/USDImporter/Source/USDStage/Private/USDLevelSequenceHelper.cpp
@@ -2,6 +2,7 @@
 
 #include "USDLevelSequenceHelper.h"
 
+#include "Objects/USDPrimLinkCache.h"
 #include "USDAssetUserData.h"
 #include "USDAttributeUtils.h"
 #include "USDConversionUtils.h"
@@ -12,7 +13,6 @@
 #include "USDLog.h"
 #include "USDObjectUtils.h"
 #include "USDPrimConversion.h"
-#include "Objects/USDPrimLinkCache.h"
 #include "USDPrimTwin.h"
 #include "USDProjectSettings.h"
 #include "USDSkeletalDataConversion.h"
@@ -817,6 +817,9 @@ public:
 	void SetRootMotionHandling(EUsdRootMotionHandling NewValue);
 	void OnStageActorRenamed();
 
+	void SubscribeToEditorEvents();
+	void UnsubscribeToEditorEvents();
+
 	ULevelSequence* GetMainLevelSequence() const
 	{
 		return MainLevelSequence;
@@ -989,7 +992,6 @@ private:
 	void HandleDisplayRateChange(const double DisplayRate);
 
 	FDelegateHandle OnObjectTransactedHandle;
-	FDelegateHandle OnStageEditTargetChangedHandle;
 	FDelegateHandle OnUsdObjectsChangedHandle;
 
 private:
@@ -1035,40 +1037,19 @@ const double FUsdLevelSequenceHelperImpl::EmptySubSectionRange = 10.0;
 FUsdLevelSequenceHelperImpl::FUsdLevelSequenceHelperImpl()
 	: MainLevelSequence(nullptr)
 {
-#if WITH_EDITOR
-	OnObjectTransactedHandle = FCoreUObjectDelegates::OnObjectTransacted.AddRaw(this, &FUsdLevelSequenceHelperImpl::OnObjectTransacted);
-
-	if (GEditor)
-	{
-		if (UTransBuffer* Transactor = Cast<UTransBuffer>(GEditor->Trans))
-		{
-			Transactor->OnTransactionStateChanged().AddRaw(this, &FUsdLevelSequenceHelperImpl::HandleTransactionStateChanged);
-		}
-	}
-#endif	  // WITH_EDITOR
+	// Don't subscribe to editor events here: The LevelSequenceHelper is a member struct of the stage actor
+	// and the USD Import Context, so we may be an Impl of a CDO, that can't really do anything with those
+	// events anyway. We'll subscribe only if/when we actually receive a stage (on Init())
 }
 
 FUsdLevelSequenceHelperImpl::~FUsdLevelSequenceHelperImpl()
 {
 	if (StageActor.IsValid())
 	{
-		StageActor->GetUsdListener().GetOnStageEditTargetChanged().Remove(OnStageEditTargetChangedHandle);
 		StageActor->GetUsdListener().GetOnObjectsChanged().Remove(OnUsdObjectsChangedHandle);
-		OnStageEditTargetChangedHandle.Reset();
 	}
 
-#if WITH_EDITOR
-	FCoreUObjectDelegates::OnObjectTransacted.Remove(OnObjectTransactedHandle);
-	OnObjectTransactedHandle.Reset();
-
-	if (GEditor)
-	{
-		if (UTransBuffer* Transactor = Cast<UTransBuffer>(GEditor->Trans))
-		{
-			Transactor->OnTransactionStateChanged().RemoveAll(this);
-		}
-	}
-#endif	  // WITH_EDITOR
+	UnsubscribeToEditorEvents();
 }
 
 ULevelSequence* FUsdLevelSequenceHelperImpl::Init(const UE::FUsdStage& InUsdStage)
@@ -1079,6 +1060,17 @@ ULevelSequence* FUsdLevelSequenceHelperImpl::Init(const UE::FUsdStage& InUsdStag
 
 	CreateLocalLayersSequences();
 
+	// We call Init with a default (invalid) stage to "clear", so only subscribe to these events
+	// if we've actually been given a valid stage
+	if (InUsdStage)
+	{
+		SubscribeToEditorEvents();
+	}
+	else
+	{
+		UnsubscribeToEditorEvents();
+	}
+
 	return MainLevelSequence;
 }
 
@@ -1288,15 +1280,12 @@ void FUsdLevelSequenceHelperImpl::UnbindFromUsdStageActor()
 
 	if (StageActor.IsValid())
 	{
-		StageActor->GetUsdListener().GetOnStageEditTargetChanged().Remove(OnStageEditTargetChangedHandle);
 		StageActor->GetUsdListener().GetOnObjectsChanged().Remove(OnUsdObjectsChangedHandle);
 		StageActor.Reset();
 	}
 
 	SetPrimLinkCache(nullptr);
 	SetRootMotionHandling(EUsdRootMotionHandling::NoAdditionalRootMotion);
-
-	OnStageEditTargetChangedHandle.Reset();
 }
 
 EUsdRootMotionHandling FUsdLevelSequenceHelperImpl::GetRootMotionHandling() const
@@ -1358,6 +1347,42 @@ void FUsdLevelSequenceHelperImpl::OnStageActorRenamed()
 	}
 }
 
+void FUsdLevelSequenceHelperImpl::SubscribeToEditorEvents()
+{
+#if WITH_EDITOR
+	UnsubscribeToEditorEvents();
+
+	if (GEditor)
+	{
+		OnObjectTransactedHandle = FCoreUObjectDelegates::OnObjectTransacted.AddRaw(this, &FUsdLevelSequenceHelperImpl::OnObjectTransacted);
+
+		if (UTransBuffer* Transactor = Cast<UTransBuffer>(GEditor->Trans))
+		{
+			Transactor->OnTransactionStateChanged().AddRaw(this, &FUsdLevelSequenceHelperImpl::HandleTransactionStateChanged);
+		}
+	}
+#endif	  // WITH_EDITOR
+}
+
+void FUsdLevelSequenceHelperImpl::UnsubscribeToEditorEvents()
+{
+#if WITH_EDITOR
+	if (GEditor)
+	{
+		if (OnObjectTransactedHandle.IsValid())
+		{
+			FCoreUObjectDelegates::OnObjectTransacted.Remove(OnObjectTransactedHandle);
+			OnObjectTransactedHandle.Reset();
+		}
+
+		if (UTransBuffer* Transactor = Cast<UTransBuffer>(GEditor->Trans))
+		{
+			Transactor->OnTransactionStateChanged().RemoveAll(this);
+		}
+	}
+#endif	  // WITH_EDITOR
+}
+
 ULevelSequence* FUsdLevelSequenceHelperImpl::FindSequenceForAttribute(const UE::FUsdAttribute& Attribute)
 {
 	if (!Attribute || !Attribute.GetPrim())
