GeometryCache: Refactor GeometryCache code for use with Interchange.

diff --git a/Engine/Plugins/Importers/USDImporter/Source/USDSchemas/Private/USDGeometryCacheTranslator.cpp b/Engine/Plugins/Importers/USDImporter/Source/USDSchemas/Private/USDGeometryCacheTranslator.cpp
--- a/Engine/Plugins/Importers/USDImporter/Source/USDSchemas/Private/USDGeometryCacheTranslator.cpp
+++ b/Engine/Plugins/Importers/USDImporter/Source/USDSchemas/Private/USDGeometryCacheTranslator.cpp
@@ -28,6 +28,7 @@
 #include "GeometryCache.h"
 #include "GeometryCacheCodecV1.h"
 #include "GeometryCacheComponent.h"
+#include "GeometryCacheHelpers.h"
 #include "GeometryCacheMeshData.h"
 #include "GeometryCacheTrackStreamable.h"
 #include "GeometryCacheTrackUSD.h"
@@ -163,58 +164,6 @@ namespace UsdGeometryCacheTranslatorImpl
 		}
 	}
 
-	void GeometryCacheDataForMeshDescription(
-		FGeometryCacheMeshData& OutMeshData,
-		FMeshDescription& MeshDescription,
-		int32 MaterialIndex,
-		const float SecondsPerFrame
-	);
-
-	void GetGeometryCacheDataTimeCodeRange(const UE::FUsdStage& Stage, const FString& PrimPath, int32& OutStartFrame, int32& OutEndFrame)
-	{
-		if (!Stage || PrimPath.IsEmpty())
-		{
-			return;
-		}
-
-		FScopedUsdAllocs Allocs;
-
-		pxr::UsdPrim UsdPrim = pxr::UsdPrim(Stage.GetPrimAtPath(UE::FSdfPath(*PrimPath)));
-		if (!UsdPrim)
-		{
-			return;
-		}
-
-		constexpr bool bIncludeInherited = false;
-		pxr::TfTokenVector GeomMeshAttributeNames = pxr::UsdGeomMesh::GetSchemaAttributeNames(bIncludeInherited);
-		pxr::TfTokenVector GeomPointBasedAttributeNames = pxr::UsdGeomPointBased::GetSchemaAttributeNames(bIncludeInherited);
-
-		GeomMeshAttributeNames.reserve(GeomMeshAttributeNames.size() + GeomPointBasedAttributeNames.size());
-		GeomMeshAttributeNames.insert(GeomMeshAttributeNames.end(), GeomPointBasedAttributeNames.begin(), GeomPointBasedAttributeNames.end());
-
-		TOptional<double> MinStartTimeCode;
-		TOptional<double> MaxEndTimeCode;
-
-		for (const pxr::TfToken& AttributeName : GeomMeshAttributeNames)
-		{
-			if (const pxr::UsdAttribute& Attribute = UsdPrim.GetAttribute(AttributeName))
-			{
-				std::vector<double> TimeSamples;
-				if (Attribute.GetTimeSamples(&TimeSamples) && TimeSamples.size() > 0)
-				{
-					MinStartTimeCode = FMath::Min(MinStartTimeCode.Get(TNumericLimits<double>::Max()), TimeSamples[0]);
-					MaxEndTimeCode = FMath::Max(MaxEndTimeCode.Get(TNumericLimits<double>::Lowest()), TimeSamples[TimeSamples.size() - 1]);
-				}
-			}
-		}
-
-		if (MinStartTimeCode.IsSet() && MaxEndTimeCode.IsSet())
-		{
-			OutStartFrame = FMath::FloorToInt(MinStartTimeCode.GetValue());
-			OutEndFrame = FMath::CeilToInt(MaxEndTimeCode.GetValue());
-		}
-	}
-
 	struct FReadMeshDataArgs
 	{
 		FReadMeshDataArgs(const UE::FUsdStage& InStage, const UE::FUsdPrim& InRootPrim)
@@ -242,7 +191,7 @@ namespace UsdGeometryCacheTranslatorImpl
 		// Fetch the animated mesh start/end frame as they may be different from just the stage's start and end time codes
 		int32 StartFrame = FMath::FloorToInt(Stage.GetStartTimeCode());
 		int32 EndFrame = FMath::CeilToInt(Stage.GetEndTimeCode());
-		UsdGeometryCacheTranslatorImpl::GetGeometryCacheDataTimeCodeRange(Stage, RootPrimPath, StartFrame, EndFrame);
+		UsdUtils::GetAnimatedMeshTimeCodes(Stage, RootPrimPath, StartFrame, EndFrame);
 
 		double FramesPerSecond = Stage.GetTimeCodesPerSecond();
 		if (FramesPerSecond == 0)
@@ -325,7 +274,7 @@ namespace UsdGeometryCacheTranslatorImpl
 
 			FStaticMeshOperations::ComputeTangentsAndNormals(MeshDescription, ComputeNTBsOptions);
 
-			UsdGeometryCacheTranslatorImpl::GeometryCacheDataForMeshDescription(OutMeshData, MeshDescription, MaterialOffset, Args.FramesPerSecond);
+			UE::GeometryCache::Utils::GetGeometryCacheMeshDataFromMeshDescription(OutMeshData, MeshDescription, MaterialOffset, Args.FramesPerSecond);
 
 			return true;
 		}
@@ -690,121 +639,6 @@ namespace UsdGeometryCacheTranslatorImpl
 			}
 		}
 	}
-
-	// #ueent_todo: Replace MeshDescription with RawMesh and also make it work with StaticMesh
-	void GeometryCacheDataForMeshDescription(
-		FGeometryCacheMeshData& OutMeshData,
-		FMeshDescription& MeshDescription,
-		int32 MaterialIndex,
-		float FramesPerSecond
-	)
-	{
-		TRACE_CPUPROFILER_EVENT_SCOPE(GeometryCacheDataForMeshDescription);
-
-		OutMeshData.Positions.Reset();
-		OutMeshData.TextureCoordinates.Reset();
-		OutMeshData.TangentsX.Reset();
-		OutMeshData.TangentsZ.Reset();
-		OutMeshData.Colors.Reset();
-		OutMeshData.Indices.Reset();
-
-		OutMeshData.MotionVectors.Reset();
-		OutMeshData.BatchesInfo.Reset();
-		OutMeshData.BoundingBox.Init();
-
-		OutMeshData.VertexInfo.bHasColor0 = true;
-		OutMeshData.VertexInfo.bHasTangentX = true;
-		OutMeshData.VertexInfo.bHasTangentZ = true;
-		OutMeshData.VertexInfo.bHasUV0 = true;
-
-		FStaticMeshAttributes MeshDescriptionAttributes(MeshDescription);
-
-		TVertexAttributesConstRef<FVector3f> VertexPositions = MeshDescriptionAttributes.GetVertexPositions();
-		TVertexInstanceAttributesConstRef<FVector3f> VertexInstanceNormals = MeshDescriptionAttributes.GetVertexInstanceNormals();
-		TVertexInstanceAttributesConstRef<FVector3f> VertexInstanceTangents = MeshDescriptionAttributes.GetVertexInstanceTangents();
-		TVertexInstanceAttributesConstRef<float> VertexInstanceBinormalSigns = MeshDescriptionAttributes.GetVertexInstanceBinormalSigns();
-		TVertexInstanceAttributesConstRef<FVector4f> VertexInstanceColors = MeshDescriptionAttributes.GetVertexInstanceColors();
-		TVertexInstanceAttributesConstRef<FVector2f> VertexInstanceUVs = MeshDescriptionAttributes.GetVertexInstanceUVs();
-
-		TVertexInstanceAttributesConstRef<FVector3f>
-			VertexInstanceVelocities = MeshDescription.VertexInstanceAttributes().GetAttributesRef<FVector3f>(MeshAttribute::VertexInstance::Velocity
-			);
-
-		const bool bHasVelocities = VertexInstanceVelocities.IsValid();
-		OutMeshData.VertexInfo.bHasMotionVectors = bHasVelocities;
-
-		const int32 NumVertices = MeshDescription.Vertices().Num();
-		const int32 NumTriangles = MeshDescription.Triangles().Num();
-		const int32 NumMeshDataVertices = NumTriangles * 3;
-
-		OutMeshData.Positions.Reserve(NumVertices);
-		OutMeshData.Indices.Reserve(NumMeshDataVertices);
-		OutMeshData.TangentsZ.Reserve(NumMeshDataVertices);
-		OutMeshData.Colors.Reserve(NumMeshDataVertices);
-		OutMeshData.TextureCoordinates.Reserve(NumMeshDataVertices);
-		if (bHasVelocities)
-		{
-			OutMeshData.MotionVectors.Reserve(NumMeshDataVertices);
-		}
-
-		FBox BoundingBox(EForceInit::ForceInitToZero);
-		int32 VertexIndex = 0;
-		for (FPolygonGroupID PolygonGroupID : MeshDescription.PolygonGroups().GetElementIDs())
-		{
-			// Skip empty polygon groups
-			if (MeshDescription.GetNumPolygonGroupPolygons(PolygonGroupID) == 0)
-			{
-				continue;
-			}
-
-			FGeometryCacheMeshBatchInfo BatchInfo;
-			BatchInfo.StartIndex = OutMeshData.Indices.Num();
-			BatchInfo.MaterialIndex = MaterialIndex++;
-
-			int32 TriangleCount = 0;
-			for (FPolygonID PolygonID : MeshDescription.GetPolygonGroupPolygonIDs(PolygonGroupID))
-			{
-				for (FTriangleID TriangleID : MeshDescription.GetPolygonTriangles(PolygonID))
-				{
-					for (FVertexInstanceID VertexInstanceID : MeshDescription.GetTriangleVertexInstances(TriangleID))
-					{
-						const FVector3f& Position = VertexPositions[MeshDescription.GetVertexInstanceVertex(VertexInstanceID)];
-						OutMeshData.Positions.Add(Position);
-						BoundingBox += (FVector)Position;
-
-						OutMeshData.Indices.Add(VertexIndex++);
-
-						FPackedNormal Normal = VertexInstanceNormals[VertexInstanceID];
-						Normal.Vector.W = VertexInstanceBinormalSigns[VertexInstanceID] < 0 ? -127 : 127;
-						OutMeshData.TangentsZ.Add(Normal);
-						OutMeshData.TangentsX.Add(VertexInstanceTangents[VertexInstanceID]);
-
-						const bool bSRGB = true;
-						OutMeshData.Colors.Add(FLinearColor(VertexInstanceColors[VertexInstanceID]).ToFColor(bSRGB));
-
-						// Supporting only one UV channel
-						const int32 UVIndex = 0;
-						OutMeshData.TextureCoordinates.Add(VertexInstanceUVs.Get(VertexInstanceID, UVIndex));
-
-						if (bHasVelocities)
-						{
-							FVector3f MotionVector = VertexInstanceVelocities[VertexInstanceID];
-							MotionVector *= (-1.f / FramesPerSecond);	 // Velocity is per seconds but we need per frame for motion vectors
-
-							OutMeshData.MotionVectors.Add(MotionVector);
-						}
-					}
-
-					++TriangleCount;
-				}
-			}
-
-			OutMeshData.BoundingBox = (FBox3f)BoundingBox;
-
-			BatchInfo.NumTriangles = TriangleCount;
-			OutMeshData.BatchesInfo.Add(BatchInfo);
-		}
-	}
 }	 // namespace UsdGeometryCacheTranslatorImpl
 
 class FGeometryCacheCreateAssetsTaskChain : public FBuildStaticMeshTaskChain
@@ -1194,7 +1028,7 @@ void FUsdGeometryCacheTranslator::UpdateComponents(USceneComponent* SceneCompone
 			{
 				int32 StartFrame = FMath::FloorToInt(Context->Stage.GetStartTimeCode());
 				int32 EndFrame = FMath::CeilToInt(Context->Stage.GetEndTimeCode());
-				UsdGeometryCacheTranslatorImpl::GetGeometryCacheDataTimeCodeRange(Context->Stage, PrimPath.GetString(), StartFrame, EndFrame);
+				UsdUtils::GetAnimatedMeshTimeCodes(Context->Stage, PrimPath.GetString(), StartFrame, EndFrame);
 
 				TimeCode = static_cast<float>(StartFrame);
 			}
diff --git a/Engine/Plugins/Runtime/GeometryCache/Source/GeometryCache/GeometryCache.Build.cs b/Engine/Plugins/Runtime/GeometryCache/Source/GeometryCache/GeometryCache.Build.cs
--- a/Engine/Plugins/Runtime/GeometryCache/Source/GeometryCache/GeometryCache.Build.cs
+++ b/Engine/Plugins/Runtime/GeometryCache/Source/GeometryCache/GeometryCache.Build.cs
@@ -20,9 +20,11 @@ public class GeometryCache : ModuleRules
                 "NiagaraCore"
 			}
 		);
+		PrivateDependencyModuleNames.Add("MeshDescription");
+		PrivateDependencyModuleNames.Add("StaticMeshDescription");
 		PrivateDependencyModuleNames.Add("TraceLog");
 
-        PublicIncludePathModuleNames.Add("TargetPlatform");
+		PublicIncludePathModuleNames.Add("TargetPlatform");
 		PrivateIncludePathModuleNames.Add("Shaders");
 
 		if (Target.bBuildEditor)
diff --git a/Engine/Plugins/Runtime/GeometryCache/Source/GeometryCache/Private/GeometryCacheHelpers.cpp b/Engine/Plugins/Runtime/GeometryCache/Source/GeometryCache/Private/GeometryCacheHelpers.cpp
new file mode 100644
--- /dev/null
+++ b/Engine/Plugins/Runtime/GeometryCache/Source/GeometryCache/Private/GeometryCacheHelpers.cpp
@@ -0,0 +1,125 @@
+// Copyright Epic Games, Inc. All Rights Reserved.
+
+#include "GeometryCacheHelpers.h"
+#include "GeometryCacheMeshData.h"
+#include "MeshDescription.h"
+#include "StaticMeshAttributes.h"
+
+const FName MeshAttribute::VertexInstance::Velocity("Velocity");
+
+namespace UE::GeometryCache::Utils
+{
+	void GetGeometryCacheMeshDataFromMeshDescription(
+		FGeometryCacheMeshData& OutMeshData,
+		FMeshDescription& MeshDescription,
+		int32 MaterialIndex,
+		float FramesPerSecond
+	)
+	{
+		TRACE_CPUPROFILER_EVENT_SCOPE(GetGeometryCacheMeshDataFromMeshDescription);
+
+		OutMeshData.Positions.Reset();
+		OutMeshData.TextureCoordinates.Reset();
+		OutMeshData.TangentsX.Reset();
+		OutMeshData.TangentsZ.Reset();
+		OutMeshData.Colors.Reset();
+		OutMeshData.Indices.Reset();
+
+		OutMeshData.MotionVectors.Reset();
+		OutMeshData.BatchesInfo.Reset();
+		OutMeshData.BoundingBox.Init();
+
+		OutMeshData.VertexInfo.bHasColor0 = true;
+		OutMeshData.VertexInfo.bHasTangentX = true;
+		OutMeshData.VertexInfo.bHasTangentZ = true;
+		OutMeshData.VertexInfo.bHasUV0 = true;
+
+		FStaticMeshAttributes MeshDescriptionAttributes(MeshDescription);
+
+		TVertexAttributesConstRef<FVector3f> VertexPositions = MeshDescriptionAttributes.GetVertexPositions();
+		TVertexInstanceAttributesConstRef<FVector3f> VertexInstanceNormals = MeshDescriptionAttributes.GetVertexInstanceNormals();
+		TVertexInstanceAttributesConstRef<FVector3f> VertexInstanceTangents = MeshDescriptionAttributes.GetVertexInstanceTangents();
+		TVertexInstanceAttributesConstRef<float> VertexInstanceBinormalSigns = MeshDescriptionAttributes.GetVertexInstanceBinormalSigns();
+		TVertexInstanceAttributesConstRef<FVector4f> VertexInstanceColors = MeshDescriptionAttributes.GetVertexInstanceColors();
+		TVertexInstanceAttributesConstRef<FVector2f> VertexInstanceUVs = MeshDescriptionAttributes.GetVertexInstanceUVs();
+
+		TVertexInstanceAttributesConstRef<FVector3f>
+			VertexInstanceVelocities = MeshDescription.VertexInstanceAttributes().GetAttributesRef<FVector3f>(MeshAttribute::VertexInstance::Velocity
+			);
+
+		const bool bHasVelocities = VertexInstanceVelocities.IsValid();
+		OutMeshData.VertexInfo.bHasMotionVectors = bHasVelocities;
+
+		const int32 NumVertices = MeshDescription.Vertices().Num();
+		const int32 NumTriangles = MeshDescription.Triangles().Num();
+		const int32 NumMeshDataVertices = NumTriangles * 3;
+
+		OutMeshData.Positions.Reserve(NumVertices);
+		OutMeshData.Indices.Reserve(NumMeshDataVertices);
+		OutMeshData.TangentsZ.Reserve(NumMeshDataVertices);
+		OutMeshData.Colors.Reserve(NumMeshDataVertices);
+		OutMeshData.TextureCoordinates.Reserve(NumMeshDataVertices);
+		if (bHasVelocities)
+		{
+			OutMeshData.MotionVectors.Reserve(NumMeshDataVertices);
+		}
+
+		FBox BoundingBox(EForceInit::ForceInitToZero);
+		int32 VertexIndex = 0;
+		for (FPolygonGroupID PolygonGroupID : MeshDescription.PolygonGroups().GetElementIDs())
+		{
+			// Skip empty polygon groups
+			if (MeshDescription.GetNumPolygonGroupPolygons(PolygonGroupID) == 0)
+			{
+				continue;
+			}
+
+			FGeometryCacheMeshBatchInfo BatchInfo;
+			BatchInfo.StartIndex = OutMeshData.Indices.Num();
+			BatchInfo.MaterialIndex = MaterialIndex++;
+
+			int32 TriangleCount = 0;
+			for (FPolygonID PolygonID : MeshDescription.GetPolygonGroupPolygonIDs(PolygonGroupID))
+			{
+				for (FTriangleID TriangleID : MeshDescription.GetPolygonTriangles(PolygonID))
+				{
+					for (FVertexInstanceID VertexInstanceID : MeshDescription.GetTriangleVertexInstances(TriangleID))
+					{
+						const FVector3f& Position = VertexPositions[MeshDescription.GetVertexInstanceVertex(VertexInstanceID)];
+						OutMeshData.Positions.Add(Position);
+						BoundingBox += (FVector)Position;
+
+						OutMeshData.Indices.Add(VertexIndex++);
+
+						FPackedNormal Normal = VertexInstanceNormals[VertexInstanceID];
+						Normal.Vector.W = VertexInstanceBinormalSigns[VertexInstanceID] < 0 ? -127 : 127;
+						OutMeshData.TangentsZ.Add(Normal);
+						OutMeshData.TangentsX.Add(VertexInstanceTangents[VertexInstanceID]);
+
+						const bool bSRGB = true;
+						OutMeshData.Colors.Add(FLinearColor(VertexInstanceColors[VertexInstanceID]).ToFColor(bSRGB));
+
+						// Supporting only one UV channel
+						const int32 UVIndex = 0;
+						OutMeshData.TextureCoordinates.Add(VertexInstanceUVs.Get(VertexInstanceID, UVIndex));
+
+						if (bHasVelocities)
+						{
+							FVector3f MotionVector = VertexInstanceVelocities[VertexInstanceID];
+							MotionVector *= (-1.f / FramesPerSecond);	 // Velocity is per seconds but we need per frame for motion vectors
+
+							OutMeshData.MotionVectors.Add(MotionVector);
+						}
+					}
+
+					++TriangleCount;
+				}
+			}
+
+			OutMeshData.BoundingBox = (FBox3f)BoundingBox;
+
+			BatchInfo.NumTriangles = TriangleCount;
+			OutMeshData.BatchesInfo.Add(BatchInfo);
+		}
+	}
+}
\ No newline at end of file
diff --git a/Engine/Plugins/Runtime/GeometryCache/Source/GeometryCache/Public/GeometryCacheHelpers.h b/Engine/Plugins/Runtime/GeometryCache/Source/GeometryCache/Public/GeometryCacheHelpers.h
--- a/Engine/Plugins/Runtime/GeometryCache/Source/GeometryCache/Public/GeometryCacheHelpers.h
+++ b/Engine/Plugins/Runtime/GeometryCache/Source/GeometryCache/Public/GeometryCacheHelpers.h
@@ -14,3 +14,16 @@ struct GeometyCacheHelpers
 		return Time - Duration * FMath::FloorToFloat(Time / Duration);
 	}
 };
+
+namespace MeshAttribute::VertexInstance
+{
+	extern GEOMETRYCACHE_API const FName Velocity;
+}
+
+struct FGeometryCacheMeshData;
+struct FMeshDescription;
+
+namespace UE::GeometryCache::Utils
+{
+	GEOMETRYCACHE_API void GetGeometryCacheMeshDataFromMeshDescription(FGeometryCacheMeshData& OutMeshData, FMeshDescription& MeshDescription, int32 MaterialIndex, float FramesPerSecond);
+}
\ No newline at end of file
diff --git a/Engine/Plugins/Runtime/USDCore/Source/USDUtilities/Private/USDGeomMeshConversion.cpp b/Engine/Plugins/Runtime/USDCore/Source/USDUtilities/Private/USDGeomMeshConversion.cpp
--- a/Engine/Plugins/Runtime/USDCore/Source/USDUtilities/Private/USDGeomMeshConversion.cpp
+++ b/Engine/Plugins/Runtime/USDCore/Source/USDUtilities/Private/USDGeomMeshConversion.cpp
@@ -29,6 +29,7 @@
 #include "Engine/SkeletalMesh.h"
 #include "Engine/StaticMesh.h"
 #include "GeometryCache.h"
+#include "GeometryCacheHelpers.h"
 #include "GeometryCacheMeshData.h"
 #include "GeometryCacheTrack.h"
 #include "Materials/MaterialInstanceConstant.h"
@@ -131,8 +132,6 @@ static FAutoConsoleVariableRef CVarSkipMeshTangentComputation(
 	TEXT("Skip computing tangents for meshes. With meshes with a huge numer of vertices, it can take a very long time to compute them.")
 );
 
-const FName MeshAttribute::VertexInstance::Velocity("Velocity");
-
 namespace UE::UsdGeomMeshConversion::Private
 {
 	static const FString DisplayColorID = TEXT("!DisplayColor");
@@ -6448,6 +6447,16 @@ UsdUtils::EMeshTopologyVariance UsdUtils::GetMeshTopologyVariance(const pxr::Usd
 	return EMeshTopologyVariance::Constant;
 }
 
+UsdUtils::EMeshTopologyVariance UsdUtils::GetMeshTopologyVariance(const UE::FUsdPrim& UsdPrim)
+{
+	pxr::UsdGeomMesh MeshPrim(UsdPrim);
+	if (!MeshPrim)
+	{
+		return EMeshTopologyVariance::Constant;
+	}
+	return GetMeshTopologyVariance(MeshPrim);
+}
+
 uint64 UsdUtils::GetGprimVertexCount(const pxr::UsdGeomGprim& Gprim, double TimeCode)
 {
 	TRACE_CPUPROFILER_EVENT_SCOPE(UsdUtils::GetGprimVertexCount);
@@ -6570,6 +6579,51 @@ void UsdUtils::AuthorIdentityTransformGprimAttributes(const pxr::UsdPrim& UsdPri
 	}
 }
 
+void UsdUtils::GetAnimatedMeshTimeCodes(const UE::FUsdStage& Stage, const FString& PrimPath, int32& OutStartFrame, int32& OutEndFrame)
+{
+	if (!Stage || PrimPath.IsEmpty())
+	{
+		return;
+	}
+
+	FScopedUsdAllocs Allocs;
+
+	pxr::UsdPrim UsdPrim = pxr::UsdPrim(Stage.GetPrimAtPath(UE::FSdfPath(*PrimPath)));
+	if (!UsdPrim)
+	{
+		return;
+	}
+
+	constexpr bool bIncludeInherited = false;
+	pxr::TfTokenVector GeomMeshAttributeNames = pxr::UsdGeomMesh::GetSchemaAttributeNames(bIncludeInherited);
+	pxr::TfTokenVector GeomPointBasedAttributeNames = pxr::UsdGeomPointBased::GetSchemaAttributeNames(bIncludeInherited);
+
+	GeomMeshAttributeNames.reserve(GeomMeshAttributeNames.size() + GeomPointBasedAttributeNames.size());
+	GeomMeshAttributeNames.insert(GeomMeshAttributeNames.end(), GeomPointBasedAttributeNames.begin(), GeomPointBasedAttributeNames.end());
+
+	TOptional<double> MinStartTimeCode;
+	TOptional<double> MaxEndTimeCode;
+
+	for (const pxr::TfToken& AttributeName : GeomMeshAttributeNames)
+	{
+		if (const pxr::UsdAttribute& Attribute = UsdPrim.GetAttribute(AttributeName))
+		{
+			std::vector<double> TimeSamples;
+			if (Attribute.GetTimeSamples(&TimeSamples) && TimeSamples.size() > 0)
+			{
+				MinStartTimeCode = FMath::Min(MinStartTimeCode.Get(TNumericLimits<double>::Max()), TimeSamples[0]);
+				MaxEndTimeCode = FMath::Max(MaxEndTimeCode.Get(TNumericLimits<double>::Lowest()), TimeSamples[TimeSamples.size() - 1]);
+			}
+		}
+	}
+
+	if (MinStartTimeCode.IsSet() && MaxEndTimeCode.IsSet())
+	{
+		OutStartFrame = FMath::FloorToInt(MinStartTimeCode.GetValue());
+		OutEndFrame = FMath::CeilToInt(MaxEndTimeCode.GetValue());
+	}
+}
+
 #undef LOCTEXT_NAMESPACE
 
 #endif	  // #if USE_USD_SDK
diff --git a/Engine/Plugins/Runtime/USDCore/Source/USDUtilities/Public/USDGeomMeshConversion.h b/Engine/Plugins/Runtime/USDCore/Source/USDUtilities/Public/USDGeomMeshConversion.h
--- a/Engine/Plugins/Runtime/USDCore/Source/USDUtilities/Public/USDGeomMeshConversion.h
+++ b/Engine/Plugins/Runtime/USDCore/Source/USDUtilities/Public/USDGeomMeshConversion.h
@@ -49,11 +49,6 @@ namespace UE
 	class FUsdPrim;
 }
 
-namespace MeshAttribute::VertexInstance
-{
-	extern USDUTILITIES_API const FName Velocity;
-}
-
 #if USE_USD_SDK
 namespace UsdToUnreal
 {
@@ -502,6 +497,8 @@ namespace UsdUtils
 	/** Returns true if Prim is a GeomMesh with animated attributes */
 	USDUTILITIES_API bool IsAnimatedMesh(const pxr::UsdPrim& Prim);
 
+	USDUTILITIES_API void GetAnimatedMeshTimeCodes(const UE::FUsdStage& Stage, const FString& PrimPath, int32& OutStartFrame, int32& OutEndFrame);
+
 	/** Mesh topology variance as defined by Alembic's MeshTopologyVariance */
 	enum class EMeshTopologyVariance : uint8
 	{
@@ -512,6 +509,7 @@ namespace UsdUtils
 
 	/** Returns the type of topology variance UsdMesh has */
 	USDUTILITIES_API EMeshTopologyVariance GetMeshTopologyVariance(const pxr::UsdGeomMesh& UsdMesh);
+	USDUTILITIES_API EMeshTopologyVariance GetMeshTopologyVariance(const UE::FUsdPrim& UsdPrim);
 
 	USDUTILITIES_API uint64 GetGprimVertexCount(const pxr::UsdGeomGprim& Gprim, double TimeCode);
 
