USD: Silence warning when failing to generate a skeletal mesh.

diff --git a/Engine/Plugins/Importers/USDImporter/Source/USDSchemas/Private/USDSkelSkeletonTranslator.cpp b/Engine/Plugins/Importers/USDImporter/Source/USDSchemas/Private/USDSkelSkeletonTranslator.cpp
--- a/Engine/Plugins/Importers/USDImporter/Source/USDSchemas/Private/USDSkelSkeletonTranslator.cpp
+++ b/Engine/Plugins/Importers/USDImporter/Source/USDSchemas/Private/USDSkelSkeletonTranslator.cpp
@@ -524,6 +524,7 @@ namespace UsdSkelSkeletonTranslatorImpl
 			}
 		}
 
+		bool bConvertedMeshData = false;
 		TFunction<bool(const pxr::UsdGeomMesh&, int32)> ConvertLOD = [&LODIndexToSkeletalMeshImportDataMap,
 																	  &LODIndexToMaterialInfoMap,
 																	  &LODMetadata,
@@ -534,8 +535,11 @@ namespace UsdSkelSkeletonTranslatorImpl
 																	  OutBlendShapes,
 																	  &StageInfo,
 																	  Options,
-																	  &MetadataOptions](const pxr::UsdGeomMesh& LODMesh, int32 LODIndex)
+																	  &MetadataOptions,
+																	  &bConvertedMeshData](const pxr::UsdGeomMesh& LODMesh, int32 LODIndex)
 		{
+			pxr::UsdPrim LODMeshPrim = LODMesh.GetPrim();
+
 			// Construct this and SkinningQuery every time so as to survive the prim reference invalidation caused by flipping LODs
 			pxr::UsdSkelSkeletonQuery SkeletonQuery = InSkelCache.GetSkelQuery(pxr::UsdSkelSkeleton{Stage->GetPrimAtPath(SkeletonPrimPath)});
 			if (!SkeletonQuery)
@@ -543,21 +547,21 @@ namespace UsdSkelSkeletonTranslatorImpl
 				return true;	// Continue trying other LODs
 			}
 
-			pxr::UsdSkelSkinningQuery SkinningQuery = UsdUtils::CreateSkinningQuery(LODMesh.GetPrim(), SkeletonQuery);
+			pxr::UsdSkelSkinningQuery SkinningQuery = UsdUtils::CreateSkinningQuery(LODMeshPrim, SkeletonQuery);
 			if (!SkinningQuery)
 			{
 				return true;
 			}
 
 			// Ignore prims from disabled purposes
-			if (!EnumHasAllFlags(Options.PurposesToLoad, IUsdPrim::GetPurpose(LODMesh.GetPrim())))
+			if (!EnumHasAllFlags(Options.PurposesToLoad, IUsdPrim::GetPurpose(LODMeshPrim)))
 			{
 				return true;
 			}
 
 			// If a skinned prim has an alt draw mode, let's prioritize showing the alt draw mode (that is, geommesh/geometrycache translators
 			// will handle it) and skip skinning it
-			EUsdDrawMode DrawMode = UsdUtils::GetAppliedDrawMode(LODMesh.GetPrim());
+			EUsdDrawMode DrawMode = UsdUtils::GetAppliedDrawMode(LODMeshPrim);
 			if (DrawMode != EUsdDrawMode::Default)
 			{
 				return true;
@@ -578,14 +582,16 @@ namespace UsdSkelSkeletonTranslatorImpl
 			bool bSuccess = UsdToUnreal::ConvertSkinnedMesh(SkinningQuery, SkeletonQuery, LODImportData, LODMaterialInfo, Options);
 			if (!bSuccess)
 			{
+				UE_LOG(LogUsd, Warning, TEXT("Failed to convert skinned mesh '%s'"), *UsdToUnreal::ConvertPath(LODMeshPrim.GetPrimPath()));
 				return true;
 			}
+			bConvertedMeshData = true;
 
 			if (MetadataOptions.bCollectMetadata && MetadataOptions.bCollectFromEntireSubtrees)
 			{
 				// Collect metadata from this particular LOD mesh prim
 				UsdToUnreal::ConvertMetadata(
-					LODMesh.GetPrim(),
+					LODMeshPrim,
 					LODMetadata,
 					MetadataOptions.BlockedPrefixFilters,
 					MetadataOptions.bInvertFilters,
@@ -595,7 +601,7 @@ namespace UsdSkelSkeletonTranslatorImpl
 
 			if (OutBlendShapes)
 			{
-				pxr::SdfPath MeshPrimPath = LODMesh.GetPrim().GetPrimPath();
+				pxr::SdfPath MeshPrimPath = LODMeshPrim.GetPrimPath();
 
 				pxr::UsdSkelBindingAPI SkelBindingAPI{Stage->GetPrimAtPath(MeshPrimPath)};
 				if (pxr::UsdSkelBlendShapeQuery BlendShapeQuery{SkelBindingAPI})
@@ -697,7 +703,7 @@ namespace UsdSkelSkeletonTranslatorImpl
 			}
 		}
 
-		return true;
+		return bConvertedMeshData;
 	}
 
 	/** Warning: This function will temporarily switch the active LOD variant if one exists, so it's *not* thread safe! */
