Geometry Cache: Improved the progress bar when importing (and re-importing) geometry cache through Alembic and USD Importer (now cancellable).

diff --git a/Engine/Plugins/Importers/AlembicImporter/Source/AlembicLibrary/Private/AbcImporter.cpp b/Engine/Plugins/Importers/AlembicImporter/Source/AlembicLibrary/Private/AbcImporter.cpp
--- a/Engine/Plugins/Importers/AlembicImporter/Source/AlembicLibrary/Private/AbcImporter.cpp
+++ b/Engine/Plugins/Importers/AlembicImporter/Source/AlembicLibrary/Private/AbcImporter.cpp
@@ -17,6 +17,7 @@
 #include "Engine/SkeletalMesh.h"
 #include "Engine/SkinnedAssetCommon.h"
 #include "Engine/StaticMesh.h"
+#include "FeedbackContextEditor.h"
 #include "GeometryCache.h"
 #include "GeometryCacheCodecV1.h"
 #include "GeometryCacheComponent.h"
@@ -400,7 +401,14 @@ UGeometryCache* FAbcImporter::ImportAsGeometryCache(UObject* InParent, EObjectFl
 				FAbcImportLogger::AddImportMessage(Message);
 			}
 
-			FScopedSlowTask SlowTask(static_cast<float>((ImportSettings->SamplingSettings.FrameEnd + 1) - ImportSettings->SamplingSettings.FrameStart), FText::FromString(FString(TEXT("Importing Frames"))));
+			// Use new feedback context to override the one coming from the ReimportManager to be able to show the ScopedSlowTask below
+			FFeedbackContextEditor GeoCacheContext;
+			const FString AssetName = GeometryCache->GetName();
+			const int32 NumFrames = (ImportSettings->SamplingSettings.FrameEnd + 1) - ImportSettings->SamplingSettings.FrameStart;
+			int32 TotalFrames = 0;
+
+			FString Title = FString::Format(TEXT("Importing frames for {0} ({1}/{2})"), { AssetName, 0,  NumFrames });
+			FScopedSlowTask SlowTask(static_cast<float>(NumFrames), FText::FromString(Title), true, GeoCacheContext);
 			SlowTask.MakeDialog(true);
 			float CompletedFrames = 0.0f;
 
@@ -417,11 +425,11 @@ UGeometryCache* FAbcImporter::ImportAsGeometryCache(UObject* InParent, EObjectFl
 				
 				const TArray<FAbcPolyMesh*>& PolyMeshes = AbcFile->GetPolyMeshes();
 				TArray<float> FrameTimes;
-				FrameTimes.SetNum(ImportSettings->SamplingSettings.FrameEnd - ImportSettings->SamplingSettings.FrameStart + 1);
+				FrameTimes.SetNum(NumFrames);
 				
 				const int32 NumTracks = Tracks.Num();
 				int32 PreviousNumVertices = 0;
-				TFunction<void(int32, FAbcFile*)> Callback = [this, &Tracks, &PolyMeshes, &PreviousNumVertices, &FrameTimes, &SlowTask, &CompletedFrames](int32 FrameIndex, const FAbcFile* InAbcFile)
+				TFunction<void(int32, FAbcFile*)> Callback = [this, &Tracks, &PolyMeshes, &PreviousNumVertices, &FrameTimes, &SlowTask, &CompletedFrames, &TotalFrames, &AssetName, NumFrames](int32 FrameIndex, const FAbcFile* InAbcFile)
 				{
 					const bool bUseVelocitiesAsMotionVectors = (ImportSettings->GeometryCacheSettings.MotionVectors == EAbcGeometryCacheMotionVectorsImport::ImportAbcVelocitiesAsMotionVectors);
 					FGeometryCacheMeshData MeshData;
@@ -444,7 +452,9 @@ UGeometryCache* FAbcImporter::ImportAsGeometryCache(UObject* InParent, EObjectFl
 					++CompletedFrames;
 					if (IsInGameThread())
 					{
-						SlowTask.EnterProgressFrame(CompletedFrames);
+						TotalFrames += int32(CompletedFrames);
+						FString Title = FString::Format(TEXT("Importing frames for {0} ({1}/{2})"), { AssetName, TotalFrames,  NumFrames });
+						SlowTask.EnterProgressFrame(CompletedFrames, FText::FromString(Title));
 						CompletedFrames = 0.0f;
 					}					
 				};
@@ -499,7 +509,7 @@ UGeometryCache* FAbcImporter::ImportAsGeometryCache(UObject* InParent, EObjectFl
 				}
 
 				const int32 NumTracks = Tracks.Num();
-				TFunction<void(int32, FAbcFile*)> Callback = [this, NumTracks, &ImportPolyMeshes, &Tracks, &MaterialOffsets, &SlowTask, &CompletedFrames](int32 FrameIndex, const FAbcFile* InAbcFile)
+				TFunction<void(int32, FAbcFile*)> Callback = [this, NumTracks, &ImportPolyMeshes, &Tracks, &MaterialOffsets, &SlowTask, &CompletedFrames, &TotalFrames, &AssetName, NumFrames](int32 FrameIndex, const FAbcFile* InAbcFile)
 				{
 					const float FrameRate = static_cast<float>(InAbcFile->GetFramerate());
 					for (int32 TrackIndex = 0; TrackIndex < NumTracks; ++TrackIndex)
@@ -530,7 +540,9 @@ UGeometryCache* FAbcImporter::ImportAsGeometryCache(UObject* InParent, EObjectFl
 					++CompletedFrames;
 					if (IsInGameThread())
 					{
-						SlowTask.EnterProgressFrame(CompletedFrames);
+						TotalFrames += int32(CompletedFrames);
+						FString Title = FString::Format(TEXT("Importing frames for {0} ({1}/{2})"), { AssetName, TotalFrames,  NumFrames });
+						SlowTask.EnterProgressFrame(CompletedFrames, FText::FromString(Title));
 						CompletedFrames = 0.0f;
 					}					
 				};
diff --git a/Engine/Plugins/Importers/USDImporter/Source/USDSchemas/Private/USDGeometryCacheTranslator.cpp b/Engine/Plugins/Importers/USDImporter/Source/USDSchemas/Private/USDGeometryCacheTranslator.cpp
--- a/Engine/Plugins/Importers/USDImporter/Source/USDSchemas/Private/USDGeometryCacheTranslator.cpp
+++ b/Engine/Plugins/Importers/USDImporter/Source/USDSchemas/Private/USDGeometryCacheTranslator.cpp
@@ -25,6 +25,7 @@
 #include "UsdWrappers/UsdStage.h"
 
 #include "Async/ParallelFor.h"
+#include "FeedbackContextEditor.h"
 #include "GeometryCache.h"
 #include "GeometryCacheCodecV1.h"
 #include "GeometryCacheComponent.h"
@@ -37,6 +38,7 @@
 #include "HAL/IConsoleManager.h"
 #include "MaterialDomain.h"
 #include "Materials/Material.h"
+#include "Misc/ScopedSlowTask.h"
 #include "StaticMeshAttributes.h"
 #include "StaticMeshOperations.h"
 #include "UObject/Package.h"
@@ -513,6 +515,17 @@ namespace UsdGeometryCacheTranslatorImpl
 		FReadMeshDataArgs Args = GetReadMeshDataArgs(Context, RootPrimPath);
 		Args.bPropagateTransform = true;
 
+		// Use new feedback context to override the one coming from the USDImporter to be able to show the ScopedSlowTask below
+		FFeedbackContextEditor GeoCacheContext;
+		const FString AssetName = GeometryCache->GetName();
+		const int32 NumMeshes = MeshPrims.Num();
+		const int32 NumFrames = (Args.EndFrame - Args.StartFrame) * NumMeshes;
+		int32 TotalFrames = 0;
+
+		FString Title = FString::Format(TEXT("Importing frames for {0} ({1}/{2})"), { AssetName, 0,  NumFrames });
+		FScopedSlowTask SlowTask(static_cast<float>(NumFrames), FText::FromString(Title), true, GeoCacheContext);
+		SlowTask.MakeDialog(true);
+
 		const bool bSingleThreaded = !FApp::ShouldUseThreadingForPerformance();
 		if (bSingleThreaded)
 		{
@@ -525,19 +538,27 @@ namespace UsdGeometryCacheTranslatorImpl
 
 				for (int32 FrameIndex = Args.StartFrame; FrameIndex < Args.EndFrame; ++FrameIndex)
 				{
+					if (SlowTask.ShouldCancel())
+					{
+						return;
+					}
+
 					// Read frame data
 					FGeometryCacheMeshData MeshData;
 					ReadMeshData(Args, MeshPrim, MaterialOffsets[Index], FrameIndex, MeshData);
 
 					// Add it to the track
 					StreamableTrack->AddMeshSample(MeshData, (FrameIndex - Args.StartFrame) / Args.FramesPerSecond, bConstantTopology);
+
+					++TotalFrames;
+					Title = FString::Format(TEXT("Importing frames for {0} ({1}/{2})"), { AssetName, TotalFrames,  NumFrames });
+					SlowTask.EnterProgressFrame(1, FText::FromString(Title));
 				}
 			}
 		}
 		else
 		{
 			// Balance the number of threads for mesh reads vs frame reads
-			const int32 NumMeshes = MeshPrims.Num();
 			const int32 MaxWorkerThreads = FTaskGraphInterface::Get().GetNumWorkerThreads();
 			int32 NumMeshThreads = FMath::Clamp(MaxWorkerThreads, 1, FMath::Min(NumMeshes, GUsdGeometryCacheParallelMeshReads));
 			int32 NumFrameThreads = FMath::Clamp(MaxWorkerThreads, 1, GUsdGeometryCacheParallelFrameReads);
@@ -563,14 +584,23 @@ namespace UsdGeometryCacheTranslatorImpl
 				SyncEvents[Index] = FPlatformProcess::GetSynchEventFromPool();
 			}
 
+			std::atomic<float> CompletedFrames = 0.0f;
+
 			// Parallel mesh reads: Meshes can be read independently of each other
+			bool bCancelled = false;
 			ParallelFor(
 				NumMeshThreads,
-				[&Args, &MeshPrims, &MaterialOffsets, &SyncEvents, NumMeshThreads, NumFrameThreads, GeometryCache, NumMeshes](int32 MeshThreadIndex)
+				[&Args, &MeshPrims, &MaterialOffsets, &SyncEvents, NumMeshThreads, NumFrameThreads, GeometryCache, NumMeshes,
+				&SlowTask,
+				&AssetName,
+				&CompletedFrames,
+				&NumFrames,
+				&TotalFrames,
+				&bCancelled](int32 MeshThreadIndex)
 				{
 					int32 MeshIndex = MeshThreadIndex;
 
-					while (MeshIndex < NumMeshes)
+					while (!bCancelled && (MeshIndex < NumMeshes))
 					{
 						UE::FUsdPrim MeshPrim = Args.Stage.GetPrimAtPath(MeshPrims[MeshIndex]);
 						int32 MaterialOffset = MaterialOffsets[MeshIndex];
@@ -593,12 +623,24 @@ namespace UsdGeometryCacheTranslatorImpl
 							 NumFrameThreads,
 							 MaterialOffset,
 							 StreamableTrack,
-							 bConstantTopology](int32 FrameThreadIndex)
+							 bConstantTopology,
+							 &SlowTask,
+							 &AssetName,
+							 &CompletedFrames,
+							 &NumFrames,
+							 &TotalFrames,
+							 &bCancelled](int32 FrameThreadIndex)
 							{
 								int32 FrameIndex = Args.StartFrame + FrameThreadIndex;
 
-								while (FrameIndex < Args.EndFrame)
+								while (!bCancelled && (FrameIndex < Args.EndFrame))
 								{
+									if (IsInGameThread() && SlowTask.ShouldCancel())
+									{
+										bCancelled = true;
+										break;
+									}
+
 									// Read frame data into memory
 									FGeometryCacheMeshData MeshData;
 									ReadMeshData(Args, MeshPrim, MaterialOffset, FrameIndex, MeshData);
@@ -606,6 +648,10 @@ namespace UsdGeometryCacheTranslatorImpl
 									// Wait until it's our turn to process this frame.
 									while (WriteFrameIndex < FrameIndex)
 									{
+										if (bCancelled)
+										{
+											break;
+										}
 										const uint32 WaitTimeInMs = 10;
 										FrameWrittenEvent->Wait(WaitTimeInMs);
 									}
@@ -620,6 +666,17 @@ namespace UsdGeometryCacheTranslatorImpl
 										// Mark the next frame index as ready for processing.
 										++WriteFrameIndex;
 
+										CompletedFrames = CompletedFrames + 1;
+
+										if (IsInGameThread())
+										{
+											float TempCompletedFrames = CompletedFrames;
+											CompletedFrames = 0.0f;
+											TotalFrames += int32(TempCompletedFrames);
+											FString Title = FString::Format(TEXT("Importing frames for {0} ({1}/{2})"), { AssetName, TotalFrames,  NumFrames });
+											SlowTask.EnterProgressFrame(TempCompletedFrames, FText::FromString(Title));
+										}
+
 										FrameWrittenEvent->Trigger();
 									}
 
@@ -643,11 +700,15 @@ namespace UsdGeometryCacheTranslatorImpl
 
 	void FinalizeGeometryCache(UGeometryCache* GeometryCache)
 	{
-		for (int32 Index = 0; Index < GeometryCache->Tracks.Num(); ++Index)
+		for (int32 Index = GeometryCache->Tracks.Num() - 1; Index >= 0 ; --Index)
 		{
 			if (UGeometryCacheTrackStreamable* StreamableTrack = Cast<UGeometryCacheTrackStreamable>(GeometryCache->Tracks[Index]))
 			{
-				StreamableTrack->EndCoding();
+				// Remove track that have no data (either import cancelled or there was no mesh data imported)
+				if (!StreamableTrack->EndCoding())
+				{
+					GeometryCache->Tracks.RemoveAt(Index);
+				}
 			}
 		}
 	}
