USD: Add USDImporter option for setting FallbackCollisionType.

diff --git a/Engine/Plugins/Importers/USDImporter/Source/USDSchemas/Private/USDGeomMeshTranslator.cpp b/Engine/Plugins/Importers/USDImporter/Source/USDSchemas/Private/USDGeomMeshTranslator.cpp
--- a/Engine/Plugins/Importers/USDImporter/Source/USDSchemas/Private/USDGeomMeshTranslator.cpp
+++ b/Engine/Plugins/Importers/USDImporter/Source/USDSchemas/Private/USDGeomMeshTranslator.cpp
@@ -1259,7 +1259,6 @@ namespace UE::UsdCollision::Private
 	void CollectCustomCollisionShapes(const pxr::UsdPrim& UsdPrim, FKAggregateGeom& CollisionShapes)
 	{
 		using namespace UsdUtils;
-
 		UsdToUnreal::FUsdMeshConversionOptions Options;
 		Options.PurposesToLoad = EUsdPurpose::Guide;	// custom collision mesh must have guide purpose
 		Options.bMergeIdenticalMaterialSlots = false;
@@ -1332,7 +1331,7 @@ namespace UE::UsdCollision::Private
 		}
 	}
 
-	void SetupSimpleCollision(const pxr::UsdPrim& UsdPrim, UStaticMesh& StaticMesh)
+	void SetupSimpleCollision(const pxr::UsdPrim& UsdPrim, UStaticMesh& StaticMesh, EUsdCollisionType FallbackCollisionType)
 	{
 		using namespace UsdUtils;
 
@@ -1343,7 +1342,8 @@ namespace UE::UsdCollision::Private
 		}
 
 		bool bIsCollisionEnabled = false;
-		if (pxr::UsdPhysicsCollisionAPI CollisionAPI{UsdPrim})
+		pxr::UsdPhysicsCollisionAPI CollisionAPI{UsdPrim};
+		if (CollisionAPI)
 		{
 			if (pxr::UsdAttribute CollisionAttr = CollisionAPI.GetCollisionEnabledAttr())
 			{
@@ -1351,6 +1351,8 @@ namespace UE::UsdCollision::Private
 			}
 		}
 
+		bIsCollisionEnabled = bIsCollisionEnabled || FallbackCollisionType != EUsdCollisionType::None;
+
 		static IConsoleVariable* CVar = IConsoleManager::Get().FindConsoleVariable(TEXT("USD.EnableCollision"));
 		const bool bEnableCollision = CVar && CVar->GetBool();
 
@@ -1371,6 +1373,12 @@ namespace UE::UsdCollision::Private
 			return;
 		}
 
+		EUsdCollisionType Approximation = FallbackCollisionType;
+		if (CollisionAPI)
+		{
+			Approximation = UsdUtils::GetCollisionApproximationType(UsdPrim);
+		}
+
 		BodySetup->Modify();
 		BodySetup->RemoveSimpleCollision();
 		BodySetup->bNeverNeedsCookedCollisionData = false;
@@ -1378,7 +1386,6 @@ namespace UE::UsdCollision::Private
 		BodySetup->DefaultInstance.SetCollisionProfileName(UCollisionProfile::BlockAllDynamic_ProfileName);
 		BodySetup->CollisionTraceFlag = CTF_UseDefault;
 
-		EUsdCollisionType Approximation = GetCollisionApproximationType(UsdPrim);
 		switch (Approximation)
 		{
 			case EUsdCollisionType::None:
@@ -1610,7 +1617,7 @@ void FBuildStaticMeshTaskChain::SetupTasks()
 				   // Setup collision on existing mesh in case the collision settings have changed
 				   // TODO: This is recomputing collision and dirtying the StaticMesh every time, even if the mesh
 				   // is reused from the AssetCache. Maybe we can do that only when collision settings actually change?
-				   UE::UsdCollision::Private::SetupSimpleCollision(GetPrim(), *StaticMesh);
+				   UE::UsdCollision::Private::SetupSimpleCollision(GetPrim(), *StaticMesh, Context->FallbackCollisionType);
 			   }
 		   }
 
@@ -1694,7 +1701,7 @@ void FBuildStaticMeshTaskChain::SetupTasks()
 				return false;
 			}
 
-			UE::UsdCollision::Private::SetupSimpleCollision(GetPrim(), *StaticMesh);
+			UE::UsdCollision::Private::SetupSimpleCollision(GetPrim(), *StaticMesh, Context->FallbackCollisionType);
 
 			return true;
 		}
diff --git a/Engine/Plugins/Importers/USDImporter/Source/USDStage/Private/USDStageActor.cpp b/Engine/Plugins/Importers/USDImporter/Source/USDStage/Private/USDStageActor.cpp
--- a/Engine/Plugins/Importers/USDImporter/Source/USDStage/Private/USDStageActor.cpp
+++ b/Engine/Plugins/Importers/USDImporter/Source/USDStage/Private/USDStageActor.cpp
@@ -165,6 +165,7 @@ struct FUsdStageActorImpl
 		TranslationContext->RenderContext = StageActor->RenderContext;
 		TranslationContext->MaterialPurpose = StageActor->MaterialPurpose;
 		TranslationContext->RootMotionHandling = StageActor->RootMotionHandling;
+		TranslationContext->FallbackCollisionType = StageActor->FallbackCollisionType;
 		TranslationContext->GeometryCacheImport = StageActor->GeometryCacheImport;
 		TranslationContext->SubdivisionLevel = StageActor->SubdivisionLevel;
 		TranslationContext->MetadataOptions = StageActor->MetadataOptions;
@@ -588,6 +589,7 @@ struct FUsdStageActorImpl
 			EventAttributes.Emplace(TEXT("RenderContext"), StageActor->RenderContext.ToString());
 			EventAttributes.Emplace(TEXT("MaterialPurpose"), StageActor->MaterialPurpose.ToString());
 			EventAttributes.Emplace(TEXT("RootMotionHandling"), LexToString((uint8)StageActor->RootMotionHandling));
+			EventAttributes.Emplace(TEXT("FallbackCollisionType"), (uint8)StageActor->FallbackCollisionType);
 			EventAttributes.Emplace(TEXT("GeometryCacheImport"), LexToString((uint8)StageActor->GeometryCacheImport));
 			EventAttributes.Emplace(TEXT("SubdivisionLevel"), LexToString(StageActor->SubdivisionLevel));
 
@@ -929,6 +931,7 @@ AUsdStageActor::AUsdStageActor()
 	, RenderContext(UnrealIdentifiers::UnrealRenderContext)
 	, MaterialPurpose(*UnrealIdentifiers::MaterialPreviewPurpose)
 	, RootMotionHandling(EUsdRootMotionHandling::NoAdditionalRootMotion)
+	, FallbackCollisionType(EUsdCollisionType::ConvexHull)
 	, SubdivisionLevel(0)
 	, MetadataOptions(FUsdMetadataImportOptions{
 		  false, /* bCollectMetadata */
@@ -2789,6 +2792,20 @@ void AUsdStageActor::SetRootMotionHandling(EUsdRootMotionHandling NewHandlingStr
 	LoadUsdStage();
 }
 
+void AUsdStageActor::SetFallbackCollisionType(EUsdCollisionType NewCollisionType)
+{
+	if (NewCollisionType == FallbackCollisionType)
+	{
+		return;
+	}
+
+	const bool bMarkDirty = false;
+	Modify(bMarkDirty);
+
+	FallbackCollisionType = NewCollisionType;
+	LoadUsdStage();
+}
+
 void AUsdStageActor::SetSubdivisionLevel(int32 NewSubdivisionLevel)
 {
 	if (NewSubdivisionLevel == SubdivisionLevel)
@@ -4762,6 +4779,12 @@ void AUsdStageActor::HandlePropertyChangedEvent(FPropertyChangedEvent& PropertyC
 		RootMotionHandling = (EUsdRootMotionHandling) !((uint8)RootMotionHandling);
 		SetRootMotionHandling(CorrectHandling);
 	}
+	else if (PropertyName == GET_MEMBER_NAME_CHECKED(AUsdStageActor, FallbackCollisionType))
+	{
+		const EUsdCollisionType CorrectCollisionType = FallbackCollisionType;
+		FallbackCollisionType = (EUsdCollisionType) !((uint8)FallbackCollisionType);
+		SetFallbackCollisionType(CorrectCollisionType);
+	}
 	else if (PropertyName == GET_MEMBER_NAME_CHECKED(AUsdStageActor, SubdivisionLevel))
 	{
 		int32 CorrectSubdivisionLevel = SubdivisionLevel;
diff --git a/Engine/Plugins/Importers/USDImporter/Source/USDStage/Public/USDStageActor.h b/Engine/Plugins/Importers/USDImporter/Source/USDStage/Public/USDStageActor.h
--- a/Engine/Plugins/Importers/USDImporter/Source/USDStage/Public/USDStageActor.h
+++ b/Engine/Plugins/Importers/USDImporter/Source/USDStage/Public/USDStageActor.h
@@ -25,6 +25,7 @@ class UUsdAssetCache3;
 class UUsdPrimTwin;
 class UUsdTransactor;
 enum class EMovieSceneDataChangeType;
+enum class EUsdCollisionType : uint8;
 struct FUsdSchemaTranslationContext;
 namespace UE
 {
@@ -157,6 +158,10 @@ public:
 	UPROPERTY(EditAnywhere, BlueprintReadOnly, Category = "USD", config)
 	EUsdRootMotionHandling RootMotionHandling;
 
+	// What type of collision to use for static meshes generated from Prims that don't have physics schemas applied
+	UPROPERTY(EditAnywhere, BlueprintReadOnly, Category = "USD", config)
+	EUsdCollisionType FallbackCollisionType;
+
 	/**
 	 * Subdivision level to use for all subdivision meshes on the opened stage. 0 means "don't subdivide".
 	 * The maximum level of subdivision allowed can be configured via the 'USD.Subdiv.MaxSubdivLevel' cvar.
@@ -246,6 +251,9 @@ public:
 	UFUNCTION(BlueprintCallable, Category = "USD", meta = (CallInEditor = "true"))
 	USDSTAGE_API void SetRootMotionHandling(EUsdRootMotionHandling NewHandlingStrategy);
 
+	UFUNCTION(BlueprintCallable, Category = "USD", meta = (CallInEditor = "true"))
+	USDSTAGE_API void SetFallbackCollisionType(EUsdCollisionType NewCollisionType);
+
 	UFUNCTION(BlueprintCallable, Category = "USD", meta = (CallInEditor = "true"))
 	USDSTAGE_API void SetSubdivisionLevel(int32 NewLevel);
 
diff --git a/Engine/Plugins/Importers/USDImporter/Source/USDStageEditor/Private/SUSDStage.cpp b/Engine/Plugins/Importers/USDImporter/Source/USDStageEditor/Private/SUSDStage.cpp
--- a/Engine/Plugins/Importers/USDImporter/Source/USDStageEditor/Private/SUSDStage.cpp
+++ b/Engine/Plugins/Importers/USDImporter/Source/USDStageEditor/Private/SUSDStage.cpp
@@ -1119,6 +1119,15 @@ void SUsdStage::FillOptionsMenu(FMenuBuilder& MenuBuilder)
 			FNewMenuDelegate::CreateSP(this, &SUsdStage::FillRootMotionSubMenu)
 		);
 
+		MenuBuilder.AddSubMenu(
+			LOCTEXT("FallbackCollisionType", "Fallback collision type"),
+			LOCTEXT(
+				"FallbackCollisionType_ToolTip",
+				"What type of collision to use for static meshes generated from Prims that don't have physics schemas applied"
+			),
+			FNewMenuDelegate::CreateSP(this, &SUsdStage::FillFallbackCollisionTypeMenu)
+		);
+
 		MenuBuilder.AddSubMenu(
 			LOCTEXT("SubdivLevel", "Subdivision level"),
 			LOCTEXT("SubdivLevel_ToolTip", "Subdivision level to use for all subdivision meshes on the opened stage. 0 means 'don't subdivide'"),
@@ -1749,6 +1758,65 @@ void SUsdStage::FillRootMotionSubMenu(FMenuBuilder& MenuBuilder)
 	AddRootMotionEntry(EUsdRootMotionHandling::UseMotionFromSkeleton, LOCTEXT("UseMotionFromSkeletonText", "Use motion from Skeleton"));
 }
 
+void SUsdStage::FillFallbackCollisionTypeMenu(FMenuBuilder& MenuBuilder)
+{
+	auto AddFallbackCollisionTypeEntry = [&](const EUsdCollisionType& NewFallbackType, const FText& Text)
+	{
+		UEnum* Enum = StaticEnum<EUsdCollisionType>();
+		FText ToolTip = Enum->GetToolTipTextByIndex(Enum->GetIndexByValue(static_cast<uint8>(NewFallbackType)));
+
+		MenuBuilder.AddMenuEntry(
+			Text,
+			ToolTip,
+			FSlateIcon(),
+			FUIAction(
+				FExecuteAction::CreateLambda(
+					[this, NewFallbackType]()
+					{
+						if (AUsdStageActor* StageActor = GetStageActorOrCDO())
+						{
+							FScopedTransaction Transaction(FText::Format(
+								LOCTEXT("FallbackCollisionTypeTransaction", "Change fallback collision type for USD stage actor '{0}'"),
+								FText::FromString(StageActor->GetActorLabel())
+							));
+
+							// c.f. comment in SUsdStage::FillCollapsingSubMenu
+							TGuardValue<bool> MaintainSelectionGuard(bUpdatingViewportSelection, true);
+
+							StageActor->SetFallbackCollisionType(NewFallbackType);
+							if (StageActor->IsTemplate())
+							{
+								StageActor->SaveConfig();
+							}
+						}
+					}
+				),
+				FCanExecuteAction{},
+				FIsActionChecked::CreateLambda(
+					[this, NewFallbackType]()
+					{
+						if (AUsdStageActor* StageActor = GetStageActorOrCDO())
+						{
+							return StageActor->FallbackCollisionType == NewFallbackType;
+						}
+						return false;
+					}
+				)
+			),
+			NAME_None,
+			EUserInterfaceActionType::RadioButton
+		);
+	};
+
+	AddFallbackCollisionTypeEntry(EUsdCollisionType::None, LOCTEXT("CollisionNone", "None"));
+	AddFallbackCollisionTypeEntry(EUsdCollisionType::ConvexDecomposition, LOCTEXT("CollisionConvexDecomposition", "Convex decomposition"));
+	AddFallbackCollisionTypeEntry(EUsdCollisionType::ConvexHull, LOCTEXT("CollisionConvexHull", "Convex hull"));
+	AddFallbackCollisionTypeEntry(EUsdCollisionType::Sphere, LOCTEXT("CollisionSphere", "Sphere"));
+	AddFallbackCollisionTypeEntry(EUsdCollisionType::Cube, LOCTEXT("CollisionCube", "Cube"));
+	AddFallbackCollisionTypeEntry(EUsdCollisionType::MeshSimplification, LOCTEXT("CollisionMeshSimplification", "Mesh simplification"));
+	AddFallbackCollisionTypeEntry(EUsdCollisionType::Capsule, LOCTEXT("CollisionCapsule", "Capsule"));
+}
+
 void SUsdStage::FillSubdivisionLevelSubMenu(FMenuBuilder& MenuBuilder)
 {
 	if (AUsdStageActor* StageActor = GetStageActorOrCDO())
diff --git a/Engine/Plugins/Importers/USDImporter/Source/USDStageEditor/Private/SUSDStage.h b/Engine/Plugins/Importers/USDImporter/Source/USDStageEditor/Private/SUSDStage.h
--- a/Engine/Plugins/Importers/USDImporter/Source/USDStageEditor/Private/SUSDStage.h
+++ b/Engine/Plugins/Importers/USDImporter/Source/USDStageEditor/Private/SUSDStage.h
@@ -95,6 +95,7 @@ protected:
 	void FillRenderContextSubMenu(FMenuBuilder& MenuBuilder);
 	void FillMaterialPurposeSubMenu(FMenuBuilder& MenuBuilder);
 	void FillRootMotionSubMenu(FMenuBuilder& MenuBuilder);
+	void FillFallbackCollisionTypeMenu(FMenuBuilder& MenuBuilder);
 	void FillSubdivisionLevelSubMenu(FMenuBuilder& MenuBuilder);
 	void FillMetadataSubMenu(FMenuBuilder& MenuBuilder);
 	void FillCollapsingSubMenu(FMenuBuilder& MenuBuilder);
diff --git a/Engine/Plugins/Importers/USDImporter/Source/USDStageEditorViewModels/Private/USDStageViewModel.cpp b/Engine/Plugins/Importers/USDImporter/Source/USDStageEditorViewModels/Private/USDStageViewModel.cpp
--- a/Engine/Plugins/Importers/USDImporter/Source/USDStageEditorViewModels/Private/USDStageViewModel.cpp
+++ b/Engine/Plugins/Importers/USDImporter/Source/USDStageEditorViewModels/Private/USDStageViewModel.cpp
@@ -400,6 +400,7 @@ void FUsdStageViewModel::ImportStage(const TCHAR* TargetContentFolder, UUsdStage
 		ImportContext.ImportOptions->ImportTimeCode = StageActor->GetTime();
 		ImportContext.ImportOptions->NaniteTriangleThreshold = StageActor->NaniteTriangleThreshold;
 		ImportContext.ImportOptions->RootMotionHandling = StageActor->RootMotionHandling;
+		ImportContext.ImportOptions->FallbackCollisionType = StageActor->FallbackCollisionType;
 		ImportContext.ImportOptions->SubdivisionLevel = StageActor->SubdivisionLevel;
 		ImportContext.ImportOptions->MetadataOptions = StageActor->MetadataOptions;
 		ImportContext.ImportOptions->KindsToCollapse = StageActor->KindsToCollapse;
diff --git a/Engine/Plugins/Importers/USDImporter/Source/USDStageImporter/Private/USDStageImportOptions.cpp b/Engine/Plugins/Importers/USDImporter/Source/USDStageImporter/Private/USDStageImportOptions.cpp
--- a/Engine/Plugins/Importers/USDImporter/Source/USDStageImporter/Private/USDStageImportOptions.cpp
+++ b/Engine/Plugins/Importers/USDImporter/Source/USDStageImporter/Private/USDStageImportOptions.cpp
@@ -2,8 +2,8 @@
 
 #include "USDStageImportOptions.h"
 
-#include "USDSchemasModule.h"
 #include "Objects/USDSchemaTranslator.h"
+#include "USDSchemasModule.h"
 
 #include "AnalyticsEventAttribute.h"
 #include "Modules/ModuleManager.h"
@@ -112,6 +112,7 @@ void UsdUtils::AddAnalyticsAttributes(const UUsdStageImportOptions& Options, TAr
 	InOutAttributes.Emplace(TEXT("RenderContextToImport"), Options.RenderContextToImport.ToString());
 	InOutAttributes.Emplace(TEXT("MaterialPurpose"), Options.MaterialPurpose.ToString());
 	InOutAttributes.Emplace(TEXT("RootMotionHandling"), LexToString((uint8)Options.RootMotionHandling));
+	InOutAttributes.Emplace(TEXT("FallbackCollisionType"), (uint8)Options.FallbackCollisionType);
 	InOutAttributes.Emplace(TEXT("SubdivisionLevel"), LexToString(Options.SubdivisionLevel));
 
 	UsdUtils::AddAnalyticsAttributes(Options.MetadataOptions, InOutAttributes);
diff --git a/Engine/Plugins/Importers/USDImporter/Source/USDStageImporter/Private/USDStageImporter.cpp b/Engine/Plugins/Importers/USDImporter/Source/USDStageImporter/Private/USDStageImporter.cpp
--- a/Engine/Plugins/Importers/USDImporter/Source/USDStageImporter/Private/USDStageImporter.cpp
+++ b/Engine/Plugins/Importers/USDImporter/Source/USDStageImporter/Private/USDStageImporter.cpp
@@ -1935,6 +1935,7 @@ void UUsdStageImporter::ImportFromFile(FUsdStageImportContext& ImportContext)
 	TranslationContext->RenderContext = ImportContext.ImportOptions->RenderContextToImport;
 	TranslationContext->MaterialPurpose = ImportContext.ImportOptions->MaterialPurpose;
 	TranslationContext->RootMotionHandling = ImportContext.ImportOptions->RootMotionHandling;
+	TranslationContext->FallbackCollisionType = ImportContext.ImportOptions->FallbackCollisionType;
 	TranslationContext->SubdivisionLevel = ImportContext.ImportOptions->SubdivisionLevel;
 	TranslationContext->MetadataOptions = ImportContext.ImportOptions->MetadataOptions;
 	TranslationContext->ParentComponent = ImportContext.SceneActor ? ImportContext.SceneActor->GetRootComponent() : nullptr;
@@ -2153,6 +2154,7 @@ bool UUsdStageImporter::ReimportSingleAsset(
 	TranslationContext->RenderContext = ImportContext.ImportOptions->RenderContextToImport;
 	TranslationContext->MaterialPurpose = ImportContext.ImportOptions->MaterialPurpose;
 	TranslationContext->RootMotionHandling = ImportContext.ImportOptions->RootMotionHandling;
+	TranslationContext->FallbackCollisionType = ImportContext.ImportOptions->FallbackCollisionType;
 	TranslationContext->SubdivisionLevel = ImportContext.ImportOptions->SubdivisionLevel;
 	TranslationContext->MetadataOptions = ImportContext.ImportOptions->MetadataOptions;
 	TranslationContext->KindsToCollapse = (EUsdDefaultKind)ImportContext.ImportOptions->KindsToCollapse;
diff --git a/Engine/Plugins/Importers/USDImporter/Source/USDStageImporter/Public/USDStageImportOptions.h b/Engine/Plugins/Importers/USDImporter/Source/USDStageImporter/Public/USDStageImportOptions.h
--- a/Engine/Plugins/Importers/USDImporter/Source/USDStageImporter/Public/USDStageImportOptions.h
+++ b/Engine/Plugins/Importers/USDImporter/Source/USDStageImporter/Public/USDStageImportOptions.h
@@ -11,6 +11,7 @@
 
 #include "USDStageImportOptions.generated.h"
 
+enum class EUsdCollisionType : uint8;
 struct FAnalyticsEventAttribute;
 
 UENUM(BlueprintType)
@@ -136,6 +137,10 @@ public:
 	UPROPERTY(BlueprintReadWrite, config, EditAnywhere, Category = "USD options", meta = (EditCondition = bImportSkeletalAnimations))
 	EUsdRootMotionHandling RootMotionHandling = EUsdRootMotionHandling::NoAdditionalRootMotion;
 
+	// What type of collision to use for static meshes generated from Prims that don't have physics schemas applied
+	UPROPERTY(BlueprintReadWrite, config, EditAnywhere, Category = "USD options")
+	EUsdCollisionType FallbackCollisionType;
+
 	/** Subdivision level to use for all subdivision meshes on the opened stage. 0 means "don't subdivide" */
 	UPROPERTY(BlueprintReadWrite, config, EditAnywhere, Category = "USD options")
 	int32 SubdivisionLevel;
diff --git a/Engine/Plugins/Interchange/Extensions/OpenUSD/Source/Import/Private/InterchangeUsdTranslator.cpp b/Engine/Plugins/Interchange/Extensions/OpenUSD/Source/Import/Private/InterchangeUsdTranslator.cpp
--- a/Engine/Plugins/Interchange/Extensions/OpenUSD/Source/Import/Private/InterchangeUsdTranslator.cpp
+++ b/Engine/Plugins/Interchange/Extensions/OpenUSD/Source/Import/Private/InterchangeUsdTranslator.cpp
@@ -1913,35 +1913,33 @@ namespace UE::InterchangeUsdTranslator::Private
 		}
 	}
 
-	EInterchangeMeshCollision ConvertApproximationType(UsdUtils::EUsdCollisionType Approximation)
+	EInterchangeMeshCollision ConvertApproximationType(EUsdCollisionType Approximation)
 	{
 		// References:
 		// - InterchangeGenericStaticMeshPipeline.cpp, GetCollisionMeshType()
 		// - InterchangeGenericStaticMeshPipeline.cpp, AddLodDataToStaticMesh()
 
-		using namespace UsdUtils;
-
 		switch (Approximation)
 		{
 			// EInterchangeMeshCollision::None means no collision, so treat EUsdCollisionType::None
 			// as convex collision instead
-			case UsdUtils::EUsdCollisionType::None:
-			case UsdUtils::EUsdCollisionType::ConvexDecomposition:
-			case UsdUtils::EUsdCollisionType::ConvexHull:
-			case UsdUtils::EUsdCollisionType::MeshSimplification:
-			case UsdUtils::EUsdCollisionType::CustomMesh:
+			case EUsdCollisionType::None:
+			case EUsdCollisionType::ConvexDecomposition:
+			case EUsdCollisionType::ConvexHull:
+			case EUsdCollisionType::MeshSimplification:
+			case EUsdCollisionType::CustomMesh:
 			{
 				return EInterchangeMeshCollision::Convex18DOP;
 			}
-			case UsdUtils::EUsdCollisionType::Sphere:
+			case EUsdCollisionType::Sphere:
 			{
 				return EInterchangeMeshCollision::Sphere;
 			}
-			case UsdUtils::EUsdCollisionType::Cube:
+			case EUsdCollisionType::Cube:
 			{
 				return EInterchangeMeshCollision::Box;
 			}
-			case UsdUtils::EUsdCollisionType::Capsule:
+			case EUsdCollisionType::Capsule:
 			{
 				return EInterchangeMeshCollision::Capsule;
 			}
@@ -2108,7 +2106,7 @@ namespace UE::InterchangeUsdTranslator::Private
 
 				if (bSetCustomCollisionType)
 				{
-					UsdUtils::EUsdCollisionType Approximation = UsdUtils::GetCollisionApproximationType(Prim);
+					EUsdCollisionType Approximation = UsdUtils::GetCollisionApproximationType(Prim);
 					EInterchangeMeshCollision InterchangeApproximation = ConvertApproximationType(Approximation);
 					if (InterchangeApproximation != EInterchangeMeshCollision::None)
 					{
diff --git a/Engine/Plugins/Runtime/USDCore/Source/USDUtilities/Private/USDGeomMeshConversion.cpp b/Engine/Plugins/Runtime/USDCore/Source/USDUtilities/Private/USDGeomMeshConversion.cpp
--- a/Engine/Plugins/Runtime/USDCore/Source/USDUtilities/Private/USDGeomMeshConversion.cpp
+++ b/Engine/Plugins/Runtime/USDCore/Source/USDUtilities/Private/USDGeomMeshConversion.cpp
@@ -6075,7 +6075,7 @@ bool UsdUtils::IsCollisionMesh(const pxr::UsdPrim& UsdPrim)
 	return true;
 }
 
-UsdUtils::EUsdCollisionType UsdUtils::GetCollisionApproximationType(const pxr::UsdPrim& UsdPrim)
+EUsdCollisionType UsdUtils::GetCollisionApproximationType(const pxr::UsdPrim& UsdPrim)
 {
 	EUsdCollisionType Approximation(EUsdCollisionType::None);
 	if (pxr::UsdGeomMesh UsdMesh{UsdPrim})
diff --git a/Engine/Plugins/Runtime/USDCore/Source/USDUtilities/Public/Objects/USDSchemaTranslator.h b/Engine/Plugins/Runtime/USDCore/Source/USDUtilities/Public/Objects/USDSchemaTranslator.h
--- a/Engine/Plugins/Runtime/USDCore/Source/USDUtilities/Public/Objects/USDSchemaTranslator.h
+++ b/Engine/Plugins/Runtime/USDCore/Source/USDUtilities/Public/Objects/USDSchemaTranslator.h
@@ -208,6 +208,9 @@ struct USDUTILITIES_API FUsdSchemaTranslationContext : public TSharedFromThis<FU
 	/** Describes what to add to the root bone animation within generated AnimSequences, if anything */
 	EUsdRootMotionHandling RootMotionHandling = EUsdRootMotionHandling::NoAdditionalRootMotion;
 
+	/** What type of collision to use for static meshes generated from Prims that don't have physics schemas applied */
+	EUsdCollisionType FallbackCollisionType = EUsdCollisionType::ConvexHull;
+
 	/** How geometry caches are handled in the stage workflow */
 	EGeometryCacheImport GeometryCacheImport = EGeometryCacheImport::Never;
 
diff --git a/Engine/Plugins/Runtime/USDCore/Source/USDUtilities/Public/USDGeomMeshConversion.h b/Engine/Plugins/Runtime/USDCore/Source/USDUtilities/Public/USDGeomMeshConversion.h
--- a/Engine/Plugins/Runtime/USDCore/Source/USDUtilities/Public/USDGeomMeshConversion.h
+++ b/Engine/Plugins/Runtime/USDCore/Source/USDUtilities/Public/USDGeomMeshConversion.h
@@ -38,6 +38,7 @@ struct FMeshDescription;
 struct FStaticMeshLODResources;
 struct FUsdCollapsingCache;
 struct FUsdStageInfo;
+enum class EUsdCollisionType : uint8;
 
 namespace UsdUtils
 {
@@ -432,25 +433,13 @@ namespace UsdUtils
 	/** Returns whether this UsdMesh can be interpreted as a LOD of a mesh with multiple LODs (i.e. if the Mesh prim is inside a LOD variant) */
 	USDUTILITIES_API bool IsGeomMeshALOD(const pxr::UsdPrim& UsdMeshPrim);
 
-	// Approximation types from PhysicsMeshCollisionAPI and UE-specific approximations
-	enum class EUsdCollisionType : uint8
-	{
-		None,	 // no approximation so equivalent to CTF_UseComplexAsSimple
-		ConvexDecomposition,
-		ConvexHull,
-		Sphere,
-		Cube,
-		MeshSimplification,
-		Capsule,
-		CustomMesh
-	};
-
 	/** Returns whether this UsdPrim has the UsdPhysicsCollisionAPI schema and has enabled collisions */
 	USDUTILITIES_API bool IsCollisionEnabledForPrim(const pxr::UsdPrim& UsdPrim);
 
 	/** Returns whether this UsdPrim is a custom collision mesh */
 	USDUTILITIES_API bool IsCollisionMesh(const pxr::UsdPrim& UsdPrim);
 
+	/** Returns which collision type should be generated for the provided prim, according to the applied physics schemas */
 	USDUTILITIES_API EUsdCollisionType GetCollisionApproximationType(const pxr::UsdPrim& UsdPrim);
 
 	/** Returns how many LOD variants the Prim has. Note that this will return 0 if called on one of the LOD meshes themselves, it's meant to be
diff --git a/Engine/Plugins/Runtime/USDCore/Source/UnrealUSDWrapper/Public/UnrealUSDWrapper.h b/Engine/Plugins/Runtime/USDCore/Source/UnrealUSDWrapper/Public/UnrealUSDWrapper.h
--- a/Engine/Plugins/Runtime/USDCore/Source/UnrealUSDWrapper/Public/UnrealUSDWrapper.h
+++ b/Engine/Plugins/Runtime/USDCore/Source/UnrealUSDWrapper/Public/UnrealUSDWrapper.h
@@ -151,6 +151,19 @@ enum class EUsdRootMotionHandling : uint8
 	UseMotionFromSkeleton
 };
 
+UENUM()
+enum class EUsdCollisionType : uint8
+{
+	None,	 // no approximation so equivalent to CTF_UseComplexAsSimple
+	ConvexDecomposition,
+	ConvexHull,
+	Sphere,
+	Cube,
+	MeshSimplification,
+	Capsule,
+	CustomMesh
+};
+
 /** Stage option on how to handle geometry caches in stage workflow */
 UENUM()
 enum class EGeometryCacheImport : uint8
