USD Stage Editor: Allow schemas to be applied on any type of prims, not just xformables (eg. scopes are not xformables).

diff --git a/Engine/Plugins/Importers/USDImporter/Source/USDStageEditorViewModels/Private/USDPrimViewModel.cpp b/Engine/Plugins/Importers/USDImporter/Source/USDStageEditorViewModels/Private/USDPrimViewModel.cpp
--- a/Engine/Plugins/Importers/USDImporter/Source/USDStageEditorViewModels/Private/USDPrimViewModel.cpp
+++ b/Engine/Plugins/Importers/USDImporter/Source/USDStageEditorViewModels/Private/USDPrimViewModel.cpp
@@ -288,11 +288,6 @@ bool FUsdPrimViewModel::CanApplySchema(FName SchemaName) const
 		return false;
 	}
 
-	if (!PxrUsdPrim.IsA<pxr::UsdGeomXformable>())
-	{
-		return false;
-	}
-
 	return UsdUtils::CanApplySchema(UsdPrim, SchemaToken);
 #else
 	return false;
