USD: Add option to set the root prim name when exporting Levels and LevelSequences to USD.

diff --git a/Engine/Plugins/Importers/USDImporter/Source/USDExporter/Private/LevelSequenceExporterUSD.cpp b/Engine/Plugins/Importers/USDImporter/Source/USDExporter/Private/LevelSequenceExporterUSD.cpp
--- a/Engine/Plugins/Importers/USDImporter/Source/USDExporter/Private/LevelSequenceExporterUSD.cpp
+++ b/Engine/Plugins/Importers/USDImporter/Source/USDExporter/Private/LevelSequenceExporterUSD.cpp
@@ -1029,7 +1029,8 @@ namespace UE::LevelSequenceExporterUSD::Private
 			FString PrimPath = UsdUtils::GetPrimPathForObject(
 				BoundComponent,
 				TEXT(""),
-				Context.ExportOptions && Context.ExportOptions->LevelExportOptions.bExportActorFolders
+				Context.ExportOptions && Context.ExportOptions->LevelExportOptions.bExportActorFolders,
+				Context.ExportOptions->LevelExportOptions.RootPrimName
 			);
 			if (PrimPath.IsEmpty())
 			{
@@ -1059,7 +1060,8 @@ namespace UE::LevelSequenceExporterUSD::Private
 							FString ParentStageActorPrimPathOnExport = UsdUtils::GetPrimPathForObject(
 								StageActor,
 								TEXT(""),
-								Context.ExportOptions && Context.ExportOptions->LevelExportOptions.bExportActorFolders
+								Context.ExportOptions && Context.ExportOptions->LevelExportOptions.bExportActorFolders,
+								Context.ExportOptions->LevelExportOptions.RootPrimName
 							);
 							if (!ParentStageActorPrimPathOnExport.IsEmpty())
 							{
@@ -1110,7 +1112,8 @@ namespace UE::LevelSequenceExporterUSD::Private
 											   : UsdUtils::GetPrimPathForObject(
 													 &Component,
 													 TEXT(""),
-													 Context.ExportOptions && Context.ExportOptions->LevelExportOptions.bExportActorFolders
+													 Context.ExportOptions && Context.ExportOptions->LevelExportOptions.bExportActorFolders,
+													 Context.ExportOptions->LevelExportOptions.RootPrimName
 												 );
 				if (PrimPath.IsEmpty())
 				{
@@ -1674,7 +1677,8 @@ namespace UE::LevelSequenceExporterUSD::Private
 		// the same DisplayRate as the original exported sequence
 		UsdStage.SetFramesPerSecond(MovieScene->GetDisplayRate().AsDecimal());
 
-		UE::FUsdPrim RootPrim = UsdStage.OverridePrim(UE::FSdfPath(TEXT("/Root")));
+		const FString RootPrimPath = TEXT("/") + Context.ExportOptions->LevelExportOptions.RootPrimName;
+		UE::FUsdPrim RootPrim = UsdStage.OverridePrim(UE::FSdfPath(*RootPrimPath));
 		if (!RootPrim)
 		{
 			return;
diff --git a/Engine/Plugins/Importers/USDImporter/Source/USDExporter/Private/USDConversionBlueprintLibrary.cpp b/Engine/Plugins/Importers/USDImporter/Source/USDExporter/Private/USDConversionBlueprintLibrary.cpp
--- a/Engine/Plugins/Importers/USDImporter/Source/USDExporter/Private/USDConversionBlueprintLibrary.cpp
+++ b/Engine/Plugins/Importers/USDImporter/Source/USDExporter/Private/USDConversionBlueprintLibrary.cpp
@@ -485,10 +485,15 @@ void UUsdConversionBlueprintLibrary::AddPayload(
 #endif	  // USE_USD_SDK
 }
 
-FString UUsdConversionBlueprintLibrary::GetPrimPathForObject(const UObject* ActorOrComponent, const FString& ParentPrimPath, bool bUseActorFolders)
+FString UUsdConversionBlueprintLibrary::GetPrimPathForObject(
+	const UObject* ActorOrComponent,
+	const FString& ParentPrimPath,
+	bool bUseActorFolders,
+	const FString& RootPrimName
+)
 {
 #if USE_USD_SDK
-	return UsdUtils::GetPrimPathForObject(ActorOrComponent, ParentPrimPath, bUseActorFolders);
+	return UsdUtils::GetPrimPathForObject(ActorOrComponent, ParentPrimPath, bUseActorFolders, RootPrimName);
 #else
 	return {};
 #endif	  // USE_USD_SDK
diff --git a/Engine/Plugins/Importers/USDImporter/Source/USDExporter/Public/LevelExporterUSDOptions.h b/Engine/Plugins/Importers/USDImporter/Source/USDExporter/Public/LevelExporterUSDOptions.h
--- a/Engine/Plugins/Importers/USDImporter/Source/USDExporter/Public/LevelExporterUSDOptions.h
+++ b/Engine/Plugins/Importers/USDImporter/Source/USDExporter/Public/LevelExporterUSDOptions.h
@@ -18,6 +18,10 @@ struct USDEXPORTER_API FLevelExporterUSDOptionsInner
 {
 	GENERATED_BODY()
 
+	/** Name of the single top-level prim that is emitted when exporting Levels and LevelSequences to USD */
+	UPROPERTY(EditAnywhere, config, BlueprintReadWrite, Category = "Export options")
+	FString RootPrimName = TEXT("Root");
+
 	/** Whether to export only the selected actors, and assets used by them */
 	UPROPERTY(EditAnywhere, config, BlueprintReadWrite, Category = "Export options")
 	bool bSelectionOnly = false;
diff --git a/Engine/Plugins/Importers/USDImporter/Source/USDExporter/Public/USDConversionBlueprintLibrary.h b/Engine/Plugins/Importers/USDImporter/Source/USDExporter/Public/USDConversionBlueprintLibrary.h
--- a/Engine/Plugins/Importers/USDImporter/Source/USDExporter/Public/USDConversionBlueprintLibrary.h
+++ b/Engine/Plugins/Importers/USDImporter/Source/USDExporter/Public/USDConversionBlueprintLibrary.h
@@ -97,7 +97,12 @@ public:
 	static void AddPayload(const FString& ReferencingStagePath, const FString& ReferencingPrimPath, const FString& TargetStagePath);
 
 	UFUNCTION(BlueprintCallable, Category = "USD|Prim utils")
-	static FString GetPrimPathForObject(const UObject* ActorOrComponent, const FString& ParentPrimPath = TEXT(""), bool bUseActorFolders = false);
+	static FString GetPrimPathForObject(
+		const UObject* ActorOrComponent,
+		const FString& ParentPrimPath = TEXT(""),
+		bool bUseActorFolders = false,
+		const FString& RootPrimName = TEXT("Root")
+	);
 
 	UFUNCTION(BlueprintCallable, Category = "USD|Prim utils")
 	static FString GetSchemaNameForComponent(const USceneComponent* Component);
diff --git a/Engine/Plugins/Runtime/USDCore/Source/USDUtilities/Private/USDConversionUtils.cpp b/Engine/Plugins/Runtime/USDCore/Source/USDUtilities/Private/USDConversionUtils.cpp
--- a/Engine/Plugins/Runtime/USDCore/Source/USDUtilities/Private/USDConversionUtils.cpp
+++ b/Engine/Plugins/Runtime/USDCore/Source/USDUtilities/Private/USDConversionUtils.cpp
@@ -648,7 +648,7 @@ FString UsdUtils::GetSchemaNameForComponent(const USceneComponent& Component)
 	return TEXT("Xform");
 }
 
-FString UsdUtils::GetPrimPathForObject(const UObject* ActorOrComponent, const FString& ParentPrimPath, bool bUseActorFolders)
+FString UsdUtils::GetPrimPathForObject(const UObject* ActorOrComponent, const FString& ParentPrimPath, bool bUseActorFolders, const FString& RootPrimName)
 {
 	if (!ActorOrComponent)
 	{
@@ -804,7 +804,7 @@ FString UsdUtils::GetPrimPathForObject(const UObject* ActorOrComponent, const FS
 		}
 		else
 		{
-			FoundParentPath = TEXT("/Root");
+			FoundParentPath = TEXT("/") + RootPrimName;
 		}
 
 		Path = FoundParentPath / Path;
diff --git a/Engine/Plugins/Runtime/USDCore/Source/USDUtilities/Public/USDConversionUtils.h b/Engine/Plugins/Runtime/USDCore/Source/USDUtilities/Public/USDConversionUtils.h
--- a/Engine/Plugins/Runtime/USDCore/Source/USDUtilities/Public/USDConversionUtils.h
+++ b/Engine/Plugins/Runtime/USDCore/Source/USDUtilities/Public/USDConversionUtils.h
@@ -147,15 +147,19 @@ namespace UsdUtils
 	/**
 	 * Returns a prim path to use when exporting a given ActorOrComponent, assuming we're exporting according to the component attachment hierarchy.
 	 * e.g. We have a top-level actor with label "Actor" with an attached child actor "Child" that has a component "Comp1" attached to its root, and
-	 * "Comp2" attached to "Comp1". If Comp2 is provided to this function, it will return "/Root/Actor/Child/Comp1/Comp2".
+	 * "Comp2" attached to "Comp1". If Comp2 is provided to this function, it will return "/<RootPrimName>/Actor/Child/Comp1/Comp2".
 	 *
 	 * The actor label is always used in place of the root component name.
 	 *
 	 * ParentPrimPath is optional, and without it the function will recurse upwards and build the full path all the way to the root. Actor folders
 	 * are handled just like if they were actors/components
 	 */
-	USDUTILITIES_API FString
-	GetPrimPathForObject(const UObject* ActorOrComponent, const FString& ParentPrimPath = TEXT(""), bool bUseActorFolders = false);
+	USDUTILITIES_API FString GetPrimPathForObject(
+		const UObject* ActorOrComponent,
+		const FString& ParentPrimPath = TEXT(""),
+		bool bUseActorFolders = false,
+		const FString& RootPrimName = TEXT("Root")
+	);
 
 	USDUTILITIES_API TUsdStore<pxr::TfToken> GetUVSetName(int32 UVChannelIndex);
 
@@ -261,10 +265,7 @@ namespace UsdUtils
 	 * Calls GetUVSetPrimvars() followed by AssemblePrimvarsIntoPrimvarToUVIndexMap().
 	 * This exists so that it can be called without having to manipulate the USD types (e.g. for Interchange)
 	 */
-	USDUTILITIES_API TMap<FString, int32> GetPrimvarToUVIndexMap(
-		const pxr::UsdPrim& UsdPrim,
-		int32 MaxNumPrimvars = USD_PREVIEW_SURFACE_MAX_UV_SETS
-	);
+	USDUTILITIES_API TMap<FString, int32> GetPrimvarToUVIndexMap(const pxr::UsdPrim& UsdPrim, int32 MaxNumPrimvars = USD_PREVIEW_SURFACE_MAX_UV_SETS);
 
 	template<typename T, typename U>
 	FString StringifyMap(const TMap<T, U>& Map)
