USD: Allow material slot merging to affect imported LOD meshes.

diff --git a/Engine/Plugins/Importers/USDImporter/Source/USDSchemas/Private/MeshTranslationImpl.cpp b/Engine/Plugins/Importers/USDImporter/Source/USDSchemas/Private/MeshTranslationImpl.cpp
--- a/Engine/Plugins/Importers/USDImporter/Source/USDSchemas/Private/MeshTranslationImpl.cpp
+++ b/Engine/Plugins/Importers/USDImporter/Source/USDSchemas/Private/MeshTranslationImpl.cpp
@@ -4,6 +4,7 @@
 
 #include "Objects/USDInfoCache.h"
 #include "Objects/USDPrimLinkCache.h"
+#include "Objects/USDSchemaTranslator.h"
 #include "UnrealUSDWrapper.h"
 #include "USDAssetCache3.h"
 #include "USDAssetUserData.h"
@@ -783,15 +784,7 @@ void MeshTranslationImpl::SetMaterialOverrides(
 	const pxr::UsdPrim& Prim,
 	const TArray<UMaterialInterface*>& ExistingAssignments,
 	UMeshComponent& MeshComponent,
-	UUsdAssetCache3& AssetCache,
-	FUsdInfoCache& UsdInfoCache,
-	FUsdPrimLinkCache& PrimLinkCache,
-	float Time,
-	EObjectFlags Flags,
-	bool bInterpretLODs,
-	const FName& RenderContext,
-	const FName& MaterialPurpose,
-	bool bShareAssetsForIdenticalPrims
+	FUsdSchemaTranslationContext& Context
 )
 {
 	FScopedUsdAllocs Allocs;
@@ -800,15 +793,15 @@ void MeshTranslationImpl::SetMaterialOverrides(
 	pxr::UsdStageRefPtr Stage = Prim.GetStage();
 
 	pxr::TfToken RenderContextToken = pxr::UsdShadeTokens->universalRenderContext;
-	if (!RenderContext.IsNone())
+	if (!Context.RenderContext.IsNone())
 	{
-		RenderContextToken = UnrealToUsd::ConvertToken(*RenderContext.ToString()).Get();
+		RenderContextToken = UnrealToUsd::ConvertToken(*Context.RenderContext.ToString()).Get();
 	}
 
 	pxr::TfToken MaterialPurposeToken = pxr::UsdShadeTokens->allPurpose;
-	if (!MaterialPurpose.IsNone())
+	if (!Context.MaterialPurpose.IsNone())
 	{
-		MaterialPurposeToken = UnrealToUsd::ConvertToken(*MaterialPurpose.ToString()).Get();
+		MaterialPurposeToken = UnrealToUsd::ConvertToken(*Context.MaterialPurpose.ToString()).Get();
 	}
 
 	TArray<UsdUtils::FUsdPrimMaterialAssignmentInfo> LODIndexToAssignments;
@@ -816,34 +809,59 @@ void MeshTranslationImpl::SetMaterialOverrides(
 
 	// Extract material assignment info from prim if it is a LOD mesh
 	bool bInterpretedLODs = false;
-	if (bInterpretLODs && UsdUtils::IsGeomMeshALOD(Prim))
+	if (Context.bAllowInterpretingLODs && UsdUtils::IsGeomMeshALOD(Prim))
 	{
-		TMap<int32, UsdUtils::FUsdPrimMaterialAssignmentInfo> LODIndexToAssignmentsMap;
-		TFunction<bool(const pxr::UsdGeomMesh&, int32)> IterateLODs = [&](const pxr::UsdGeomMesh& LODMesh, int32 LODIndex)
+		TMap<int32, UsdUtils::FUsdPrimMaterialAssignmentInfo> LODIndexToMaterialInfoMap;
+		TMap<int32, TSet<UsdUtils::FUsdPrimMaterialSlot>> CombinedSlotsForLODIndex;
+
+		TFunction<bool(const pxr::UsdGeomMesh&, int32)> IterateLODs = [&Context,
+																	   &RenderContextToken,
+																	   &MaterialPurposeToken,
+																	   &LODIndexToMaterialInfoMap,
+																	   &CombinedSlotsForLODIndex](const pxr::UsdGeomMesh& LODMesh, int32 LODIndex)
 		{
 			// In here we need to parse the assignments again and can't rely on the cache because the info cache
 			// only has info about the default variant selection state of the stage: It won't have info about the
 			// LOD variant set setups as that involves actively toggling variants.
 			// TODO: Make the cache rebuild collect this info. Right now is not a good time for this as that would
 			// break the parallel-for setup that that function has
-			UsdUtils::FUsdPrimMaterialAssignmentInfo LODInfo = UsdUtils::GetPrimMaterialAssignments(
+			UsdUtils::FUsdPrimMaterialAssignmentInfo LocalInfo = UsdUtils::GetPrimMaterialAssignments(
 				LODMesh.GetPrim(),
-				pxr::UsdTimeCode(Time),
+				pxr::UsdTimeCode(Context.Time),
 				bProvideMaterialIndices,
 				RenderContextToken,
 				MaterialPurposeToken
 			);
-			LODIndexToAssignmentsMap.Add(LODIndex, LODInfo);
+
+			// When merging slots, we share the same material info across all LODs
+			int32 LODIndexToUse = Context.bMergeIdenticalMaterialSlots ? 0 : LODIndex;
+			TArray<UsdUtils::FUsdPrimMaterialSlot>& LODSlots = LODIndexToMaterialInfoMap.FindOrAdd(LODIndexToUse).Slots;
+			TSet<UsdUtils::FUsdPrimMaterialSlot>& CombinedLODSlotsSet = CombinedSlotsForLODIndex.FindOrAdd(LODIndexToUse);
+
+			for (UsdUtils::FUsdPrimMaterialSlot& LocalSlot : LocalInfo.Slots)
+			{
+				if (Context.bMergeIdenticalMaterialSlots)
+				{
+					if (!CombinedLODSlotsSet.Contains(LocalSlot))
+					{
+						LODSlots.Add(LocalSlot);
+						CombinedLODSlotsSet.Add(LocalSlot);
+					}
+				}
+				else
+				{
+					LODSlots.Add(LocalSlot);
+				}
+			}
 			return true;
 		};
 
 		pxr::UsdPrim ParentPrim = Prim.GetParent();
 		bInterpretedLODs = UsdUtils::IterateLODMeshes(ParentPrim, IterateLODs);
-
 		if (bInterpretedLODs)
 		{
-			LODIndexToAssignmentsMap.KeySort(TLess<int32>());
-			for (TPair<int32, UsdUtils::FUsdPrimMaterialAssignmentInfo>& Entry : LODIndexToAssignmentsMap)
+			LODIndexToMaterialInfoMap.KeySort(TLess<int32>());
+			for (TPair<int32, UsdUtils::FUsdPrimMaterialAssignmentInfo>& Entry : LODIndexToMaterialInfoMap)
 			{
 				LODIndexToAssignments.Add(MoveTemp(Entry.Value));
 			}
@@ -861,7 +879,7 @@ void MeshTranslationImpl::SetMaterialOverrides(
 		// likely just some root Xform prim.
 		// Note: This only works because we'll rebuild the cache when our material purpose/render context changes,
 		// and because in USD relationships (and so material bindings) can't vary with time
-		TOptional<TArray<UsdUtils::FUsdPrimMaterialSlot>> SubtreeSlots = UsdInfoCache.GetSubtreeMaterialSlots(UE::FSdfPath{PrimPath});
+		TOptional<TArray<UsdUtils::FUsdPrimMaterialSlot>> SubtreeSlots = Context.UsdInfoCache->GetSubtreeMaterialSlots(UE::FSdfPath{PrimPath});
 		if (SubtreeSlots.IsSet())
 		{
 			UsdUtils::FUsdPrimMaterialAssignmentInfo& NewInfo = LODIndexToAssignments.Emplace_GetRef();
@@ -871,7 +889,7 @@ void MeshTranslationImpl::SetMaterialOverrides(
 		{
 			LODIndexToAssignments = {UsdUtils::GetPrimMaterialAssignments(
 				ValidPrim,
-				pxr::UsdTimeCode(Time),
+				pxr::UsdTimeCode(Context.Time),
 				bProvideMaterialIndices,
 				RenderContextToken,
 				MaterialPurposeToken
@@ -928,10 +946,10 @@ void MeshTranslationImpl::SetMaterialOverrides(
 		ResolvedMaterials = MeshTranslationImpl::ResolveMaterialAssignmentInfo(
 			ValidPrim,
 			LODIndexToAssignments,
-			AssetCache,
-			PrimLinkCache,
-			Flags,
-			bShareAssetsForIdenticalPrims
+			*Context.UsdAssetCache,
+			*Context.PrimLinkCache,
+			Context.ObjectFlags,
+			Context.bShareAssetsForIdenticalPrims
 		);
 	}
 
diff --git a/Engine/Plugins/Importers/USDImporter/Source/USDSchemas/Private/MeshTranslationImpl.h b/Engine/Plugins/Importers/USDImporter/Source/USDSchemas/Private/MeshTranslationImpl.h
--- a/Engine/Plugins/Importers/USDImporter/Source/USDSchemas/Private/MeshTranslationImpl.h
+++ b/Engine/Plugins/Importers/USDImporter/Source/USDSchemas/Private/MeshTranslationImpl.h
@@ -21,6 +21,7 @@ class UMaterialInterface;
 class UMeshComponent;
 class UUsdAssetCache3;
 class UUsdMeshAssetUserData;
+struct FUsdSchemaTranslationContext;
 namespace UsdUtils
 {
 	struct FUsdPrimMaterialSlot;
@@ -48,15 +49,7 @@ namespace MeshTranslationImpl
 		const pxr::UsdPrim& Prim,
 		const TArray<UMaterialInterface*>& ExistingAssignments,
 		UMeshComponent& MeshComponent,
-		UUsdAssetCache3& AssetCache,
-		FUsdInfoCache& UsdInfoCache,
-		FUsdPrimLinkCache& PrimLinkCache,
-		float Time,
-		EObjectFlags Flags,
-		bool bInterpretLODs,
-		const FName& RenderContext,
-		const FName& MaterialPurpose,
-		bool bShareAssetsForIdenticalPrims
+		FUsdSchemaTranslationContext& Context
 	);
 
 	void RecordSourcePrimsForMaterialSlots(
diff --git a/Engine/Plugins/Importers/USDImporter/Source/USDSchemas/Private/USDGeomMeshTranslator.cpp b/Engine/Plugins/Importers/USDImporter/Source/USDSchemas/Private/USDGeomMeshTranslator.cpp
--- a/Engine/Plugins/Importers/USDImporter/Source/USDSchemas/Private/USDGeomMeshTranslator.cpp
+++ b/Engine/Plugins/Importers/USDImporter/Source/USDSchemas/Private/USDGeomMeshTranslator.cpp
@@ -283,7 +283,7 @@ namespace UsdGeomMeshTranslatorImpl
 					StaticMesh.GetStaticMaterials().Add(MoveTemp(StaticMaterial));
 					bMaterialAssignementsHaveChanged = true;
 				}
-				else if (!(StaticMesh.GetStaticMaterials()[StaticMeshSlotIndex] == StaticMaterial))
+				else if (StaticMesh.GetStaticMaterials()[StaticMeshSlotIndex] != StaticMaterial)
 				{
 					StaticMesh.GetStaticMaterials()[StaticMeshSlotIndex] = MoveTemp(StaticMaterial);
 					bMaterialAssignementsHaveChanged = true;
@@ -382,14 +382,17 @@ namespace UsdGeomMeshTranslatorImpl
 		TMap<FString, int32> PrimvarToUVIndex = UsdUtils::CombinePrimvarsIntoUVSets(AllPrimvars, PreferredPrimvars);
 
 		TFunction<bool(const pxr::UsdGeomMesh&, int32)> ConvertLOD =
-			[&OptionsCopy, &Options, &PrimvarToUVIndex, &LODIndexToMeshDescriptionMap, &LODIndexToMaterialInfoMap, &OutLODMetadata, &MetadataOptions](
+			[&OptionsCopy, &Options, &LODIndexToMeshDescriptionMap, &LODIndexToMaterialInfoMap, &OutLODMetadata, &MetadataOptions, &PrimvarToUVIndex](
 				const pxr::UsdGeomMesh& LODMesh,
 				int32 LODIndex
 			)
 		{
 			FMeshDescription TempMeshDescription;
 
-			UsdUtils::FUsdPrimMaterialAssignmentInfo TempMaterialInfo;
+			// When merging slots we combine all our assignments on a single MaterialAssignmentInfo
+			UsdUtils::FUsdPrimMaterialAssignmentInfo& TempMaterialInfo = LODIndexToMaterialInfoMap.FindOrAdd(
+				Options.bMergeIdenticalMaterialSlots ? 0 : LODIndex
+			);
 			TempMaterialInfo.PrimvarToUVIndex = PrimvarToUVIndex;
 
 			FStaticMeshAttributes StaticMeshAttributes(TempMeshDescription);
@@ -423,7 +426,6 @@ namespace UsdGeomMeshTranslatorImpl
 					if (bSuccess)
 					{
 						OptionsCopy.AdditionalTransform = MeshTransform * Options.AdditionalTransform;
-						OptionsCopy.bMergeIdenticalMaterialSlots = false;	 // We only merge slots when collapsing, and we never collapse LODs
 
 						bSuccess &= UsdToUnreal::ConvertGeomMesh(LODMesh, TempMeshDescription, TempMaterialInfo, OptionsCopy);
 					}
@@ -433,7 +435,6 @@ namespace UsdGeomMeshTranslatorImpl
 			if (bSuccess)
 			{
 				LODIndexToMeshDescriptionMap.Add(LODIndex, MoveTemp(TempMeshDescription));
-				LODIndexToMaterialInfoMap.Add(LODIndex, MoveTemp(TempMaterialInfo));
 
 				if (MetadataOptions.bCollectMetadata)
 				{
@@ -455,14 +456,22 @@ namespace UsdGeomMeshTranslatorImpl
 		}
 
 		// Place them in order as we can't have e.g. LOD0 and LOD2 without LOD1, and there's no reason downstream code needs to care about this
-		OutLODIndexToMeshDescription.Reset(LODIndexToMeshDescriptionMap.Num());
-		OutLODIndexToMaterialInfo.Reset(LODIndexToMaterialInfoMap.Num());
-		LODIndexToMeshDescriptionMap.KeySort(TLess<int32>());
-		for (TPair<int32, FMeshDescription>& Entry : LODIndexToMeshDescriptionMap)
 		{
-			const int32 OldLODIndex = Entry.Key;
-			OutLODIndexToMeshDescription.Add(MoveTemp(Entry.Value));
-			OutLODIndexToMaterialInfo.Add(MoveTemp(LODIndexToMaterialInfoMap[OldLODIndex]));
+			LODIndexToMeshDescriptionMap.KeySort(TLess<int32>());
+			OutLODIndexToMeshDescription.Reset(LODIndexToMeshDescriptionMap.Num());
+			for (TPair<int32, FMeshDescription>& Entry : LODIndexToMeshDescriptionMap)
+			{
+				const int32 OldLODIndex = Entry.Key;
+				OutLODIndexToMeshDescription.Add(MoveTemp(Entry.Value));
+			}
+
+			LODIndexToMaterialInfoMap.KeySort(TLess<int32>());
+			OutLODIndexToMaterialInfo.Reset(LODIndexToMaterialInfoMap.Num());
+			for (TPair<int32, UsdUtils::FUsdPrimMaterialAssignmentInfo>& Entry : LODIndexToMaterialInfoMap)
+			{
+				const int32 OldLODIndex = Entry.Key;
+				OutLODIndexToMaterialInfo.Add(MoveTemp(LODIndexToMaterialInfoMap[OldLODIndex]));
+			}
 		}
 
 		return true;
@@ -1547,7 +1556,11 @@ void FBuildStaticMeshTaskChain::SetupTasks()
 
 			   if (UUsdMeshAssetUserData* UserData = UsdUnreal::ObjectUtils::GetOrCreateAssetUserData<UUsdMeshAssetUserData>(StaticMesh))
 			   {
-				   UserData->PrimvarToUVIndex = LODIndexToMaterialInfo[0].PrimvarToUVIndex;	   // We use the same primvar mapping for all LODs
+				   if (LODIndexToMaterialInfo.Num() > 0)
+				   {
+					   // We use the same primvar mapping for all LODs
+					   UserData->PrimvarToUVIndex = LODIndexToMaterialInfo[0].PrimvarToUVIndex;
+				   }
 				   UserData->PrimPaths.AddUnique(PrimPathString);
 
 				   if (Context->MetadataOptions.bCollectMetadata)
@@ -1806,7 +1819,12 @@ void FGeomMeshCreateAssetsTaskChain::SetupTasks()
 			   if (MeshDescription.IsEmpty())
 			   {
 				   LODIndexToMeshDescription.RemoveAt(LODIndex, EAllowShrinking::No);
-				   LODIndexToMaterialInfo.RemoveAt(LODIndex, EAllowShrinking::No);
+
+				   if (LODIndexToMaterialInfo.IsValidIndex(LODIndex))
+				   {
+					   LODIndexToMaterialInfo.RemoveAt(LODIndex, EAllowShrinking::No);
+				   }
+
 				   UE_LOG(
 					   LogUsd,
 					   Warning,
diff --git a/Engine/Plugins/Importers/USDImporter/Source/USDSchemas/Private/USDGeomPointInstancerTranslator.cpp b/Engine/Plugins/Importers/USDImporter/Source/USDSchemas/Private/USDGeomPointInstancerTranslator.cpp
--- a/Engine/Plugins/Importers/USDImporter/Source/USDSchemas/Private/USDGeomPointInstancerTranslator.cpp
+++ b/Engine/Plugins/Importers/USDImporter/Source/USDSchemas/Private/USDGeomPointInstancerTranslator.cpp
@@ -559,20 +559,7 @@ void FUsdGeomPointInstancerTranslator::UpdateComponents(USceneComponent* PointIn
 					ExistingAssignments.Add(StaticMaterial.MaterialInterface);
 				}
 
-				MeshTranslationImpl::SetMaterialOverrides(
-					PrototypeUsdPrim.Get(),
-					ExistingAssignments,
-					*ISMComponent,
-					*Context->UsdAssetCache,
-					*Context->UsdInfoCache,
-					*Context->PrimLinkCache,
-					Context->Time,
-					Context->ObjectFlags,
-					Context->bAllowInterpretingLODs,
-					Context->RenderContext,
-					Context->MaterialPurpose,
-					Context->bShareAssetsForIdenticalPrims
-				);
+				MeshTranslationImpl::SetMaterialOverrides(PrototypeUsdPrim.Get(), ExistingAssignments, *ISMComponent, *Context);
 			}
 
 			// Handle visibility for the ISM component
diff --git a/Engine/Plugins/Importers/USDImporter/Source/USDSchemas/Private/USDGeomPrimitiveTranslator.cpp b/Engine/Plugins/Importers/USDImporter/Source/USDSchemas/Private/USDGeomPrimitiveTranslator.cpp
--- a/Engine/Plugins/Importers/USDImporter/Source/USDSchemas/Private/USDGeomPrimitiveTranslator.cpp
+++ b/Engine/Plugins/Importers/USDImporter/Source/USDSchemas/Private/USDGeomPrimitiveTranslator.cpp
@@ -162,20 +162,7 @@ USceneComponent* FUsdGeomPrimitiveTranslator::CreateComponents()
 					ExistingAssignments.Add(StaticMaterial.MaterialInterface);
 				}
 
-				MeshTranslationImpl::SetMaterialOverrides(
-					GetPrim(),
-					ExistingAssignments,
-					*StaticMeshComponent,
-					*Context->UsdAssetCache,
-					*Context->UsdInfoCache,
-					*Context->PrimLinkCache,
-					Context->Time,
-					Context->ObjectFlags,
-					Context->bAllowInterpretingLODs,
-					Context->RenderContext,
-					Context->MaterialPurpose,
-					Context->bShareAssetsForIdenticalPrims
-				);
+				MeshTranslationImpl::SetMaterialOverrides(GetPrim(), ExistingAssignments, *StaticMeshComponent, *Context);
 			}
 		}
 	}
diff --git a/Engine/Plugins/Importers/USDImporter/Source/USDSchemas/Private/USDGeomXformableTranslator.cpp b/Engine/Plugins/Importers/USDImporter/Source/USDSchemas/Private/USDGeomXformableTranslator.cpp
--- a/Engine/Plugins/Importers/USDImporter/Source/USDSchemas/Private/USDGeomXformableTranslator.cpp
+++ b/Engine/Plugins/Importers/USDImporter/Source/USDSchemas/Private/USDGeomXformableTranslator.cpp
@@ -780,20 +780,7 @@ void FUsdGeomXformableTranslator::UpdateComponents(USceneComponent* SceneCompone
 						ExistingAssignments.Add(StaticMaterial.MaterialInterface);
 					}
 
-					MeshTranslationImpl::SetMaterialOverrides(
-						GetPrim(),
-						ExistingAssignments,
-						*StaticMeshComponent,
-						*Context->UsdAssetCache,
-						*Context->UsdInfoCache,
-						*Context->PrimLinkCache,
-						Context->Time,
-						Context->ObjectFlags,
-						Context->bAllowInterpretingLODs,
-						Context->RenderContext,
-						Context->MaterialPurpose,
-						Context->bShareAssetsForIdenticalPrims
-					);
+					MeshTranslationImpl::SetMaterialOverrides(GetPrim(), ExistingAssignments, *StaticMeshComponent, *Context);
 				}
 			}
 		}
diff --git a/Engine/Plugins/Importers/USDImporter/Source/USDSchemas/Private/USDGeometryCacheTranslator.cpp b/Engine/Plugins/Importers/USDImporter/Source/USDSchemas/Private/USDGeometryCacheTranslator.cpp
--- a/Engine/Plugins/Importers/USDImporter/Source/USDSchemas/Private/USDGeometryCacheTranslator.cpp
+++ b/Engine/Plugins/Importers/USDImporter/Source/USDSchemas/Private/USDGeometryCacheTranslator.cpp
@@ -1001,20 +1001,7 @@ USceneComponent* FUsdGeometryCacheTranslator::CreateComponents()
 				// Geometry caches don't support LODs
 				const bool bAllowInterpretingLODs = false;
 
-				MeshTranslationImpl::SetMaterialOverrides(
-					GetPrim(),
-					GeometryCache->Materials,
-					*Component,
-					*Context->UsdAssetCache,
-					*Context->UsdInfoCache,
-					*Context->PrimLinkCache,
-					Context->Time,
-					Context->ObjectFlags,
-					bAllowInterpretingLODs,
-					Context->RenderContext,
-					Context->MaterialPurpose,
-					Context->bShareAssetsForIdenticalPrims
-				);
+				MeshTranslationImpl::SetMaterialOverrides(GetPrim(), GeometryCache->Materials, *Component, *Context);
 
 				// Check if the prim has the GroomBinding schema and setup the component and assets necessary to bind the groom to the GeometryCache
 				if (UsdUtils::PrimHasSchema(GetPrim(), UnrealIdentifiers::GroomBindingAPI) && Context->bAllowParsingGroomAssets)
diff --git a/Engine/Plugins/Importers/USDImporter/Source/USDSchemas/Private/USDSkelSkeletonTranslator.cpp b/Engine/Plugins/Importers/USDImporter/Source/USDSchemas/Private/USDSkelSkeletonTranslator.cpp
--- a/Engine/Plugins/Importers/USDImporter/Source/USDSchemas/Private/USDSkelSkeletonTranslator.cpp
+++ b/Engine/Plugins/Importers/USDImporter/Source/USDSchemas/Private/USDSkelSkeletonTranslator.cpp
@@ -573,7 +573,9 @@ namespace UsdSkelSkeletonTranslatorImpl
 			}
 
 			FSkeletalMeshImportData& LODImportData = LODIndexToSkeletalMeshImportDataMap.FindOrAdd(LODIndex);
-			UsdUtils::FUsdPrimMaterialAssignmentInfo& LODMaterialInfo = LODIndexToMaterialInfoMap.FindOrAdd(LODIndex);
+			UsdUtils::FUsdPrimMaterialAssignmentInfo& LODMaterialInfo = LODIndexToMaterialInfoMap.FindOrAdd(
+				Options.bMergeIdenticalMaterialSlots ? 0 : LODIndex
+			);
 
 			// BlendShape data is respective to point indices for each mesh in isolation, but we combine all points
 			// into one FSkeletalMeshImportData per LOD, so we need to remap the indices using this
@@ -673,7 +675,11 @@ namespace UsdSkelSkeletonTranslatorImpl
 
 			const int32 OldLODIndex = Entry.Key;
 			const int32 NewLODIndex = OutLODIndexToSkeletalMeshImportData.Add(MoveTemp(ImportData));
-			OutLODIndexToMaterialInfo.Add(LODIndexToMaterialInfoMap[OldLODIndex]);
+
+			if (UsdUtils::FUsdPrimMaterialAssignmentInfo* FoundInfo = LODIndexToMaterialInfoMap.Find(OldLODIndex))
+			{
+				OutLODIndexToMaterialInfo.Add(*FoundInfo);
+			}
 
 			// Keep track of these to remap blendshapes
 			OldLODIndexToNewLODIndex.Add(OldLODIndex, NewLODIndex);
@@ -711,15 +717,7 @@ namespace UsdSkelSkeletonTranslatorImpl
 		const pxr::UsdSkelBinding& SkeletonBinding,
 		const TArray<UMaterialInterface*>& ExistingAssignments,
 		UMeshComponent& MeshComponent,
-		UUsdAssetCache3& AssetCache,
-		FUsdPrimLinkCache& PrimLinkCache,
-		float Time,
-		EObjectFlags Flags,
-		bool bInterpretLODs,
-		const FName& RenderContext,
-		const FName& MaterialPurpose,
-		const EUsdPurpose PurposesToLoad,
-		bool bShareAssetsForIdenticalPrims
+		FUsdSchemaTranslationContext& Context
 	)
 	{
 		FScopedUsdAllocs Allocs;
@@ -735,24 +733,25 @@ namespace UsdSkelSkeletonTranslatorImpl
 		pxr::UsdStageRefPtr Stage = SkeletonPrim.GetStage();
 
 		pxr::TfToken RenderContextToken = pxr::UsdShadeTokens->universalRenderContext;
-		if (!RenderContext.IsNone())
+		if (!Context.RenderContext.IsNone())
 		{
-			RenderContextToken = UnrealToUsd::ConvertToken(*RenderContext.ToString()).Get();
+			RenderContextToken = UnrealToUsd::ConvertToken(*Context.RenderContext.ToString()).Get();
 		}
 
 		pxr::TfToken MaterialPurposeToken = pxr::UsdShadeTokens->allPurpose;
-		if (!MaterialPurpose.IsNone())
+		if (!Context.MaterialPurpose.IsNone())
 		{
-			MaterialPurposeToken = UnrealToUsd::ConvertToken(*MaterialPurpose.ToString()).Get();
+			MaterialPurposeToken = UnrealToUsd::ConvertToken(*Context.MaterialPurpose.ToString()).Get();
 		}
 
 		TMap<int32, UsdUtils::FUsdPrimMaterialAssignmentInfo> LODIndexToMaterialInfoMap;
 		TMap<int32, TSet<UsdUtils::FUsdPrimMaterialSlot>> CombinedSlotsForLODIndex;
-		TFunction<bool(const pxr::UsdGeomMesh&, int32)> IterateLODsLambda =
-			[&LODIndexToMaterialInfoMap, &CombinedSlotsForLODIndex, Time, RenderContextToken, MaterialPurposeToken, PurposesToLoad](
-				const pxr::UsdGeomMesh& LODMesh,
-				int32 LODIndex
-			)
+
+		TFunction<bool(const pxr::UsdGeomMesh&, int32)> IterateLODsLambda = [&LODIndexToMaterialInfoMap,
+																			 &CombinedSlotsForLODIndex,
+																			 &Context,
+																			 RenderContextToken,
+																			 MaterialPurposeToken](const pxr::UsdGeomMesh& LODMesh, int32 LODIndex)
 		{
 			if (LODMesh && LODMesh.ComputeVisibility() == pxr::UsdGeomTokens->invisible)
 			{
@@ -761,19 +760,22 @@ namespace UsdSkelSkeletonTranslatorImpl
 
 			// Ignore prims with disabled purposes: We need to match the material slot ordering that was used
 			// to generate the mesh in the first place, so this is important
-			if (!EnumHasAllFlags(PurposesToLoad, IUsdPrim::GetPurpose(LODMesh.GetPrim())))
+			if (!EnumHasAllFlags(Context.PurposesToLoad, IUsdPrim::GetPurpose(LODMesh.GetPrim())))
 			{
 				return true;
 			}
 
-			TArray<UsdUtils::FUsdPrimMaterialSlot>& CombinedLODSlots = LODIndexToMaterialInfoMap.FindOrAdd(LODIndex).Slots;
-			TSet<UsdUtils::FUsdPrimMaterialSlot>& CombinedLODSlotsSet = CombinedSlotsForLODIndex.FindOrAdd(LODIndex);
+			// When merging slots, we share the same material info across all LODs
+			int32 LODIndexToUse = Context.bMergeIdenticalMaterialSlots ? 0 : LODIndex;
+
+			TArray<UsdUtils::FUsdPrimMaterialSlot>& CombinedLODSlots = LODIndexToMaterialInfoMap.FindOrAdd(LODIndexToUse).Slots;
+			TSet<UsdUtils::FUsdPrimMaterialSlot>& CombinedLODSlotsSet = CombinedSlotsForLODIndex.FindOrAdd(LODIndexToUse);
 
 			const bool bProvideMaterialIndices = false;	   // We have no use for material indices and it can be slow to retrieve, as it will iterate
 														   // all faces
 			UsdUtils::FUsdPrimMaterialAssignmentInfo LocalInfo = UsdUtils::GetPrimMaterialAssignments(
 				LODMesh.GetPrim(),
-				pxr::UsdTimeCode(Time),
+				pxr::UsdTimeCode(Context.Time),
 				bProvideMaterialIndices,
 				RenderContextToken,
 				MaterialPurposeToken
@@ -818,7 +820,7 @@ namespace UsdSkelSkeletonTranslatorImpl
 			FString ParentPrimPath = UsdToUnreal::ConvertPath(ParentPrim.GetPath());
 
 			bool bInterpretedLODs = false;
-			if (bInterpretLODs && UsdUtils::IsGeomMeshALOD(MeshPrim) && !ProcessedLODParentPaths.Contains(ParentPrimPath))
+			if (Context.bAllowInterpretingLODs && UsdUtils::IsGeomMeshALOD(MeshPrim) && !ProcessedLODParentPaths.Contains(ParentPrimPath))
 			{
 				ProcessedLODParentPaths.Add(ParentPrimPath);
 
@@ -872,10 +874,10 @@ namespace UsdSkelSkeletonTranslatorImpl
 		TMap<const UsdUtils::FUsdPrimMaterialSlot*, UMaterialInterface*> ResolvedMaterials = MeshTranslationImpl::ResolveMaterialAssignmentInfo(
 			ValidSkeletonPrim,
 			LODIndexToAssignments,
-			AssetCache,
-			PrimLinkCache,
-			Flags,
-			bShareAssetsForIdenticalPrims
+			*Context.UsdAssetCache,
+			*Context.PrimLinkCache,
+			Context.ObjectFlags,
+			Context.bShareAssetsForIdenticalPrims
 		);
 
 		// Compare resolved materials with existing assignments, and create overrides if we need to
@@ -2127,20 +2129,7 @@ void FUsdSkelSkeletonTranslator::UpdateComponents(USceneComponent* SceneComponen
 				ExistingAssignments.Add(SkeletalMaterial.MaterialInterface);
 			}
 
-			UsdSkelSkeletonTranslatorImpl::SetMaterialOverrides(
-				SkeletonBindingForPrim.Get(),
-				ExistingAssignments,
-				*SkeletalMeshComponent,
-				*Context->UsdAssetCache,
-				*Context->PrimLinkCache,
-				Context->Time,
-				Context->ObjectFlags,
-				Context->bAllowInterpretingLODs,
-				Context->RenderContext,
-				Context->MaterialPurpose,
-				Context->PurposesToLoad,
-				Context->bShareAssetsForIdenticalPrims
-			);
+			UsdSkelSkeletonTranslatorImpl::SetMaterialOverrides(SkeletonBindingForPrim.Get(), ExistingAssignments, *SkeletalMeshComponent, *Context);
 		}
 	}
 
diff --git a/Engine/Plugins/Importers/USDImporter/Source/USDStageImporter/Public/USDStageImportOptions.h b/Engine/Plugins/Importers/USDImporter/Source/USDStageImporter/Public/USDStageImportOptions.h
--- a/Engine/Plugins/Importers/USDImporter/Source/USDStageImporter/Public/USDStageImportOptions.h
+++ b/Engine/Plugins/Importers/USDImporter/Source/USDStageImporter/Public/USDStageImportOptions.h
@@ -220,8 +220,8 @@ public:
 	int32 KindsToCollapse;
 
 	/**
-	 * If enabled, when multiple mesh prims are collapsed into a single static mesh, identical material slots are merged into one slot.
-	 * Otherwise, material slots are simply appended to the list.
+	 * Identical material slots will be combined into a single slot if this is enabled. This is only performed in the context
+	 * of mesh collapsing, or when parsing LOD variant sets (see bInterpretLODs)
 	 */
 	UPROPERTY(BlueprintReadWrite, config, EditAnywhere, Category = "Processing")
 	bool bMergeIdenticalMaterialSlots;
diff --git a/Engine/Plugins/Runtime/USDCore/Source/USDUtilities/Private/Objects/USDInfoCache.cpp b/Engine/Plugins/Runtime/USDCore/Source/USDUtilities/Private/Objects/USDInfoCache.cpp
--- a/Engine/Plugins/Runtime/USDCore/Source/USDUtilities/Private/Objects/USDInfoCache.cpp
+++ b/Engine/Plugins/Runtime/USDCore/Source/USDUtilities/Private/Objects/USDInfoCache.cpp
@@ -740,7 +740,7 @@ namespace UE::UsdInfoCache::Private
 			TArray<pxr::UsdPrim> Children;
 			for (pxr::UsdPrim Child : UsdPrim.GetFilteredChildren(pxr::UsdTraverseInstanceProxies(pxr::UsdPrimAllPrimsPredicate)))
 			{
-				// We don't care about non-GeomImagable prims 
+				// We don't care about non-GeomImagable prims
 				// (materials, etc., stuff we don't have schema translators for will be skipped and default to bCanBeCollapsed=true)
 				if (Child.IsA<pxr::UsdGeomImageable>())
 				{
@@ -1636,9 +1636,11 @@ namespace UE::USDInfoCache::Private
 			}
 		}
 
-		// For now we only ever merge material slots when collapsing (and not if we're collapsing due to being a geometry cache)
+		// For now we only ever merge material slots when collapsing and if parsing LODs (and not if we're collapsing due to being a geometry cache)
 		bPrimIsCollapsedOrCollapseRoot |= Info.bCollapsesChildren;
-		bool bCanMergeSlotsForThisPrim = bPrimIsCollapsedOrCollapseRoot || Info.PrimPath.IsAbsoluteRootPath();
+		bool bCanMergeSlotsForThisPrim = bPrimIsCollapsedOrCollapseRoot || Info.PrimPath.IsAbsoluteRootPath()
+										 || (Context.bAllowInterpretingLODs
+											 && UsdUtils::DoesPrimContainMeshLODs(Context.Stage.GetPrimAtPath(Info.PrimPath)));
 		if (bCanMergeSlotsForThisPrim)
 		{
 			if (Impl.IsPotentialGeometryCacheRootInner(Info, Context.Stage.GetPrimAtPath(Info.PrimPath)))
diff --git a/Engine/Plugins/Runtime/USDCore/Source/USDUtilities/Private/USDGeomMeshConversion.cpp b/Engine/Plugins/Runtime/USDCore/Source/USDUtilities/Private/USDGeomMeshConversion.cpp
--- a/Engine/Plugins/Runtime/USDCore/Source/USDUtilities/Private/USDGeomMeshConversion.cpp
+++ b/Engine/Plugins/Runtime/USDCore/Source/USDUtilities/Private/USDGeomMeshConversion.cpp
@@ -2723,38 +2723,54 @@ namespace UE::UsdGeomMeshConversion::Private
 
 		// Polygons
 		{
-			TMap<int32, FPolygonGroupID> PolygonGroupMapping;
 			TArray<FVertexInstanceID> CornerInstanceIDs;
 			TArray<FVertexID> CornerVerticesIDs;
 			int32 CurrentVertexInstanceIndex = 0;
 			TPolygonGroupAttributesRef<FName> MaterialSlotNames = StaticMeshAttributes.GetPolygonGroupMaterialSlotNames();
 
-			// If we're going to share existing material slots, we'll need to share existing PolygonGroups in the mesh description,
-			// so we need to traverse it and prefill PolygonGroupMapping
-			if (InOptions.bMergeIdenticalMaterialSlots)
+			// Prepare our polygon group mapping acknowledging whatever is already in the mesh description
+			TMap<int32, FPolygonGroupID> PolygonGroupMapping;
+			for (FPolygonGroupID PolygonGroupID : OutMeshDescription.PolygonGroups().GetElementIDs())
 			{
-				for (FPolygonGroupID PolygonGroupID : OutMeshDescription.PolygonGroups().GetElementIDs())
-				{
-					// We always create our polygon groups in order with our combined material slots, so its easy to reconstruct this mapping here
-					PolygonGroupMapping.Add(PolygonGroupID.GetValue(), PolygonGroupID);
-				}
+				PolygonGroupMapping.Add(PolygonGroupID.GetValue(), PolygonGroupID);
 			}
 
-			// Material slots
-			// Note that we always create these in the order they show up in LocalInfo: If we created these on-demand when parsing polygons (like
-			// before) we could run into polygons in a different order than the material slots and end up with different material assignments. We
-			// could use the StaticMesh's SectionInfoMap to unswitch things, but that's not available at runtime so we better do this here
-			for (int32 LocalMaterialIndex = 0; LocalMaterialIndex < InMeshData.LocalMaterialInfo.Slots.Num(); ++LocalMaterialIndex)
+			auto EnsurePolygonGroupForIndex = [&PolygonGroupMapping, &OutMeshDescription, &MaterialSlotNames](int32 Index)
 			{
-				const int32 CombinedMaterialIndex = LocalToCombinedMaterialSlotIndices[LocalMaterialIndex];
-				if (!PolygonGroupMapping.Contains(CombinedMaterialIndex))
+				if (!PolygonGroupMapping.Contains(Index))
 				{
 					FPolygonGroupID NewPolygonGroup = OutMeshDescription.CreatePolygonGroup();
-					PolygonGroupMapping.Add(CombinedMaterialIndex, NewPolygonGroup);
+					PolygonGroupMapping.Add(Index, NewPolygonGroup);
 
 					// This is important for runtime, where the material slots are matched to LOD sections based on their material slot name
 					MaterialSlotNames[NewPolygonGroup] = *LexToString(NewPolygonGroup.GetValue());
 				}
+			};
+
+			// Material slots
+			int32 LastHandledIndex = 0;
+			for (int32 LocalMaterialIndex = 0; LocalMaterialIndex < InMeshData.LocalMaterialInfo.Slots.Num(); ++LocalMaterialIndex)
+			{
+				const int32 CombinedMaterialIndex = LocalToCombinedMaterialSlotIndices[LocalMaterialIndex];
+
+				if (InOptions.bMergeIdenticalMaterialSlots)
+				{
+					// Ensure we have polygon groups for every combined material slot index up until the one we'll use.
+					// Doing this keeps a 1-to-1 mapping between polygon group IDs and combined material slot IDs.
+					// This is particularly relevant when merging slots in the context of LODs, when we won't share the MeshDescriptions between
+					// calls to this function, but we will share the combined material assignments. There will be nothing to prefill our
+					// PolygonGroupMapping in that case, but since it's a 1-to-1 mapping it's easy to set it up here
+					for (; LastHandledIndex <= CombinedMaterialIndex; ++LastHandledIndex)
+					{
+						EnsurePolygonGroupForIndex(LastHandledIndex);
+					}
+				}
+				else
+				{
+					// If we're not merging slots we should always make brand new polygon groups for every material slot
+					// we get here
+					EnsurePolygonGroupForIndex(CombinedMaterialIndex);
+				}
 			}
 
 			// Velocities
diff --git a/Engine/Plugins/Runtime/USDCore/Source/USDUtilities/Private/USDSkeletalDataConversion.cpp b/Engine/Plugins/Runtime/USDCore/Source/USDUtilities/Private/USDSkeletalDataConversion.cpp
--- a/Engine/Plugins/Runtime/USDCore/Source/USDUtilities/Private/USDSkeletalDataConversion.cpp
+++ b/Engine/Plugins/Runtime/USDCore/Source/USDUtilities/Private/USDSkeletalDataConversion.cpp
@@ -1516,7 +1516,7 @@ bool UsdToUnreal::ConvertSkinnedMesh(
 			LocalToCombinedMaterialIndex[LocalIndex] = NewIndex;
 		}
 	}
-	// Now that we merged all prim paths into they keys of CombinedMaterialSlotsToIndex, let's copy them back into
+	// Now that we merged all prim paths into they keys of SlotToCombinedMaterialIndex, let's copy them back into
 	// our output
 	for (UsdUtils::FUsdPrimMaterialSlot& Slot : MaterialAssignments.Slots)
 	{
@@ -1806,6 +1806,7 @@ bool UsdToUnreal::ConvertSkinnedMesh(
 
 		// The SkelMeshImportData now requires that the Materials array has number of entries that matches
 		// the max material index
+		// TODO: This really doesn't need to be done *per polygon*... It could be done outside the current for loop
 		SkelMeshImportData.Materials.SetNum(SkelMeshImportData.MaxMaterialIndex + 1);
 		for (int32 Index = 0; Index < SkelMeshImportData.Materials.Num(); ++Index)
 		{
diff --git a/Engine/Plugins/Runtime/USDCore/Source/USDUtilities/Public/Objects/USDSchemaTranslator.h b/Engine/Plugins/Runtime/USDCore/Source/USDUtilities/Public/Objects/USDSchemaTranslator.h
--- a/Engine/Plugins/Runtime/USDCore/Source/USDUtilities/Public/Objects/USDSchemaTranslator.h
+++ b/Engine/Plugins/Runtime/USDCore/Source/USDUtilities/Public/Objects/USDSchemaTranslator.h
@@ -269,8 +269,8 @@ struct USDUTILITIES_API FUsdSchemaTranslationContext : public TSharedFromThis<FU
 	bool bUsePrimKindsForCollapsing = true;
 
 	/**
-	 * If enabled, when multiple mesh prims are collapsed into a single static mesh, identical material slots are merged into one slot.
-	 * Otherwise, material slots are simply appended to the list.
+	 * Identical material slots will be combined into a single slot if this is enabled. This is only performed in the context
+	 * of mesh collapsing, or when parsing LOD variant sets (see bInterpretLODs)
 	 */
 	bool bMergeIdenticalMaterialSlots = true;
 
