USD: Fix duplicate LOCTEXT keys.

diff --git a/Engine/Plugins/Importers/USDImporter/Source/USDSchemas/Private/USDGeomMeshTranslator.cpp b/Engine/Plugins/Importers/USDImporter/Source/USDSchemas/Private/USDGeomMeshTranslator.cpp
--- a/Engine/Plugins/Importers/USDImporter/Source/USDSchemas/Private/USDGeomMeshTranslator.cpp
+++ b/Engine/Plugins/Importers/USDImporter/Source/USDSchemas/Private/USDGeomMeshTranslator.cpp
@@ -187,7 +187,7 @@ namespace UsdGeomMeshTranslatorImpl
 
 #if !WITH_EDITOR
 		USD_LOG_USERWARNING(FText::Format(
-			LOCTEXT("NaniteSlotWarning", "Not enabling Nanite for mesh generated for prim '{0}' as we can't setup Nanite during runtime"),
+			LOCTEXT("NaniteRuntimeWarning", "Not enabling Nanite for mesh generated for prim '{0}' as we can't setup Nanite during runtime"),
 			FText::FromString(PrimPath.GetString())
 		));
 		return false;
diff --git a/Engine/Plugins/Importers/USDImporter/Source/USDSchemas/Private/USDVolVolumeTranslator.cpp b/Engine/Plugins/Importers/USDImporter/Source/USDSchemas/Private/USDVolVolumeTranslator.cpp
--- a/Engine/Plugins/Importers/USDImporter/Source/USDSchemas/Private/USDVolVolumeTranslator.cpp
+++ b/Engine/Plugins/Importers/USDImporter/Source/USDSchemas/Private/USDVolVolumeTranslator.cpp
@@ -151,7 +151,7 @@ namespace UE::UsdVolVolumeTranslator::Private
 						{
 							USD_LOG_USERWARNING(FText::Format(
 								LOCTEXT(
-									"InvalidComponentName",
+									"InvalidComponentNameIndex",
 									"Invalid component name '{0}' for grid '{1}' in VDB file '{2}', as that particular grid only has {3} components."
 								),
 								FText::FromString(DesiredComponent),
