USD: Fix level not being updated on USD asset reimports.

diff --git a/Engine/Plugins/Importers/USDImporter/Source/USDStageImporter/Private/USDStageImporter.cpp b/Engine/Plugins/Importers/USDImporter/Source/USDStageImporter/Private/USDStageImporter.cpp
--- a/Engine/Plugins/Importers/USDImporter/Source/USDStageImporter/Private/USDStageImporter.cpp
+++ b/Engine/Plugins/Importers/USDImporter/Source/USDStageImporter/Private/USDStageImporter.cpp
@@ -1289,10 +1289,13 @@ namespace UE::USDStageImporter::Private
 
 		// Remap references held by assets that were moved directly to the destination package, and won't be in ProcessedObjectsToRemap
 		TSet<UObject*> Referencers = PublishedObjects;
-		if (AActor* SceneActor = ImportContext.SceneActor)
+		if (GEditor)
 		{
 			// Remap references to spawned actors
-			Referencers.Add(ImportContext.SceneActor->GetWorld()->GetCurrentLevel());
+			if (UWorld* World = GEditor->GetEditorWorldContext().World())
+			{
+				Referencers.Add(World);
+			}
 		}
 		for (const TPair<UObject*, UObject*>& Pair : ProcessedObjectsToRemap)
 		{
