USD: Fix Prims to Import tree view no longer actually filtering prims.

diff --git a/Engine/Plugins/Importers/USDImporter/Source/USDStageImporter/Private/USDStageImporter.cpp b/Engine/Plugins/Importers/USDImporter/Source/USDStageImporter/Private/USDStageImporter.cpp
--- a/Engine/Plugins/Importers/USDImporter/Source/USDStageImporter/Private/USDStageImporter.cpp
+++ b/Engine/Plugins/Importers/USDImporter/Source/USDStageImporter/Private/USDStageImporter.cpp
@@ -144,6 +144,10 @@ namespace UE::USDStageImporter::Private
 			ImportContext.Stage = UE::FUsdStage();
 			USD_LOG_USERERROR(FText::Format(LOCTEXT("CouldNotImportUSDFile", "Could not import USD file {0}"), FText::FromString(FilePath)));
 		}
+		else
+		{
+			ImportContext.Stage = Stage;
+		}
 	}
 
 	FString FindValidPackagePath(const FString& InPackagePath)
