USD: Fix issues with RootPrimName option when exporting levels and level sequences to USD.

diff --git a/Engine/Plugins/Importers/USDImporter/Source/USDExporter/Private/LevelExporterUSDOptionsCustomization.cpp b/Engine/Plugins/Importers/USDImporter/Source/USDExporter/Private/LevelExporterUSDOptionsCustomization.cpp
--- a/Engine/Plugins/Importers/USDImporter/Source/USDExporter/Private/LevelExporterUSDOptionsCustomization.cpp
+++ b/Engine/Plugins/Importers/USDImporter/Source/USDExporter/Private/LevelExporterUSDOptionsCustomization.cpp
@@ -5,6 +5,7 @@
 #include "LevelExporterUSDOptions.h"
 #include "LevelSequenceExporterUSDOptions.h"
 #include "USDClassesModule.h"
+#include "USDConversionUtils.h"
 
 #include "Brushes/SlateColorBrush.h"
 #include "CoreMinimal.h"
@@ -176,6 +177,7 @@ void FLevelExporterUSDOptionsCustomization::CustomizeDetails(IDetailLayoutBuilde
 	FName ExportSublayersPropName;
 	FName TexturesDirPropName;
 	FName SublayersCategoryName;
+	FName RootPrimNamePropName;
 	TAttribute<bool> SublayersEditCondition;
 
 	if (ULevelExporterUSDOptions* Options = Cast<ULevelExporterUSDOptions>(SelectedObjects[0].Get()))
@@ -191,6 +193,7 @@ void FLevelExporterUSDOptionsCustomization::CustomizeDetails(IDetailLayoutBuilde
 		ExportSublayersPropName = TEXT("Inner.bExportSublayers");
 		TexturesDirPropName = TEXT("Inner.AssetOptions.MaterialBakingOptions.TexturesDir");
 		SublayersCategoryName = TEXT("Sublayers");
+		RootPrimNamePropName = TEXT("Inner.RootPrimName");
 
 		SublayersEditCondition = true;
 	}
@@ -203,6 +206,7 @@ void FLevelExporterUSDOptionsCustomization::CustomizeDetails(IDetailLayoutBuilde
 		ExportSublayersPropName = TEXT("LevelExportOptions.bExportSublayers");
 		TexturesDirPropName = TEXT("LevelExportOptions.AssetOptions.MaterialBakingOptions.TexturesDir");
 		SublayersCategoryName = TEXT("Level Export");
+		RootPrimNamePropName = TEXT("LevelExportOptions.RootPrimName");
 
 		// Refresh the dialog whenever we pick a new world to export, so that we can show this world's sublevels in the sublevel picker
 		FSimpleDelegate RebuildDisplayDelegate = FSimpleDelegate::CreateLambda(
@@ -236,6 +240,28 @@ void FLevelExporterUSDOptionsCustomization::CustomizeDetails(IDetailLayoutBuilde
 
 	TSharedRef<IPropertyHandle> LevelFilterProp = DetailLayoutBuilder.GetProperty(LevelFilterPropName);
 	TSharedRef<IPropertyHandle> ExportSublayersProp = DetailLayoutBuilder.GetProperty(ExportSublayersPropName);
+	TSharedRef<IPropertyHandle> RootPrimNameProp = DetailLayoutBuilder.GetProperty(RootPrimNamePropName);
+
+	// Auto-sanitize Root Prim Name
+	TWeakObjectPtr<UObject> OptionsObject = SelectedObjects[0];
+	FSimpleDelegate SanitizeRootPrimName = FSimpleDelegate::CreateLambda(
+		[&DetailLayoutBuilder, OptionsObject]()
+		{
+			if (ULevelExporterUSDOptions* Options = Cast<ULevelExporterUSDOptions>(OptionsObject.Get()))
+			{
+				Options->Inner.RootPrimName = UsdUtils::SanitizeUsdIdentifier(*Options->Inner.RootPrimName);
+			}
+			else if (ULevelSequenceExporterUsdOptions* LevelSequenceOptions = Cast<ULevelSequenceExporterUsdOptions>(OptionsObject.Get()))
+			{
+				LevelSequenceOptions->LevelExportOptions.RootPrimName = UsdUtils::SanitizeUsdIdentifier(
+					*LevelSequenceOptions->LevelExportOptions.RootPrimName
+				);
+			}
+
+			DetailLayoutBuilder.ForceRefreshDetails();
+		}
+	);
+	RootPrimNameProp->SetOnPropertyValueChanged(SanitizeRootPrimName);
 
 	// Touch these properties and categories to enforce this ordering
 	DetailLayoutBuilder.EditCategory(TEXT("Stage options"));
diff --git a/Engine/Plugins/Importers/USDImporter/Source/USDExporter/Private/LevelSequenceExporterUSD.cpp b/Engine/Plugins/Importers/USDImporter/Source/USDExporter/Private/LevelSequenceExporterUSD.cpp
--- a/Engine/Plugins/Importers/USDImporter/Source/USDExporter/Private/LevelSequenceExporterUSD.cpp
+++ b/Engine/Plugins/Importers/USDImporter/Source/USDExporter/Private/LevelSequenceExporterUSD.cpp
@@ -1634,7 +1634,7 @@ namespace UE::LevelSequenceExporterUSD::Private
 		// the same DisplayRate as the original exported sequence
 		UsdStage.SetFramesPerSecond(MovieScene->GetDisplayRate().AsDecimal());
 
-		const FString RootPrimPath = TEXT("/") + Context.ExportOptions->LevelExportOptions.RootPrimName;
+		const FString RootPrimPath = TEXT("/") + UsdUtils::SanitizeUsdIdentifier(*Context.ExportOptions->LevelExportOptions.RootPrimName);
 		UE::FUsdPrim RootPrim = UsdStage.OverridePrim(UE::FSdfPath(*RootPrimPath));
 		if (!RootPrim)
 		{
diff --git a/Engine/Plugins/Runtime/USDCore/Source/USDUtilities/Private/USDConversionUtils.cpp b/Engine/Plugins/Runtime/USDCore/Source/USDUtilities/Private/USDConversionUtils.cpp
--- a/Engine/Plugins/Runtime/USDCore/Source/USDUtilities/Private/USDConversionUtils.cpp
+++ b/Engine/Plugins/Runtime/USDCore/Source/USDUtilities/Private/USDConversionUtils.cpp
@@ -802,11 +802,11 @@ FString UsdUtils::GetPrimPathForObject(
 
 		if (USceneComponent* ParentComp = Component->GetAttachParent())
 		{
-			FoundParentPath = GetPrimPathForObject(ParentComp, TEXT(""), bUseActorFolders);
+			FoundParentPath = GetPrimPathForObject(ParentComp, TEXT(""), bUseActorFolders, RootPrimName);
 		}
 		else
 		{
-			FoundParentPath = TEXT("/") + RootPrimName;
+			FoundParentPath = TEXT("/") + UsdUtils::SanitizeUsdIdentifier(*RootPrimName);
 		}
 
 		Path = FoundParentPath / Path;
