USD: Fix crash when undoing a LevelSequence change that regenerated the Time Track.

diff --git a/Engine/Plugins/Importers/USDImporter/Source/USDStage/Private/USDLevelSequenceHelper.cpp b/Engine/Plugins/Importers/USDImporter/Source/USDStage/Private/USDLevelSequenceHelper.cpp
--- a/Engine/Plugins/Importers/USDImporter/Source/USDStage/Private/USDLevelSequenceHelper.cpp
+++ b/Engine/Plugins/Importers/USDImporter/Source/USDStage/Private/USDLevelSequenceHelper.cpp
@@ -1936,6 +1936,7 @@ void FUsdLevelSequenceHelperImpl::CreateTimeTrack(const FLayerTimeInfo& Info)
 	);
 	if (TimeTrack)
 	{
+		TimeTrack->Modify();
 		TimeTrack->RemoveAllAnimationData();
 	}
 	else
@@ -2764,6 +2765,7 @@ void FUsdLevelSequenceHelperImpl::AddSkeletalTracks(const UUsdPrimTwin& PrimTwin
 
 		FFrameNumber StartOffsetTick = FFrameTime::FromDecimal(LayerStartOffsetSeconds * MovieScene->GetTickResolution().AsDecimal()).RoundToFrame();
 
+		SkeletalTrack->Modify();
 		SkeletalTrack->RemoveAllAnimationData();
 
 		UMovieSceneSkeletalAnimationSection* NewSection = Cast<UMovieSceneSkeletalAnimationSection>(
@@ -2825,6 +2827,7 @@ void FUsdLevelSequenceHelperImpl::AddGeometryCacheTracks(const UUsdPrimTwin& Pri
 	if (UMovieSceneGeometryCacheTrack* GeometryCacheTrack = AddTrack<
 			UMovieSceneGeometryCacheTrack>(Prim.GetName(), PrimTwin, *ComponentToBind, *GeometryCacheSequence, bIsMuted))
 	{
+		GeometryCacheTrack->Modify();
 		GeometryCacheTrack->RemoveAllAnimationData();
 
 		double LayerStartOffsetSeconds = 0;
@@ -2901,6 +2904,7 @@ void FUsdLevelSequenceHelperImpl::AddGroomTracks(const UUsdPrimTwin& PrimTwin, c
 	if (UMovieSceneGroomCacheTrack* GroomCacheTrack = AddTrack<
 			UMovieSceneGroomCacheTrack>(Prim.GetName(), PrimTwin, *ComponentToBind, *GroomAnimationSequence, bIsMuted))
 	{
+		GroomCacheTrack->Modify();
 		GroomCacheTrack->RemoveAllAnimationData();
 
 		const FFrameNumber StartOffset;
@@ -3309,6 +3313,7 @@ void FUsdLevelSequenceHelperImpl::AddAudioTracks(const UUsdPrimTwin& PrimTwin, c
 		if (UMovieSceneAudioTrack* AudioTrack = AddTrack<
 				UMovieSceneAudioTrack>(TEXT("Audio track"), PrimTwin, *AudioComponent, *AttributeSequence, bIsMuted))
 		{
+			AudioTrack->Modify();
 			AudioTrack->RemoveAllAnimationData();
 
 			const FFrameNumber Time = StartFrameNumber;
@@ -3499,6 +3504,7 @@ TrackType* FUsdLevelSequenceHelperImpl::AddTrack(
 	TrackType* Track = MovieScene->FindTrack<TrackType>(ComponentBinding, TrackName);
 	if (Track)
 	{
+		Track->Modify();
 		Track->RemoveAllAnimationData();
 	}
 	else
