USD: Remove spammy unnecessary log message with wrong placeholders.

diff --git a/Engine/Plugins/Importers/USDImporter/Source/USDSchemas/Private/USDShadeMaterialTranslator.cpp b/Engine/Plugins/Importers/USDImporter/Source/USDSchemas/Private/USDShadeMaterialTranslator.cpp
--- a/Engine/Plugins/Importers/USDImporter/Source/USDSchemas/Private/USDShadeMaterialTranslator.cpp
+++ b/Engine/Plugins/Importers/USDImporter/Source/USDSchemas/Private/USDShadeMaterialTranslator.cpp
@@ -303,14 +303,6 @@ void FUsdShadeMaterialTranslator::CreateAssets()
 		TOptional<FString> UnrealMaterial = UsdUtils::GetUnrealSurfaceOutput(ShadeMaterial.GetPrim());
 		if (UnrealMaterial.IsSet())
 		{
-			USD_LOG_INFO(
-				TEXT(
-					"Skipping generation of assets for material prim '{0}' as all prims that bind this material will use its referenced Unreal material '{1}' instead."
-				),
-				*PrimPath.GetString(),
-				*UnrealMaterial.GetValue()
-			);
-
 			UObject* Object = FSoftObjectPath(UnrealMaterial.GetValue()).TryLoad();
 			if (!Object)
 			{
