USD: Fix crash when using Actions -> Import while the Stage's LevelSequence is opened.

diff --git a/Engine/Plugins/Importers/USDImporter/Source/USDStage/Private/USDStageActor.cpp b/Engine/Plugins/Importers/USDImporter/Source/USDStage/Private/USDStageActor.cpp
--- a/Engine/Plugins/Importers/USDImporter/Source/USDStage/Private/USDStageActor.cpp
+++ b/Engine/Plugins/Importers/USDImporter/Source/USDStage/Private/USDStageActor.cpp
@@ -3427,19 +3427,22 @@ void AUsdStageActor::UnloadUsdStage()
 			// close from the callstack of its LevelSequence being evaluated (for example, imagine we had a track to
 			// set StateState to Closed: Sequencer evaluates the track -> Calls SetStageState -> Ends up here -> We try
 			// destroying the Sequencer -> Crash).
-			// Note that an AsyncTask has it run on the same tick and so within the same frame, while FTSTicker would end
-			// up on the next tick, showing a frame of the broken sequence bindings on the UI, which doesn't look great...
+			//
+			// Originally we did this in an AsyncTask so that it would hopefully run later in the same frame and avoid
+			// drawing the Sequencer with broken bindings. Unfortunately the task graph system can fast-forward and run these
+			// async tasks at any time, which may be e.g. before other Sequencer callbacks intended to run on this same frame.
+			// That could lead to a crash, as the LevelSequence would have been closed/destroyed at that point
 			TWeakObjectPtr<ULevelSequence> LevelSequencePtr = LevelSequence;
-			AsyncTask(
-				ENamedThreads::GameThread,
-				[LevelSequencePtr]()
+			FTSTicker::GetCoreTicker().AddTicker(FTickerDelegate::CreateLambda(
+				[LevelSequencePtr](float Time)
 				{
 					if (ULevelSequence* ValidSequence = LevelSequencePtr.Get())
 					{
 						GEditor->GetEditorSubsystem<UAssetEditorSubsystem>()->CloseAllEditorsForAsset(ValidSequence);
 					}
+					return false;
 				}
-			);
+			));
 		}
 #endif	  // WITH_EDITOR
 
