USD: Fix broken SkelAnimation prims when exporting LevelSequences with skeletal tracks to USD.

diff --git a/Engine/Plugins/Importers/USDImporter/Source/USDExporter/Private/LevelSequenceExporterUSD.cpp b/Engine/Plugins/Importers/USDImporter/Source/USDExporter/Private/LevelSequenceExporterUSD.cpp
--- a/Engine/Plugins/Importers/USDImporter/Source/USDExporter/Private/LevelSequenceExporterUSD.cpp
+++ b/Engine/Plugins/Importers/USDImporter/Source/USDExporter/Private/LevelSequenceExporterUSD.cpp
@@ -1460,8 +1460,12 @@ namespace UE::LevelSequenceExporterUSD::Private
 				RecursiveSetOver(Child);
 			}
 
-			Prim.ClearTypeName();
-			Prim.SetSpecifier(ESdfSpecifier::Over);
+			// Retain SkelAnimations though: Those are actually authored directly on the LevelSequence layer
+			if (!Prim.IsA("SkelAnimation"))
+			{
+				Prim.ClearTypeName();
+				Prim.SetSpecifier(ESdfSpecifier::Over);
+			}
 		};
 		FUsdPrim Root = UsdStage.GetPseudoRoot();
 		for (FUsdPrim& TopLevelPrim : Root.GetChildren())
