USD: Rename lambda parameter preventing it from shadowing member variable 'Time'.

diff --git a/Engine/Plugins/Importers/USDImporter/Source/USDStage/Private/USDStageActor.cpp b/Engine/Plugins/Importers/USDImporter/Source/USDStage/Private/USDStageActor.cpp
--- a/Engine/Plugins/Importers/USDImporter/Source/USDStage/Private/USDStageActor.cpp
+++ b/Engine/Plugins/Importers/USDImporter/Source/USDStage/Private/USDStageActor.cpp
@@ -3434,7 +3434,7 @@ void AUsdStageActor::UnloadUsdStage()
 			// That could lead to a crash, as the LevelSequence would have been closed/destroyed at that point
 			TWeakObjectPtr<ULevelSequence> LevelSequencePtr = LevelSequence;
 			FTSTicker::GetCoreTicker().AddTicker(FTickerDelegate::CreateLambda(
-				[LevelSequencePtr](float Time)
+				[LevelSequencePtr](float DelayTime)
 				{
 					if (ULevelSequence* ValidSequence = LevelSequencePtr.Get())
 					{
