USD: Fix crash when interacting with Sequencer in invalid state.

diff --git a/Engine/Plugins/Importers/USDImporter/Source/USDStage/Private/USDStageActor.cpp b/Engine/Plugins/Importers/USDImporter/Source/USDStage/Private/USDStageActor.cpp
--- a/Engine/Plugins/Importers/USDImporter/Source/USDStage/Private/USDStageActor.cpp
+++ b/Engine/Plugins/Importers/USDImporter/Source/USDStage/Private/USDStageActor.cpp
@@ -1857,9 +1857,10 @@ void AUsdStageActor::HandleAccumulatedNotices()
 
 	// Recreate our LevelSequences before we regenerate components and want to add bindings/tracks
 	// back onto it
+	bool bSequenceWasOpened = false;
 	if (bNeedsAnimationReload)
 	{
-		RegenerateLevelSequence();
+		bSequenceWasOpened = RegenerateLevelSequence();
 	}
 
 	// Reset our translated prototypes only here, and reuse them for all individual changes. This because some types of operations
@@ -1956,6 +1957,11 @@ void AUsdStageActor::HandleAccumulatedNotices()
 		FUsdStageActorImpl::RepairExternalSequencerBindings();
 	}
 
+	if (bSequenceWasOpened)
+	{
+		OpenLevelSequence();
+	}
+
 #if WITH_EDITOR
 	if (GIsEditor && GEditor && !IsGarbageCollecting())	   // Make sure we're not in standalone either
 	{
@@ -3315,7 +3321,7 @@ void AUsdStageActor::LoadUsdStage(bool bOpenIfNeeded)
 		UsdInfoCache->GetInner().ResetTranslatedPrototypes();
 	}
 
-	RegenerateLevelSequence();
+	const bool bSequenceWasOpened = RegenerateLevelSequence();
 
 	TSharedRef<FUsdSchemaTranslationContext> TranslationContext = FUsdStageActorImpl::CreateUsdSchemaTranslationContext(this, RootTwin->PrimPath);
 
@@ -3371,6 +3377,11 @@ void AUsdStageActor::LoadUsdStage(bool bOpenIfNeeded)
 
 	OnStageLoaded.Broadcast();
 
+	if (bSequenceWasOpened)
+	{
+		OpenLevelSequence();
+	}
+
 	// Log time spent to load the stage
 	double ElapsedSeconds = FPlatformTime::ToSeconds64(FPlatformTime::Cycles64() - StartTime);
 
@@ -3567,7 +3578,7 @@ void AUsdStageActor::ReloadAnimations()
 }
 PRAGMA_ENABLE_DEPRECATION_WARNINGS
 
-void AUsdStageActor::RegenerateLevelSequence()
+bool AUsdStageActor::RegenerateLevelSequence()
 {
 	TRACE_CPUPROFILER_EVENT_SCOPE(AUsdStageActor::RegenerateLevelSequence);
 
@@ -3575,7 +3586,7 @@ void AUsdStageActor::RegenerateLevelSequence()
 	// want to listen to it and have valid level sequences
 	if (IsTemplate())
 	{
-		return;
+		return false;
 	}
 
 	// If we're using some property editor that can trigger a stage reload (like the Nanite threshold spinbox),
@@ -3589,14 +3600,14 @@ void AUsdStageActor::RegenerateLevelSequence()
 	static bool bIsReentrant = false;
 	if (bIsReentrant)
 	{
-		return;
+		return false;
 	}
 	TGuardValue<bool> ReentrantGuard(bIsReentrant, true);
 
 	const UE::FUsdStage& CurrentStage = GetOrOpenUsdStage();
 	if (!CurrentStage)
 	{
-		return;
+		return false;
 	}
 
 	// We will update our levelsequence objects (tracks, moviescene, sections, etc.) due to these changes.
@@ -3605,9 +3616,14 @@ void AUsdStageActor::RegenerateLevelSequence()
 
 	// Don't check for full authority here because even if we can't write back to the stage (i.e. during PIE) we still
 	// want to listen to it and have valid level sequences
+	bool bSequencerWasOpened = false;
 #if WITH_EDITOR
 	UAssetEditorSubsystem* AssetEditorSubsystem = GEditor ? GEditor->GetEditorSubsystem<UAssetEditorSubsystem>() : nullptr;
-	const bool bSequencerWasOpened = AssetEditorSubsystem && AssetEditorSubsystem->FindEditorsForAssetAndSubObjects(LevelSequence).Num() > 0;
+	bSequencerWasOpened = AssetEditorSubsystem && AssetEditorSubsystem->FindEditorsForAssetAndSubObjects(LevelSequence).Num() > 0;
+	if (LevelSequence && AssetEditorSubsystem)
+	{
+		AssetEditorSubsystem->CloseAllEditorsForAsset(LevelSequence);
+	}
 #endif	  // WITH_EDITOR
 
 	// We need to guarantee we'll record our change of LevelSequence into the transaction, as Init() will create a new one
@@ -3617,14 +3633,7 @@ void AUsdStageActor::RegenerateLevelSequence()
 	LevelSequence = LevelSequenceHelper.Init(CurrentStage);
 	LevelSequenceHelper.BindToUsdStageActor(this);
 
-#if WITH_EDITOR
-	if (bSequencerWasOpened && LevelSequence && AssetEditorSubsystem)
-	{
-		// Open the LevelSequence editor via the UAssetEditorSubsystem, otherwise calls to CloseAllEditorsForAsset
-		// won't be able to close it
-		AssetEditorSubsystem->OpenEditorForAsset(LevelSequence);
-	}
-#endif	  // WITH_EDITOR
+	return bSequencerWasOpened;
 }
 
 void AUsdStageActor::RepopulateLevelSequence()
@@ -3645,6 +3654,22 @@ void AUsdStageActor::RepopulateLevelSequence()
 	UpdatePrim(UE::FSdfPath::AbsoluteRootPath(), bIsResync, *TranslationContext);
 }
 
+void AUsdStageActor::OpenLevelSequence()
+{
+#if WITH_EDITOR
+	if (!LevelSequence)
+	{
+		return;
+	}
+
+	UAssetEditorSubsystem* AssetEditorSubsystem = GEditor ? GEditor->GetEditorSubsystem<UAssetEditorSubsystem>() : nullptr;
+	if (AssetEditorSubsystem)
+	{
+		AssetEditorSubsystem->OpenEditorForAsset(LevelSequence);
+	}
+#endif	  // WITH_EDITOR
+}
+
 PRAGMA_DISABLE_DEPRECATION_WARNINGS
 TSharedPtr<FUsdInfoCache> AUsdStageActor::GetInfoCache()
 {
@@ -3758,8 +3783,12 @@ void AUsdStageActor::PostTransacted(const FTransactionObjectEvent& TransactionEv
 				UsdInfoCache->GetInner().RebuildCacheForSubtrees({UE::FSdfPath::AbsoluteRootPath()}, TranslationContext.Get());
 			}
 
-			RegenerateLevelSequence();
+			const bool bSequenceWasOpened = RegenerateLevelSequence();
 			RepopulateLevelSequence();
+			if (bSequenceWasOpened)
+			{
+				OpenLevelSequence();
+			}
 		}
 		else if (ChangedProperties.Contains(GET_MEMBER_NAME_CHECKED(AUsdStageActor, StageState)))
 		{
@@ -3982,8 +4011,12 @@ void AUsdStageActor::PostDuplicate(bool bDuplicateForPIE)
 		// We always want our own LevelSequence though, otherwise we could end up with some strange behavior, like:
 		// PIE -> Change LevelSequence -> PIE actor is inside of a FScopedBlockNoticeListening from HandleTrackChange,
 		// but the actor back in the editor is not -> Editor actor writes changes back to the stage anyway
-		RegenerateLevelSequence();
+		const bool bSequenceWasOpened = RegenerateLevelSequence();
 		RepopulateLevelSequence();
+		if (bSequenceWasOpened)
+		{
+			OpenLevelSequence();
+		}
 	}
 	else
 	{
diff --git a/Engine/Plugins/Importers/USDImporter/Source/USDStage/Public/USDStageActor.h b/Engine/Plugins/Importers/USDImporter/Source/USDStage/Public/USDStageActor.h
--- a/Engine/Plugins/Importers/USDImporter/Source/USDStage/Public/USDStageActor.h
+++ b/Engine/Plugins/Importers/USDImporter/Source/USDStage/Public/USDStageActor.h
@@ -357,8 +357,9 @@ public:
 	UE_DEPRECATED(5.5, "This function has been renamed into 'RegenerateLevelSequence', which better describes what it does")
 	USDSTAGE_API void ReloadAnimations();
 
-	// Creates new ULevelSequence assets for the root layer and sublayers
-	USDSTAGE_API void RegenerateLevelSequence();
+	// Creates new ULevelSequence assets for the root layer and sublayers.
+	// If we already had a previous LevelSequence opened, we will close its Sequencer and return true
+	USDSTAGE_API bool RegenerateLevelSequence();
 
 	/**
 	 * Fills in the current LevelSequences with animations for the prims of the opened stage.
@@ -368,6 +369,8 @@ public:
 	 */
 	USDSTAGE_API void RepopulateLevelSequence();
 
+	USDSTAGE_API void OpenLevelSequence();
+
 	PRAGMA_DISABLE_DEPRECATION_WARNINGS
 	UE_DEPRECATED(5.5, "The FUsdInfoCache struct is deprecated in favor of the UUsdInfoCache and UUsdPrimLinkCache, referenced via UPROPERTYs")
 	USDSTAGE_API TSharedPtr<FUsdInfoCache> GetInfoCache();
