USD: Fix not being able to reimport USD Interchange assets after cl 41486351

diff --git a/Engine/Plugins/Importers/USDImporter/Source/USDStageImporter/Private/USDStageAssetImportFactory.cpp b/Engine/Plugins/Importers/USDImporter/Source/USDStageImporter/Private/USDStageAssetImportFactory.cpp
--- a/Engine/Plugins/Importers/USDImporter/Source/USDStageImporter/Private/USDStageAssetImportFactory.cpp
+++ b/Engine/Plugins/Importers/USDImporter/Source/USDStageImporter/Private/USDStageAssetImportFactory.cpp
@@ -27,10 +27,9 @@ UUsdStageAssetImportFactory::UUsdStageAssetImportFactory(const FObjectInitialize
 	bEditAfterNew = true;
 	SupportedClass = nullptr;
 
-	// Its ok if we intercept most calls as there aren't other USD importers, and
-	// for reimport we can definitely tell that we should be handling an asset, as we
-	// use a custom asset import data
-	ImportPriority += 100;
+	// Lower our priority to let Interchange/others try and import/reimport first.
+	// We use a custom asset import data anyway, so the order likely doesn't matter.
+	ImportPriority -= 10000;
 
 	bEditorImport = true;
 	bText = false;
@@ -121,7 +120,30 @@ void UUsdStageAssetImportFactory::CleanUp()
 
 bool UUsdStageAssetImportFactory::CanReimport(UObject* Obj, TArray<FString>& OutFilenames)
 {
-	if (UAssetImportData* ImportData = UsdUnreal::ObjectUtils::GetBaseAssetImportData(Obj))
+	// It seems that every import attempt creates a new factory instance, which is just left behind for GC to collect.
+	// Since we may emit a warning from here, that means that if we do 5 imports and then want to emit the warning,
+	// all of our 5 factory instances will do it (+ the CDO), and we'll get spammed...
+	// Conversely, if we force GC before import, it will clean up those leftover factories and emit the warning only once.
+	// We can avoid these issues by only checking for reimport via the Factory's CDO, which is guaranteed to be
+	// registered exactly once.
+	if (!IsTemplate())
+	{
+		return false;
+	}
+
+	// Try looking for the exact UUsdAssetImportData the USDImporter emits
+	bool bFromOtherImporter = false;
+	UAssetImportData* ImportData = UsdUnreal::ObjectUtils::GetAssetImportData(Obj);
+
+	// Failing that, check if we have any other asset import data with USD file extensions, but
+	// just so that we can emit a warning for some feedback
+	if (!ImportData)
+	{
+		bFromOtherImporter = true;
+		ImportData = UsdUnreal::ObjectUtils::GetBaseAssetImportData(Obj);
+	}
+
+	if (ImportData)
 	{
 		const FString FileName = ImportData->GetFirstFilename();
 		const FString FileExtension = FPaths::GetExtension(FileName);
@@ -133,8 +155,28 @@ bool UUsdStageAssetImportFactory::CanReimport(UObject* Obj, TArray<FString>& Out
 		{
 			if (Extension == FileExtension)
 			{
-				OutFilenames.Add(ImportData->GetFirstFilename());
-				return true;
+				if (bFromOtherImporter)
+				{
+					// This came from Interchange or some other custom USD format importer, just emit a warning
+					// and return. Note that our factory has very low ImportPriority, so if we're here it's very
+					// likely nothing else can handle this asset anyway, so it's probably OK to emit a warning
+					FScopedUsdMessageLog ScopedMessageLog;
+
+					USD_LOG_USERWARNING(FText::Format(
+						LOCTEXT(
+							"ReimportErrorWrongImportData",
+							"Skipped trying to reimport asset '{0}' with the USD Importer plugin as it doesn't seem to have valid USD import data or user data! Only assets created by the USD Importer plugin can be reimported by the USD Importer plugin."
+						),
+						FText::FromName(Obj->GetFName())
+					));
+					return false;
+				}
+				else
+				{
+					// We can actually reimport this
+					OutFilenames.Add(ImportData->GetFirstFilename());
+					return true;
+				}
 			}
 		}
 	}
@@ -189,10 +231,7 @@ EReimportResult::Type UUsdStageAssetImportFactory::Reimport(UObject* Obj)
 	if (ReimportFilePath.IsEmpty() || OriginalPrimPath.IsEmpty())
 	{
 		USD_LOG_USERERROR(FText::Format(
-			LOCTEXT(
-				"ReimportErrorNoImportData",
-				"Failed to reimport asset '{0}' as it doesn't seem to have valid USD import data or user data! Only assets created by the USD Importer plugin can be reimported via the USD Importer plugin."
-			),
+			LOCTEXT("ReimportErrorNoImportData", "Failed to reimport asset '{0}' as it doesn't seem to have valid USD import data or user data!"),
 			FText::FromName(Obj->GetFName())
 		));
 		return EReimportResult::Failed;
