USD: Also remap references to reimported assets from within blueprint previews.

diff --git a/Engine/Plugins/Importers/USDImporter/Source/USDStageImporter/Private/USDStageImporter.cpp b/Engine/Plugins/Importers/USDImporter/Source/USDStageImporter/Private/USDStageImporter.cpp
--- a/Engine/Plugins/Importers/USDImporter/Source/USDStageImporter/Private/USDStageImporter.cpp
+++ b/Engine/Plugins/Importers/USDImporter/Source/USDStageImporter/Private/USDStageImporter.cpp
@@ -1295,10 +1295,18 @@ namespace UE::USDStageImporter::Private
 		TSet<UObject*> Referencers = PublishedObjects;
 		if (GEditor)
 		{
-			// Remap references to spawned actors
-			if (UWorld* World = GEditor->GetEditorWorldContext().World())
+			// Remap references to spawned assets/actors from all worlds.
+			// Checking EditorPreview worlds also lets us remap references from components inside Blueprint viewports
+			const TIndirectArray<FWorldContext> WorldContexts = GEngine->GetWorldContexts();
+			for (const FWorldContext& Context : WorldContexts)
 			{
-				Referencers.Add(World);
+				if (Context.WorldType == EWorldType::Editor || Context.WorldType == EWorldType::EditorPreview)
+				{
+					if (UWorld* World = Context.World())
+					{
+						Referencers.Add(World);
+					}
+				}
 			}
 		}
 		for (const TPair<UObject*, UObject*>& Pair : ProcessedObjectsToRemap)
@@ -1311,7 +1319,8 @@ namespace UE::USDStageImporter::Private
 		// Note we don't care if transient assets reference each other, as we'll delete them all at once anyway
 		for (UObject* Referencer : Referencers)
 		{
-			if (!Referencer || Referencer->GetOutermost() == GetTransientPackage())
+			// Allow remapping references from transient worlds in particular because that's where the EditorPreview worlds live
+			if (!Referencer || (!Referencer->IsA<UWorld>() && Referencer->GetOutermost() == GetTransientPackage()))
 			{
 				continue;
 			}
