USD: Don't try to regenerate collision meshes on existing static meshes, but add the collision settings to the hash instead.

diff --git a/Engine/Plugins/Importers/USDImporter/Source/USDSchemas/Private/USDGeomMeshTranslator.cpp b/Engine/Plugins/Importers/USDImporter/Source/USDSchemas/Private/USDGeomMeshTranslator.cpp
--- a/Engine/Plugins/Importers/USDImporter/Source/USDSchemas/Private/USDGeomMeshTranslator.cpp
+++ b/Engine/Plugins/Importers/USDImporter/Source/USDSchemas/Private/USDGeomMeshTranslator.cpp
@@ -538,6 +538,37 @@ namespace UsdGeomMeshTranslatorImpl
 		}
 	}
 
+	void GetCollisionSettings(
+		const pxr::UsdPrim& InUsdPrim,
+		EUsdCollisionType InFallbackCollisionType,
+		bool& OutEnableCollision,
+		EUsdCollisionType& OutApproximationType
+	)
+	{
+		bool bPrimEnablesCollision = false;
+		pxr::UsdPhysicsCollisionAPI CollisionAPI{InUsdPrim};
+		if (CollisionAPI)
+		{
+			if (pxr::UsdAttribute CollisionAttr = CollisionAPI.GetCollisionEnabledAttr())
+			{
+				CollisionAttr.Get(&bPrimEnablesCollision);
+			}
+		}
+
+		static IConsoleVariable* CVar = IConsoleManager::Get().FindConsoleVariable(TEXT("USD.EnableCollision"));
+		const bool bCvarEnablesCollision = CVar && CVar->GetBool();
+
+		OutEnableCollision = bCvarEnablesCollision && (bPrimEnablesCollision || InFallbackCollisionType != EUsdCollisionType::None);
+
+		if (!OutEnableCollision)
+		{
+			OutApproximationType = EUsdCollisionType::None;	   // It doesn't matter, so let's always return None for consistency (helps with hashing)
+			return;
+		}
+
+		OutApproximationType = CollisionAPI ? UsdUtils::GetCollisionApproximationType(InUsdPrim) : InFallbackCollisionType;
+	}
+
 	UStaticMesh* CreateStaticMesh(
 		const UE::FUsdPrim& Prim,
 		TArray<FMeshDescription>& LODIndexToMeshDescription,
@@ -593,6 +624,16 @@ namespace UsdGeomMeshTranslatorImpl
 			SHA1.Update(reinterpret_cast<const uint8*>(&MeshNormalRepairThreshold), sizeof(MeshNormalRepairThreshold));
 			SHA1.Update(reinterpret_cast<const uint8*>(&SkipTangentComputation), sizeof(SkipTangentComputation));
 
+			// We don't want to reuse static meshes if they need different collision shapes
+			{
+				bool bIsCollisionEnabled = false;
+				EUsdCollisionType Approximation = EUsdCollisionType::None;
+				GetCollisionSettings(Prim, Context.FallbackCollisionType, bIsCollisionEnabled, Approximation);
+
+				SHA1.Update(reinterpret_cast<const uint8*>(&bIsCollisionEnabled), sizeof(bIsCollisionEnabled));
+				SHA1.Update(reinterpret_cast<const uint8*>(&Approximation), sizeof(Approximation));
+			}
+
 			SHA1.Final();
 			SHA1.GetHash(&AllLODHash.Hash[0]);
 		}
@@ -1325,22 +1366,17 @@ namespace UE::UsdCollision::Private
 			return;
 		}
 
-		bool bIsCollisionEnabled = false;
-		pxr::UsdPhysicsCollisionAPI CollisionAPI{UsdPrim};
-		if (CollisionAPI)
+		FMeshData MeshData(StaticMesh);
+		if (!MeshData.IsValid())
 		{
-			if (pxr::UsdAttribute CollisionAttr = CollisionAPI.GetCollisionEnabledAttr())
-			{
-				CollisionAttr.Get(&bIsCollisionEnabled);
-			}
+			return;
 		}
 
-		bIsCollisionEnabled = bIsCollisionEnabled || FallbackCollisionType != EUsdCollisionType::None;
-
-		static IConsoleVariable* CVar = IConsoleManager::Get().FindConsoleVariable(TEXT("USD.EnableCollision"));
-		const bool bEnableCollision = CVar && CVar->GetBool();
+		bool bIsCollisionEnabled = false;
+		EUsdCollisionType Approximation = EUsdCollisionType::None;
+		UsdGeomMeshTranslatorImpl::GetCollisionSettings(UsdPrim, FallbackCollisionType, bIsCollisionEnabled, Approximation);
 
-		if (!bIsCollisionEnabled || !bEnableCollision)
+		if (!bIsCollisionEnabled)
 		{
 			BodySetup->Modify();
 			BodySetup->RemoveSimpleCollision();
@@ -1351,18 +1387,6 @@ namespace UE::UsdCollision::Private
 			return;
 		}
 
-		FMeshData MeshData(StaticMesh);
-		if (!MeshData.IsValid())
-		{
-			return;
-		}
-
-		EUsdCollisionType Approximation = FallbackCollisionType;
-		if (CollisionAPI)
-		{
-			Approximation = UsdUtils::GetCollisionApproximationType(UsdPrim);
-		}
-
 		BodySetup->Modify();
 		BodySetup->RemoveSimpleCollision();
 		BodySetup->bNeverNeedsCookedCollisionData = false;
@@ -1600,13 +1624,6 @@ void FBuildStaticMeshTaskChain::SetupTasks()
 					}
 #endif	  // WITH_EDITOR
 			   }
-			   else
-			   {
-				   // Setup collision on existing mesh in case the collision settings have changed
-				   // TODO: This is recomputing collision and dirtying the StaticMesh every time, even if the mesh
-				   // is reused from the AssetCache. Maybe we can do that only when collision settings actually change?
-				   UE::UsdCollision::Private::SetupSimpleCollision(GetPrim(), *StaticMesh, Context->FallbackCollisionType);
-			   }
 		   }
 
 		   // Only need to continue building the mesh if we just created it
