USD: Fix USD asset imports persistently disabling "Import Into Level" option to import actors.

diff --git a/Engine/Plugins/Importers/USDImporter/Source/USDStageImporter/Private/USDStageAssetImportFactory.cpp b/Engine/Plugins/Importers/USDImporter/Source/USDStageImporter/Private/USDStageAssetImportFactory.cpp
--- a/Engine/Plugins/Importers/USDImporter/Source/USDStageImporter/Private/USDStageAssetImportFactory.cpp
+++ b/Engine/Plugins/Importers/USDImporter/Source/USDStageImporter/Private/USDStageAssetImportFactory.cpp
@@ -69,10 +69,16 @@ UObject* UUsdStageAssetImportFactory::FactoryCreateFile(
 	// When importing from file we don't want to use any opened stage
 	ImportContext.bReadFromStageCache = false;
 
+	// We shouldn't be able to import actors when doing a manual asset import (from the content browser)
+	TOptional<UsdUtils::FScopedSuppressActorImport> SuppressActorImport;
+	if (!IsAutomatedImport())
+	{
+		SuppressActorImport.Emplace(ImportContext.ImportOptions);
+	}
+
 	const FString InitialPackagePath = InParent ? InParent->GetName() : TEXT("/Game/");
 	const bool bIsReimport = false;
-	const bool bAllowActorImport = false;
-	if (ImportContext.Init(InName.ToString(), Filename, InitialPackagePath, Flags, IsAutomatedImport(), bIsReimport, bAllowActorImport))
+	if (ImportContext.Init(InName.ToString(), Filename, InitialPackagePath, Flags, IsAutomatedImport(), bIsReimport))
 	{
 		GEditor->GetEditorSubsystem<UImportSubsystem>()->BroadcastAssetPreImport(this, InClass, InParent, InName, Parms);
 
@@ -247,9 +253,15 @@ EReimportResult::Type UUsdStageAssetImportFactory::Reimport(UObject* Obj)
 
 	ImportContext.bReadFromStageCache = false;
 
+	// We shouldn't be able to import actors when doing a manual asset reimport (from the content browser)
+	TOptional<UsdUtils::FScopedSuppressActorImport> SuppressActorImport;
+	if (!IsAutomatedImport())
+	{
+		SuppressActorImport.Emplace(ImportContext.ImportOptions);
+	}
+
 	const bool bIsReimport = true;
-	const bool bAllowActorImport = false;
-	if (!ImportContext.Init(Obj->GetName(), ReimportFilePath, Obj->GetName(), Obj->GetFlags(), IsAutomatedImport(), bIsReimport, bAllowActorImport))
+	if (!ImportContext.Init(Obj->GetName(), ReimportFilePath, Obj->GetName(), Obj->GetFlags(), IsAutomatedImport(), bIsReimport))
 	{
 		USD_LOG_USERERROR(FText::Format(
 			LOCTEXT("ReimportErrorNoContext", "Failed to initialize re-import context for asset '{0}'!"),
diff --git a/Engine/Plugins/Importers/USDImporter/Source/USDStageImporter/Private/USDStageImportContext.cpp b/Engine/Plugins/Importers/USDImporter/Source/USDStageImporter/Private/USDStageImportContext.cpp
--- a/Engine/Plugins/Importers/USDImporter/Source/USDStageImporter/Private/USDStageImportContext.cpp
+++ b/Engine/Plugins/Importers/USDImporter/Source/USDStageImporter/Private/USDStageImportContext.cpp
@@ -67,8 +67,6 @@ bool FUsdStageImportContext::Init(
 			PackagePath = FString::Printf(TEXT("%s/%s/"), *PickContentPathDlg->GetPath().ToString(), *InName);
 		}
 
-		ImportOptions->EnableActorImport(bAllowActorImport);
-
 		// Show dialog for import options
 		bool bProceedWithImport = SUsdOptionsWindow::ShowImportOptions(*ImportOptions, &Stage);
 		if (!bProceedWithImport)
diff --git a/Engine/Plugins/Importers/USDImporter/Source/USDStageImporter/Private/USDStageImportFactory.cpp b/Engine/Plugins/Importers/USDImporter/Source/USDStageImporter/Private/USDStageImportFactory.cpp
--- a/Engine/Plugins/Importers/USDImporter/Source/USDStageImporter/Private/USDStageImportFactory.cpp
+++ b/Engine/Plugins/Importers/USDImporter/Source/USDStageImporter/Private/USDStageImportFactory.cpp
@@ -53,8 +53,7 @@ UObject* UUsdStageImportFactory::FactoryCreateFile(
 #if USE_USD_SDK
 	const FString InitialPackagePath = InParent ? InParent->GetName() : TEXT("/Game/");
 	const bool bIsReimport = false;
-	const bool bAllowActorImport = true;
-	if (ImportContext.Init(InName.ToString(), Filename, InitialPackagePath, Flags, IsAutomatedImport(), bIsReimport, bAllowActorImport))
+	if (ImportContext.Init(InName.ToString(), Filename, InitialPackagePath, Flags, IsAutomatedImport(), bIsReimport))
 	{
 		GEditor->GetEditorSubsystem<UImportSubsystem>()->BroadcastAssetPreImport(this, InClass, InParent, InName, Parms);
 
diff --git a/Engine/Plugins/Importers/USDImporter/Source/USDStageImporter/Private/USDStageImportOptions.cpp b/Engine/Plugins/Importers/USDImporter/Source/USDStageImporter/Private/USDStageImportOptions.cpp
--- a/Engine/Plugins/Importers/USDImporter/Source/USDStageImporter/Private/USDStageImportOptions.cpp
+++ b/Engine/Plugins/Importers/USDImporter/Source/USDStageImporter/Private/USDStageImportOptions.cpp
@@ -54,30 +54,9 @@ UUsdStageImportOptions::UUsdStageImportOptions(const FObjectInitializer& ObjectI
 	bInterpretLODs = true;
 }
 
+// Deprecated
 void UUsdStageImportOptions::EnableActorImport(bool bEnable)
 {
-	for (TFieldIterator<FProperty> PropertyIterator(UUsdStageImportOptions::StaticClass()); PropertyIterator; ++PropertyIterator)
-	{
-		FProperty* Property = *PropertyIterator;
-		if (Property && Property->GetFName() == GET_MEMBER_NAME_CHECKED(UUsdStageImportOptions, bImportActors))
-		{
-			if (bEnable)
-			{
-				Property->SetMetaData(TEXT("ToolTip"), TEXT("Whether to spawn imported actors into the current level"));
-				Property->ClearPropertyFlags(CPF_EditConst);
-			}
-			else
-			{
-				bImportActors = false;
-				Property->SetMetaData(
-					TEXT("ToolTip"),
-					TEXT("Actor import is disabled when importing via the Content Browser. Use File->\"Import into Level...\" to also import actors.")
-				);
-				Property->SetPropertyFlags(CPF_EditConst);
-			}
-			break;
-		}
-	}
 }
 
 void UUsdStageImportOptions::PostEditChangeProperty(struct FPropertyChangedEvent& PropertyChangedEvent)
@@ -137,3 +116,48 @@ void UsdUtils::AddAnalyticsAttributes(const UUsdStageImportOptions& Options, TAr
 	InOutAttributes.Emplace(TEXT("MergeIdenticalMaterialSlots"), LexToString(Options.bMergeIdenticalMaterialSlots));
 	InOutAttributes.Emplace(TEXT("InterpretLODs"), LexToString(Options.bInterpretLODs));
 }
+
+UsdUtils::FScopedSuppressActorImport::FScopedSuppressActorImport(UUsdStageImportOptions* InOptions)
+	: Options{InOptions}
+{
+	for (TFieldIterator<FProperty> PropertyIterator(UUsdStageImportOptions::StaticClass()); PropertyIterator; ++PropertyIterator)
+	{
+		FProperty* Property = *PropertyIterator;
+		if (Property && Property->GetFName() == GET_MEMBER_NAME_CHECKED(UUsdStageImportOptions, bImportActors))
+		{
+			Property->SetMetaData(
+				TEXT("ToolTip"),
+				TEXT("Actor import is disabled when importing via the Content Browser. Use File->\"Import into Level...\" to also import actors.")
+			);
+			Property->SetPropertyFlags(CPF_EditConst);
+			break;
+		}
+	}
+
+	if (UUsdStageImportOptions* OptionsPtr = Options.Get())
+	{
+		bOldImportActorsValue = OptionsPtr->bImportActors;
+		OptionsPtr->bImportActors = false;
+		OptionsPtr->SaveConfig();
+	}
+}
+
+UsdUtils::FScopedSuppressActorImport::~FScopedSuppressActorImport()
+{
+	for (TFieldIterator<FProperty> PropertyIterator(UUsdStageImportOptions::StaticClass()); PropertyIterator; ++PropertyIterator)
+	{
+		FProperty* Property = *PropertyIterator;
+		if (Property && Property->GetFName() == GET_MEMBER_NAME_CHECKED(UUsdStageImportOptions, bImportActors))
+		{
+			Property->SetMetaData(TEXT("ToolTip"), TEXT("Whether to spawn imported actors into the current level"));
+			Property->ClearPropertyFlags(CPF_EditConst);
+			break;
+		}
+	}
+
+	if (UUsdStageImportOptions* OptionsPtr = Options.Get())
+	{
+		OptionsPtr->bImportActors = bOldImportActorsValue;
+		OptionsPtr->SaveConfig();
+	}
+}
diff --git a/Engine/Plugins/Importers/USDImporter/Source/USDStageImporter/Public/USDStageImportOptions.h b/Engine/Plugins/Importers/USDImporter/Source/USDStageImporter/Public/USDStageImportOptions.h
--- a/Engine/Plugins/Importers/USDImporter/Source/USDStageImporter/Public/USDStageImportOptions.h
+++ b/Engine/Plugins/Importers/USDImporter/Source/USDStageImporter/Public/USDStageImportOptions.h
@@ -248,11 +248,28 @@ public:
 	bool bInterpretLODs;
 
 public:
+	UE_DEPRECATED(5.6, "Use FScopedSuppressActorImport instead.")
 	void EnableActorImport(bool bEnable);
+
 	virtual void PostEditChangeProperty(struct FPropertyChangedEvent& PropertyChangedEvent) override;
 };
 
 namespace UsdUtils
 {
+	struct FScopedSuppressActorImport
+	{
+		FScopedSuppressActorImport(UUsdStageImportOptions* InOptions);
+		~FScopedSuppressActorImport();
+
+		FScopedSuppressActorImport(const FScopedSuppressActorImport& Other) = delete;
+		FScopedSuppressActorImport(FScopedSuppressActorImport&& Other) = delete;
+		FScopedSuppressActorImport& operator=(const FScopedSuppressActorImport& Other) = delete;
+		FScopedSuppressActorImport& operator=(FScopedSuppressActorImport&& Other) = delete;
+
+	private:
+		bool bOldImportActorsValue = true;
+		TWeakObjectPtr<UUsdStageImportOptions> Options;
+	};
+
 	USDSTAGEIMPORTER_API void AddAnalyticsAttributes(const UUsdStageImportOptions& Options, TArray<FAnalyticsEventAttribute>& InOutAttributes);
 }
