This change moves the later setup phases (Commit, PreBuild and Build) in FBuildStaticMeshTaskChain into a new FBaseBuildStaticMeshTaskChain base class, so that these phases can be utilized during Nanite static mesh creation.

Applying this change separate to the (incoming) Nanite USD import change should help avoid potential merge/integration conflicts that may occur due to large chunks of code moving around.

diff --git a/Engine/Plugins/Importers/USDImporter/Source/USDSchemas/Private/USDGeomMeshTranslator.cpp b/Engine/Plugins/Importers/USDImporter/Source/USDSchemas/Private/USDGeomMeshTranslator.cpp
--- a/Engine/Plugins/Importers/USDImporter/Source/USDSchemas/Private/USDGeomMeshTranslator.cpp
+++ b/Engine/Plugins/Importers/USDImporter/Source/USDSchemas/Private/USDGeomMeshTranslator.cpp
@@ -1483,13 +1483,125 @@
 	}
 }	 // namespace UE::UsdCollision::Private
 
+FBaseBuildStaticMeshTaskChain::FBaseBuildStaticMeshTaskChain(
+	const TSharedRef<FUsdSchemaTranslationContext>& InContext,
+	const UE::FSdfPath& InPrimPath
+)
+	: PrimPath(InPrimPath)
+	, Context(InContext)
+{
+}
+
+void FBaseBuildStaticMeshTaskChain::SetupTasks()
+{
+#if WITH_EDITOR
+	// Commit mesh description (Async)
+	Do(
+		ESchemaTranslationLaunchPolicy::Async,
+		[this]()
+		{
+			UStaticMesh::FCommitMeshDescriptionParams Params;
+			Params.bMarkPackageDirty = false;
+			Params.bUseHashAsGuid = true;
+
+			int32 NumMeshDescriptions = LODIndexToMeshDescription.Num();
+			if (StaticMesh->IsNaniteAssembly() && NumMeshDescriptions == 0)
+			{
+				NumMeshDescriptions = 1;
+			}
+
+			for (int32 LODIndex = 0; LODIndex < NumMeshDescriptions; ++LODIndex)
+			{
+				StaticMesh->CommitMeshDescription(LODIndex, Params);
+			}
+
+			return true;
+		}
+	);
+#endif	  // WITH_EDITOR
+
+	// PreBuild static mesh (Main thread)
+	Do(
+		ESchemaTranslationLaunchPolicy::Sync,
+		[this]()
+		{
+			RecreateRenderStateContextPtr = MakeShared<FStaticMeshComponentRecreateRenderStateContext>(StaticMesh, true, true);
+
+			UsdGeomMeshTranslatorImpl::PreBuildStaticMesh(*StaticMesh);
+
+			return true;
+		}
+	);
+
+	// Build static mesh (Async)
+	Then(
+		ESchemaTranslationLaunchPolicy::Async,
+		[this]() mutable
+		{
+			FStaticFeatureLevel FeatureLevel = GMaxRHIFeatureLevel;
+
+			UWorld* World = Context->Level ? Context->Level->GetWorld() : nullptr;
+			if (!World)
+			{
+				World = IUsdClassesModule::GetCurrentWorld();
+			}
+			if (World)
+			{
+				FeatureLevel = World->GetFeatureLevel();
+			}
+
+			bool bSuccess = UsdGeomMeshTranslatorImpl::BuildStaticMesh(*StaticMesh, FeatureLevel, LODIndexToMeshDescription);
+			if (bSuccess)
+			{
+				FStaticMeshRenderData* RenderData = StaticMesh->GetRenderData();
+
+				// We want to discard our StaticMesh if it generated no valid RenderData (e.g. was compeletely degenerate).
+				// They are not really going to be useful, and this is also handy as in some cases the StaticMeshComponents that use them
+				// may get confused as to why they have a StaticMesh with no valid RenderData at points where they expected to have them
+				// (e.g. see UE-201946)
+				bSuccess = RenderData && !RenderData->Bounds.ContainsNaN() && RenderData->GetCurrentFirstLODIdx(0) != INDEX_NONE;
+			}
+
+			// Build failed, abandon mesh
+			if (!bSuccess)
+			{
+				USD_LOG_WARNING(
+					TEXT(
+						"Discarding StaticMesh generated for prim '%s' as it didn't produce any valid RenderData (likely all triangles were degenerate)"
+					),
+					*PrimPath.GetString()
+				);
+				UsdUnreal::TranslatorUtils::AbandonFailedAsset(StaticMesh, Context->UsdAssetCache.Get(), Context->PrimLinkCache);
+				return false;
+			}
+
+			UE::UsdCollision::Private::SetupSimpleCollision(GetPrim(), *StaticMesh, Context->FallbackCollisionType);
+
+			return true;
+		}
+	);
+
+	// PostBuild static mesh (Main thread)
+	Then(
+		ESchemaTranslationLaunchPolicy::Sync,
+		[this]()
+		{
+			UsdGeomMeshTranslatorImpl::PostBuildStaticMesh(*StaticMesh, LODIndexToMeshDescription);
+
+			RecreateRenderStateContextPtr.Reset();
+
+			return true;
+		}
+	);
+}
+
 FBuildStaticMeshTaskChain::FBuildStaticMeshTaskChain(
 	const TSharedRef<FUsdSchemaTranslationContext>& InContext,
 	const UE::FSdfPath& InPrimPath,
 	const TOptional<UE::FSdfPath>& InAlternativePrimToLinkAssetsTo
 )
-	: PrimPath(InPrimPath)
-	, Context(InContext)
+
+	: FBaseBuildStaticMeshTaskChain(InContext, InPrimPath)
 	, AlternativePrimToLinkAssetsTo(InAlternativePrimToLinkAssetsTo)
 {
 }
@@ -1630,99 +1742,7 @@
 		   return bIsNew;
 	   });
 
-#if WITH_EDITOR
-	// Commit mesh description (Async)
-	Then(
-		ESchemaTranslationLaunchPolicy::Async,
-		[this]()
-		{
-			UStaticMesh::FCommitMeshDescriptionParams Params;
-			Params.bMarkPackageDirty = false;
-			Params.bUseHashAsGuid = true;
-
-			for (int32 LODIndex = 0; LODIndex < LODIndexToMeshDescription.Num(); ++LODIndex)
-			{
-				StaticMesh->CommitMeshDescription(LODIndex, Params);
-			}
-
-			return true;
-		}
-	);
-#endif	  // WITH_EDITOR
-
-	// PreBuild static mesh (Main thread)
-	Then(
-		ESchemaTranslationLaunchPolicy::Sync,
-		[this]()
-		{
-			RecreateRenderStateContextPtr = MakeShared<FStaticMeshComponentRecreateRenderStateContext>(StaticMesh, true, true);
-
-			UsdGeomMeshTranslatorImpl::PreBuildStaticMesh(*StaticMesh);
-
-			return true;
-		}
-	);
-
-	// Build static mesh (Async)
-	Then(
-		ESchemaTranslationLaunchPolicy::Async,
-		[this]() mutable
-		{
-			FStaticFeatureLevel FeatureLevel = GMaxRHIFeatureLevel;
-
-			UWorld* World = Context->Level ? Context->Level->GetWorld() : nullptr;
-			if (!World)
-			{
-				World = IUsdClassesModule::GetCurrentWorld();
-			}
-			if (World)
-			{
-				FeatureLevel = World->GetFeatureLevel();
-			}
-
-			bool bSuccess = UsdGeomMeshTranslatorImpl::BuildStaticMesh(*StaticMesh, FeatureLevel, LODIndexToMeshDescription);
-			if (bSuccess)
-			{
-				FStaticMeshRenderData* RenderData = StaticMesh->GetRenderData();
-
-				// We want to discard our StaticMesh if it generated no valid RenderData (e.g. was compeletely degenerate).
-				// They are not really going to be useful, and this is also handy as in some cases the StaticMeshComponents that use them
-				// may get confused as to why they have a StaticMesh with no valid RenderData at points where they expected to have them
-				// (e.g. see UE-201946)
-				bSuccess = RenderData && !RenderData->Bounds.ContainsNaN() && RenderData->GetCurrentFirstLODIdx(0) != INDEX_NONE;
-			}
-
-			// Build failed, abandon mesh
-			if (!bSuccess)
-			{
-				USD_LOG_WARNING(
-					TEXT(
-						"Discarding StaticMesh generated for prim '%s' as it didn't produce any valid RenderData (likely all triangles were degenerate)"
-					),
-					*PrimPath.GetString()
-				);
-				UsdUnreal::TranslatorUtils::AbandonFailedAsset(StaticMesh, Context->UsdAssetCache.Get(), Context->PrimLinkCache);
-				return false;
-			}
-
-			UE::UsdCollision::Private::SetupSimpleCollision(GetPrim(), *StaticMesh, Context->FallbackCollisionType);
-
-			return true;
-		}
-	);
-
-	// PostBuild static mesh (Main thread)
-	Then(
-		ESchemaTranslationLaunchPolicy::Sync,
-		[this]()
-		{
-			UsdGeomMeshTranslatorImpl::PostBuildStaticMesh(*StaticMesh, LODIndexToMeshDescription);
-
-			RecreateRenderStateContextPtr.Reset();
-
-			return true;
-		}
-	);
+	FBaseBuildStaticMeshTaskChain::SetupTasks();
 }
 
 FGeomMeshCreateAssetsTaskChain::FGeomMeshCreateAssetsTaskChain(
diff --git a/Engine/Plugins/Importers/USDImporter/Source/USDSchemas/Public/USDGeomMeshTranslator.h b/Engine/Plugins/Importers/USDImporter/Source/USDSchemas/Public/USDGeomMeshTranslator.h
--- a/Engine/Plugins/Importers/USDImporter/Source/USDSchemas/Public/USDGeomMeshTranslator.h
+++ b/Engine/Plugins/Importers/USDImporter/Source/USDSchemas/Public/USDGeomMeshTranslator.h
@@ -20,33 +20,25 @@ namespace UsdUtils
 	struct FUsdPrimMaterialAssignmentInfo;
 }
 
-class FBuildStaticMeshTaskChain : public FUsdSchemaTranslatorTaskChain
+class FBaseBuildStaticMeshTaskChain : public FUsdSchemaTranslatorTaskChain
 {
 public:
-	explicit FBuildStaticMeshTaskChain(
+	explicit FBaseBuildStaticMeshTaskChain(
 		const TSharedRef<FUsdSchemaTranslationContext>& InContext,
-		const UE::FSdfPath& InPrimPath,
-		const TOptional<UE::FSdfPath>& AlternativePrimToLinkAssetsTo = {}
+		const UE::FSdfPath& InPrimPath
 	);
 
 protected:
 	// Inputs
-	// When multiple meshes are collapsed together, this Schema might not be the same as the Context schema, which is the root schema
 	UE::FSdfPath PrimPath;
 	TSharedRef<FUsdSchemaTranslationContext> Context;
 	TArray<FMeshDescription> LODIndexToMeshDescription;
-	TArray<UsdUtils::FUsdPrimMaterialAssignmentInfo> LODIndexToMaterialInfo;
-	// We collect metadata early (during LOD parsing) so that we don't have to flip through
-	// LODs multiple times
-	FUsdCombinedPrimMetadata LODMetadata;
 
 	// Outputs
-	TOptional<UE::FSdfPath> AlternativePrimToLinkAssetsTo;
 	UStaticMesh* StaticMesh = nullptr;
 
 	// Required to prevent StaticMesh from being used for drawing while it is being rebuilt
 	TSharedPtr<FStaticMeshComponentRecreateRenderStateContext> RecreateRenderStateContextPtr;
-	bool bCollectedMetadata = false;
 
 protected:
 	UE::FUsdPrim GetPrim() const
@@ -57,6 +49,30 @@ protected:
 	virtual void SetupTasks();
 };
 
+class FBuildStaticMeshTaskChain : public FBaseBuildStaticMeshTaskChain
+{
+public:
+	explicit FBuildStaticMeshTaskChain(
+		const TSharedRef<FUsdSchemaTranslationContext>& InContext,
+		const UE::FSdfPath& InPrimPath,
+		const TOptional<UE::FSdfPath>& AlternativePrimToLinkAssetsTo = {}
+	);
+
+protected:
+	// Inputs
+	TArray<UsdUtils::FUsdPrimMaterialAssignmentInfo> LODIndexToMaterialInfo;
+	// We collect metadata early (during LOD parsing) so that we don't have to flip through
+	// LODs multiple times
+	FUsdCombinedPrimMetadata LODMetadata;
+
+	// Outputs
+	TOptional<UE::FSdfPath> AlternativePrimToLinkAssetsTo;
+	bool bCollectedMetadata = false;
+
+protected:
+	virtual void SetupTasks() override;
+};
+
 class FGeomMeshCreateAssetsTaskChain : public FBuildStaticMeshTaskChain
 {
 public:
