USD: Fix AnimSequences seeming to have no animation data after doing Actions -> Import on the USD Stage Editor.

diff --git a/Engine/Plugins/Importers/USDImporter/Source/USDStageImporter/Private/USDStageImporter.cpp b/Engine/Plugins/Importers/USDImporter/Source/USDStageImporter/Private/USDStageImporter.cpp
--- a/Engine/Plugins/Importers/USDImporter/Source/USDStageImporter/Private/USDStageImporter.cpp
+++ b/Engine/Plugins/Importers/USDImporter/Source/USDStageImporter/Private/USDStageImporter.cpp
@@ -1331,7 +1331,7 @@ namespace UE::USDStageImporter::Private
 		}
 	}
 
-	void CallAssetsPostEditChange(const TSet<UObject*>& PublishedObjects)
+	void PostRemapAssets(const TSet<UObject*>& PublishedObjects)
 	{
 		for (UObject* PublishedObject : PublishedObjects)
 		{
@@ -1345,6 +1345,18 @@ namespace UE::USDStageImporter::Private
 				// remapped by our RemapReferences and RemapSoftReferences functions.
 				Material->PostEditChange();
 			}
+			else if (UAnimSequence* AnimSequence = Cast<UAnimSequence>(PublishedObject))
+			{
+				// Refresh AnimSequences: They contain internal transient ControlRigs that won't be remapped properly via RemapReferences or
+				// RemapSoftReferences and may be left trying to point at the original skeletons.
+				//
+				// Realistically we should probably do this with *every* asset type, as they may all have complex transient references to
+				// other assets that don't get picked up by the simple FProperty remapping mentioned above... The properties should be the
+				// only references that get serialized however, so if we remap the properties and then PostLoad() we should hopefully be able
+				// to achieve a valid state every time, as it should cause these assets to rebuild their transient info from their deserialized
+				// (and remapped) properties. It seems that only UAnimSequences are complex enough to require this at this time though
+				AnimSequence->PostLoad();
+			}
 		}
 	}
 
@@ -1896,7 +1908,7 @@ void UUsdStageImporter::ImportFromFile(FUsdStageImportContext& ImportContext)
 		ImportContext.ImportOptions->ExistingActorPolicy,
 		AssetsToFinalize
 	);
-	UE::USDStageImporter::Private::CallAssetsPostEditChange(PublishedAssets);
+	UE::USDStageImporter::Private::PostRemapAssets(PublishedAssets);
 	UE::USDStageImporter::Private::BroadcastImportEvents(PublishedAssets);
 	UE::USDStageImporter::Private::RefreshComponents(ImportContext.SceneActor, ImportContext.ImportOptions->bImportAtSpecificTimeCode);
 
@@ -2142,10 +2154,10 @@ bool UUsdStageImporter::ReimportSingleAsset(
 			AssetsToFinalize
 		);
 
-		TSet<UObject*> PublishedAssetsAndDependencies;
-		UE::USDStageImporter::Private::RemapReferences(ImportContext, {OutReimportedAsset}, ObjectsToRemap);
-		UE::USDStageImporter::Private::RemapSoftReferences(ImportContext, {OutReimportedAsset}, SoftObjectsToRemap);
-		UE::USDStageImporter::Private::CallAssetsPostEditChange(PublishedAssetsAndDependencies);
+		TSet<UObject*> PublishedAssetsAndDependencies{OutReimportedAsset};
+		UE::USDStageImporter::Private::RemapReferences(ImportContext, PublishedAssetsAndDependencies, ObjectsToRemap);
+		UE::USDStageImporter::Private::RemapSoftReferences(ImportContext, PublishedAssetsAndDependencies, SoftObjectsToRemap);
+		UE::USDStageImporter::Private::PostRemapAssets(PublishedAssetsAndDependencies);
 
 		bSuccess = OutReimportedAsset != nullptr && ActiveAssets.Contains(ReimportedObject);
 	}
