USD: Fixed potential animation compression error due to the AnimSequence being marked as garbage before it has finished the compression task.

diff --git a/Engine/Plugins/Importers/USDImporter/Source/USDStage/Private/USDLevelSequenceHelper.cpp b/Engine/Plugins/Importers/USDImporter/Source/USDStage/Private/USDLevelSequenceHelper.cpp
--- a/Engine/Plugins/Importers/USDImporter/Source/USDStage/Private/USDLevelSequenceHelper.cpp
+++ b/Engine/Plugins/Importers/USDImporter/Source/USDStage/Private/USDLevelSequenceHelper.cpp
@@ -483,6 +483,9 @@ namespace UE::USDLevelSequenceHelper::Private
 cleanup:
 		if (bCreatedTempSequence && AnimSequence)
 		{
+			// Animation compression happens asynchonously so here we make sure it's completed before marking it as garbage
+			const bool bWantResults = false;
+			AnimSequence->WaitOnExistingCompression(bWantResults);
 			AnimSequence->MarkAsGarbage();
 		}
 
