USDCore, USDImporter, InterchangeOpenUSD: use ENABLE_USD_MATERIALX definition to control usage of MaterialX via USD

This addresses a mismatch between the build and the code where the linkage
against MaterialX defined by the build did not align with the usage of the MaterialX
API in the code. This was an issue in particular in the USDSchemas module of
USDImporter.

This change adds a public definition to control whether to enable MaterialX via
USD in the USDUtilities module, which is a common dependency among all
of the client code.

The platform-based restrictions should be temporary while we look to address
lingering issues with MaterialX on Mac and Linux before we re-enable MaterialX
support via USD on all platforms.

diff --git a/Engine/Plugins/Importers/USDImporter/Source/USDSchemas/Private/Custom/MaterialXUSDShadeMaterialTranslator.cpp b/Engine/Plugins/Importers/USDImporter/Source/USDSchemas/Private/Custom/MaterialXUSDShadeMaterialTranslator.cpp
--- a/Engine/Plugins/Importers/USDImporter/Source/USDSchemas/Private/Custom/MaterialXUSDShadeMaterialTranslator.cpp
+++ b/Engine/Plugins/Importers/USDImporter/Source/USDSchemas/Private/Custom/MaterialXUSDShadeMaterialTranslator.cpp
@@ -36,9 +36,9 @@
 #include "pxr/usd/usdShade/material.h"
 #include "USDIncludesEnd.h"
 
-#if WITH_EDITOR
+#if ENABLE_USD_MATERIALX
 #include "MaterialXFormat/Util.h"
-#endif
+#endif // ENABLE_USD_MATERIALX
 
 #define LOCTEXT_NAMESPACE "MaterialXUSDShadeMaterialTranslator"
 
@@ -53,7 +53,7 @@ static FAutoConsoleVariableRef CvarUseInterchangeMaterialXTranslator(
 
 namespace UE::USDMaterialXTranslator::Private
 {
-#if WITH_EDITOR
+#if ENABLE_USD_MATERIALX
 	// Quick shallow parse of MaterialX documents ourselves so that we can find all the referenced MaterialX files and
 	// build a robust hash to use for caching our generated assets
 	FString HashMaterialXFile(const FString& MaterialXFilePath)
@@ -264,12 +264,12 @@ namespace UE::USDMaterialXTranslator::Private
 
 		return Manager.ImportAsset(TempPackagePath, SourceData, InterchangeParameters);
 	}
-#endif	  // WITH_EDITOR
+#endif // ENABLE_USD_MATERIALX
 }	 // namespace UE::USDMaterialXTranslator::Private
 
 void FMaterialXUsdShadeMaterialTranslator::CreateAssets()
 {
-#if WITH_EDITOR
+#if ENABLE_USD_MATERIALX
 	// We handle MaterialX materials here by leveraging USD's own usdMtlx plugin. What it does is translate references
 	// to .mtlx files into hierarchies of Material/Shader/NodeGraph prims as the stage is opened. To actually use
 	// the materials contained in the .mtlx files, other prims from the stage will create material bindings to those
@@ -512,7 +512,7 @@ void FMaterialXUsdShadeMaterialTranslator::CreateAssets()
 	}
 #else
 	Super::CreateAssets();
-#endif	  // WITH_EDITOR
+#endif // ENABLE_USD_MATERIALX
 }
 
 #undef LOCTEXT_NAMESPACE
diff --git a/Engine/Plugins/Interchange/Extensions/OpenUSD/Source/Import/Private/InterchangeUsdTranslator.cpp b/Engine/Plugins/Interchange/Extensions/OpenUSD/Source/Import/Private/InterchangeUsdTranslator.cpp
--- a/Engine/Plugins/Interchange/Extensions/OpenUSD/Source/Import/Private/InterchangeUsdTranslator.cpp
+++ b/Engine/Plugins/Interchange/Extensions/OpenUSD/Source/Import/Private/InterchangeUsdTranslator.cpp
@@ -1151,8 +1151,8 @@ namespace UE::InterchangeUsdTranslator::Private
 				return true;
 			}
 
-			// Only enable the shader graph on Windows for the time being as it crashes on Linux cause of a probable double free
-#if PLATFORM_WINDOWS
+			// Conditionally enable the shader graph for the time being as it crashes on Linux cause of a probable double free
+#if ENABLE_USD_MATERIALX
 			bool bHasShaderGraph = false;
 			FUsdMaterialXShaderGraph ShaderGraph(Prim, *RenderContext.ToString());
 			GeomPropValues = ShaderGraph.GetGeomPropValueNames();
@@ -1160,7 +1160,7 @@ namespace UE::InterchangeUsdTranslator::Private
 			{
 				return FMaterialXManager::GetInstance().Translate(Document, NodeContainer);
 			}
-#endif	  // PLATFORM_WINDOWS
+#endif // ENABLE_USD_MATERIALX
 		}
 #endif	  // WITH_EDITOR
 		return false;
diff --git a/Engine/Plugins/Runtime/USDCore/Source/USDUtilities/Private/USDMaterialXShaderGraph.cpp b/Engine/Plugins/Runtime/USDCore/Source/USDUtilities/Private/USDMaterialXShaderGraph.cpp
--- a/Engine/Plugins/Runtime/USDCore/Source/USDUtilities/Private/USDMaterialXShaderGraph.cpp
+++ b/Engine/Plugins/Runtime/USDCore/Source/USDUtilities/Private/USDMaterialXShaderGraph.cpp
@@ -2,7 +2,7 @@
 
 #include "USDMaterialXShaderGraph.h"
 
-#ifdef ENABLE_USD_MATERIALX
+#if USE_USD_SDK && ENABLE_USD_MATERIALX
 
 #include "Misc/Paths.h"
 #include "USDErrorUtils.h"
@@ -540,4 +540,4 @@ bool FUsdMaterialXShaderGraph::IsTangentSpaceInput(const pxr::UsdShadeInput& Usd
 {
 	return TangentSpaceInputs.find(UsdInput.GetBaseName().GetString()) != TangentSpaceInputs.cend();
 }
-#endif	  // ENABLE_USD_MATERIALX
+#endif // USE_USD_SDK && ENABLE_USD_MATERIALX
diff --git a/Engine/Plugins/Runtime/USDCore/Source/USDUtilities/Public/USDMaterialXShaderGraph.h b/Engine/Plugins/Runtime/USDCore/Source/USDUtilities/Public/USDMaterialXShaderGraph.h
--- a/Engine/Plugins/Runtime/USDCore/Source/USDUtilities/Public/USDMaterialXShaderGraph.h
+++ b/Engine/Plugins/Runtime/USDCore/Source/USDUtilities/Public/USDMaterialXShaderGraph.h
@@ -4,19 +4,15 @@
 
 #include "Containers/UnrealString.h"
 
-#if WITH_EDITOR && PLATFORM_WINDOWS
-#include "MaterialXFormat/Util.h"
-
-#if USE_USD_SDK
+#if USE_USD_SDK && ENABLE_USD_MATERIALX
 #include "USDIncludesStart.h"
 #include "pxr/base/tf/token.h"
 #include "pxr/usd/sdf/layerUtils.h"
 #include "pxr/usd/usdShade/material.h"
 #include "USDIncludesEnd.h"
 
-#define ENABLE_USD_MATERIALX
-#endif // USE_USD_SDK
-#endif //WITH_EDITOR && PLATFORM_WINDOWS
+#include "MaterialXFormat/Util.h"
+#endif // USE_USD_SDK && ENABLE_USD_MATERIALX
 
 /** 
  * Create a MaterialX shader graph from a UsdPrim, all <geompropvalue> nodes are either converted to an <image> or a proper <geomcolor>
@@ -32,7 +28,7 @@ public:
 		bool bTangentSpace = false;
 	};
 
-#ifdef ENABLE_USD_MATERIALX
+#if USE_USD_SDK && ENABLE_USD_MATERIALX
 
 	/**
 	 * @param UsdShadeMaterialPrim - The prim used to generate the shader graph, should be a valid pxr::UsdShadeMaterial
@@ -114,5 +110,5 @@ private:
 		               std::is_same_v<MtlxType, MaterialX::Vector2>,
 		               pxr::GfVec2f,
 		   MtlxType>>>;
-#endif // ENABLE_USD_MATERIALX
+#endif // USE_USD_SDK && ENABLE_USD_MATERIALX
 };
\ No newline at end of file
diff --git a/Engine/Plugins/Runtime/USDCore/Source/USDUtilities/USDUtilities.Build.cs b/Engine/Plugins/Runtime/USDCore/Source/USDUtilities/USDUtilities.Build.cs
--- a/Engine/Plugins/Runtime/USDCore/Source/USDUtilities/USDUtilities.Build.cs
+++ b/Engine/Plugins/Runtime/USDCore/Source/USDUtilities/USDUtilities.Build.cs
@@ -47,6 +47,8 @@ namespace UnrealBuildTool.Rules
 				}
 			);
 
+			bool bEnableUsdMaterialX = false;
+
 			if (Target.bBuildEditor)
 			{
 				PrivateDependencyModuleNames.AddRange(
@@ -62,19 +64,27 @@ namespace UnrealBuildTool.Rules
 						"UnrealEd",
 					}
 				);
-			}
 
-			PrivateDefinitions.Add("SUPPRESS_PER_MODULE_INLINE_FILE"); // This module does not use core's standard operator new/delete overloads
+				// MaterialX via USD is only enabled on Windows for now while
+				// we work through issues affecting Mac and Linux.
+				if (Target.Platform == UnrealTargetPlatform.Win64)
+				{
+					// Public because MaterialX headers are #include'd in
+					// public USDUtilities headers.
+					PublicDependencyModuleNames.AddRange(
+						new string[]
+						{
+							"MaterialX"
+						}
+					);
 
-			if (Target.Platform == UnrealTargetPlatform.Win64)
-			{
-				PublicDependencyModuleNames.AddRange(
-					new string[]
-					{
-						"MaterialX"
-					}
-				);
+					bEnableUsdMaterialX = true;
+				}
 			}
+
+			PublicDefinitions.Add("ENABLE_USD_MATERIALX=" + (bEnableUsdMaterialX ? "1" : "0"));
+
+			PrivateDefinitions.Add("SUPPRESS_PER_MODULE_INLINE_FILE"); // This module does not use core's standard operator new/delete overloads
 		}
 	}
 }
