USD: Work around N^2 performance bottleneck when spawning many light components;

diff --git a/Engine/Plugins/Importers/USDImporter/Source/USDStage/Private/USDStageActor.cpp b/Engine/Plugins/Importers/USDImporter/Source/USDStage/Private/USDStageActor.cpp
--- a/Engine/Plugins/Importers/USDImporter/Source/USDStage/Private/USDStageActor.cpp
+++ b/Engine/Plugins/Importers/USDImporter/Source/USDStage/Private/USDStageActor.cpp
@@ -62,6 +62,7 @@
 #include "Engine/World.h"
 #include "EngineAnalytics.h"
 #include "Framework/Notifications/NotificationManager.h"
+#include "GeometryCacheComponent.h"
 #include "HAL/IConsoleManager.h"
 #include "LevelSequence.h"
 #include "LiveLinkComponentController.h"
@@ -2199,7 +2200,10 @@ UUsdPrimTwin* AUsdStageActor::ExpandPrim(
 	if (TwinSceneComponent && !TranslationContext.bIsJustRepopulatingLevelSequence)
 	{
 #if WITH_EDITOR
-		TwinSceneComponent->PostEditChange();
+		if (TwinSceneComponent->IsA<UGeometryCacheComponent>() || TwinSceneComponent->IsA<USkyLightComponent>())
+		{
+			TwinSceneComponent->PostEditChange();
+		}
 #endif	  // WITH_EDITOR
 
 		if (!TwinSceneComponent->IsRegistered())
