Added parentheses for clarity.

diff --git a/Engine/Plugins/Experimental/ChaosCachingUSD/Source/ChaosCachingUSD/ChaosCachingUSD.build.cs b/Engine/Plugins/Experimental/ChaosCachingUSD/Source/ChaosCachingUSD/ChaosCachingUSD.build.cs
--- a/Engine/Plugins/Experimental/ChaosCachingUSD/Source/ChaosCachingUSD/ChaosCachingUSD.build.cs
+++ b/Engine/Plugins/Experimental/ChaosCachingUSD/Source/ChaosCachingUSD/ChaosCachingUSD.build.cs
@@ -32,8 +32,8 @@ namespace UnrealBuildTool.Rules
 			);
 
 			bool bUseUsdMemoryOverloadsIfSupported = true;
-			bool bEnableUsdMemoryOverloads = EnableUsdSdk(Target) && bUseUsdMemoryOverloadsIfSupported &&
-				(Target.Platform == UnrealTargetPlatform.Win64 && Target.Architecture == UnrealArch.X64 ||
+			bool bEnableUsdMemoryOverloads = (EnableUsdSdk(Target) && bUseUsdMemoryOverloadsIfSupported) &&
+				((Target.Platform == UnrealTargetPlatform.Win64 && Target.Architecture == UnrealArch.X64) ||
 				Target.Platform == UnrealTargetPlatform.Mac ||
 				Target.Platform == UnrealTargetPlatform.Linux);
 			if (!bEnableUsdMemoryOverloads)
diff --git a/Engine/Plugins/Experimental/ChaosFlesh/Source/ChaosFleshEngine/ChaosFleshEngine.Build.cs b/Engine/Plugins/Experimental/ChaosFlesh/Source/ChaosFleshEngine/ChaosFleshEngine.Build.cs
--- a/Engine/Plugins/Experimental/ChaosFlesh/Source/ChaosFleshEngine/ChaosFleshEngine.Build.cs
+++ b/Engine/Plugins/Experimental/ChaosFlesh/Source/ChaosFleshEngine/ChaosFleshEngine.Build.cs
@@ -89,8 +89,8 @@ namespace UnrealBuildTool.Rules
 				PublicDefinitions.Add("DO_USD_CACHING=1");
 
 				bool bUseUsdMemoryOverloadsIfSupported = true;
-				bool bEnableUsdMemoryOverloads = EnableUsdSdk(Target) && bUseUsdMemoryOverloadsIfSupported &&
-					(Target.Platform == UnrealTargetPlatform.Win64 && Target.Architecture == UnrealArch.X64 ||
+				bool bEnableUsdMemoryOverloads = (EnableUsdSdk(Target) && bUseUsdMemoryOverloadsIfSupported) &&
+					((Target.Platform == UnrealTargetPlatform.Win64 && Target.Architecture == UnrealArch.X64) ||
 					Target.Platform == UnrealTargetPlatform.Mac ||
 					Target.Platform == UnrealTargetPlatform.Linux);
 				if (!bEnableUsdMemoryOverloads)
diff --git a/Engine/Plugins/Importers/USDImporter/Source/USDExporter/USDExporter.Build.cs b/Engine/Plugins/Importers/USDImporter/Source/USDExporter/USDExporter.Build.cs
--- a/Engine/Plugins/Importers/USDImporter/Source/USDExporter/USDExporter.Build.cs
+++ b/Engine/Plugins/Importers/USDImporter/Source/USDExporter/USDExporter.Build.cs
@@ -62,8 +62,8 @@ namespace UnrealBuildTool.Rules
 			);
 
 			bool bUseUsdMemoryOverloadsIfSupported = true;
-			bool bEnableUsdMemoryOverloads = EnableUsdSdk(Target) && bUseUsdMemoryOverloadsIfSupported &&
-				(Target.Platform == UnrealTargetPlatform.Win64 && Target.Architecture == UnrealArch.X64 ||
+			bool bEnableUsdMemoryOverloads = (EnableUsdSdk(Target) && bUseUsdMemoryOverloadsIfSupported) &&
+				((Target.Platform == UnrealTargetPlatform.Win64 && Target.Architecture == UnrealArch.X64) ||
 				Target.Platform == UnrealTargetPlatform.Mac ||
 				Target.Platform == UnrealTargetPlatform.Linux);
 			if (!bEnableUsdMemoryOverloads)
diff --git a/Engine/Plugins/Importers/USDImporter/Source/USDSchemas/USDSchemas.Build.cs b/Engine/Plugins/Importers/USDImporter/Source/USDSchemas/USDSchemas.Build.cs
--- a/Engine/Plugins/Importers/USDImporter/Source/USDSchemas/USDSchemas.Build.cs
+++ b/Engine/Plugins/Importers/USDImporter/Source/USDSchemas/USDSchemas.Build.cs
@@ -81,8 +81,8 @@ namespace UnrealBuildTool.Rules
 			}
 
 			bool bUseUsdMemoryOverloadsIfSupported = true;
-			bool bEnableUsdMemoryOverloads = EnableUsdSdk(Target) && bUseUsdMemoryOverloadsIfSupported &&
-				(Target.Platform == UnrealTargetPlatform.Win64 && Target.Architecture == UnrealArch.X64 ||
+			bool bEnableUsdMemoryOverloads = (EnableUsdSdk(Target) && bUseUsdMemoryOverloadsIfSupported) &&
+				((Target.Platform == UnrealTargetPlatform.Win64 && Target.Architecture == UnrealArch.X64) ||
 				Target.Platform == UnrealTargetPlatform.Mac ||
 				Target.Platform == UnrealTargetPlatform.Linux);
 			if (!bEnableUsdMemoryOverloads)
diff --git a/Engine/Plugins/Importers/USDImporter/Source/USDStage/USDStage.Build.cs b/Engine/Plugins/Importers/USDImporter/Source/USDStage/USDStage.Build.cs
--- a/Engine/Plugins/Importers/USDImporter/Source/USDStage/USDStage.Build.cs
+++ b/Engine/Plugins/Importers/USDImporter/Source/USDStage/USDStage.Build.cs
@@ -73,8 +73,8 @@ namespace UnrealBuildTool.Rules
 			}
 
 			bool bUseUsdMemoryOverloadsIfSupported = true;
-			bool bEnableUsdMemoryOverloads = EnableUsdSdk(Target) && bUseUsdMemoryOverloadsIfSupported &&
-				(Target.Platform == UnrealTargetPlatform.Win64 && Target.Architecture == UnrealArch.X64 ||
+			bool bEnableUsdMemoryOverloads = (EnableUsdSdk(Target) && bUseUsdMemoryOverloadsIfSupported) &&
+				((Target.Platform == UnrealTargetPlatform.Win64 && Target.Architecture == UnrealArch.X64) ||
 				Target.Platform == UnrealTargetPlatform.Mac ||
 				Target.Platform == UnrealTargetPlatform.Linux);
 			if (!bEnableUsdMemoryOverloads)
diff --git a/Engine/Plugins/Importers/USDImporter/Source/USDStageEditor/USDStageEditor.Build.cs b/Engine/Plugins/Importers/USDImporter/Source/USDStageEditor/USDStageEditor.Build.cs
--- a/Engine/Plugins/Importers/USDImporter/Source/USDStageEditor/USDStageEditor.Build.cs
+++ b/Engine/Plugins/Importers/USDImporter/Source/USDStageEditor/USDStageEditor.Build.cs
@@ -41,8 +41,8 @@ namespace UnrealBuildTool.Rules
 			);
 
 			bool bUseUsdMemoryOverloadsIfSupported = true;
-			bool bEnableUsdMemoryOverloads = EnableUsdSdk(Target) && bUseUsdMemoryOverloadsIfSupported &&
-				(Target.Platform == UnrealTargetPlatform.Win64 && Target.Architecture == UnrealArch.X64 ||
+			bool bEnableUsdMemoryOverloads = (EnableUsdSdk(Target) && bUseUsdMemoryOverloadsIfSupported) &&
+				((Target.Platform == UnrealTargetPlatform.Win64 && Target.Architecture == UnrealArch.X64) ||
 				Target.Platform == UnrealTargetPlatform.Mac ||
 				Target.Platform == UnrealTargetPlatform.Linux);
 			if (!bEnableUsdMemoryOverloads)
diff --git a/Engine/Plugins/Importers/USDImporter/Source/USDStageEditorViewModels/USDStageEditorViewModels.Build.cs b/Engine/Plugins/Importers/USDImporter/Source/USDStageEditorViewModels/USDStageEditorViewModels.Build.cs
--- a/Engine/Plugins/Importers/USDImporter/Source/USDStageEditorViewModels/USDStageEditorViewModels.Build.cs
+++ b/Engine/Plugins/Importers/USDImporter/Source/USDStageEditorViewModels/USDStageEditorViewModels.Build.cs
@@ -35,8 +35,8 @@ namespace UnrealBuildTool.Rules
 			);
 
 			bool bUseUsdMemoryOverloadsIfSupported = true;
-			bool bEnableUsdMemoryOverloads = EnableUsdSdk(Target) && bUseUsdMemoryOverloadsIfSupported &&
-				(Target.Platform == UnrealTargetPlatform.Win64 && Target.Architecture == UnrealArch.X64 ||
+			bool bEnableUsdMemoryOverloads = (EnableUsdSdk(Target) && bUseUsdMemoryOverloadsIfSupported) &&
+				((Target.Platform == UnrealTargetPlatform.Win64 && Target.Architecture == UnrealArch.X64) ||
 				Target.Platform == UnrealTargetPlatform.Mac ||
 				Target.Platform == UnrealTargetPlatform.Linux);
 			if (!bEnableUsdMemoryOverloads)
diff --git a/Engine/Plugins/Importers/USDImporter/Source/USDStageImporter/USDStageImporter.Build.cs b/Engine/Plugins/Importers/USDImporter/Source/USDStageImporter/USDStageImporter.Build.cs
--- a/Engine/Plugins/Importers/USDImporter/Source/USDStageImporter/USDStageImporter.Build.cs
+++ b/Engine/Plugins/Importers/USDImporter/Source/USDStageImporter/USDStageImporter.Build.cs
@@ -47,8 +47,8 @@ namespace UnrealBuildTool.Rules
 			);
 
 			bool bUseUsdMemoryOverloadsIfSupported = true;
-			bool bEnableUsdMemoryOverloads = EnableUsdSdk(Target) && bUseUsdMemoryOverloadsIfSupported &&
-				(Target.Platform == UnrealTargetPlatform.Win64 && Target.Architecture == UnrealArch.X64 ||
+			bool bEnableUsdMemoryOverloads = (EnableUsdSdk(Target) && bUseUsdMemoryOverloadsIfSupported) &&
+				((Target.Platform == UnrealTargetPlatform.Win64 && Target.Architecture == UnrealArch.X64) ||
 				Target.Platform == UnrealTargetPlatform.Mac ||
 				Target.Platform == UnrealTargetPlatform.Linux);
 			if (!bEnableUsdMemoryOverloads)
diff --git a/Engine/Plugins/Importers/USDImporter/Source/USDTests/USDTests.Build.cs b/Engine/Plugins/Importers/USDImporter/Source/USDTests/USDTests.Build.cs
--- a/Engine/Plugins/Importers/USDImporter/Source/USDTests/USDTests.Build.cs
+++ b/Engine/Plugins/Importers/USDImporter/Source/USDTests/USDTests.Build.cs
@@ -36,8 +36,8 @@ namespace UnrealBuildTool.Rules
 			}
 
 			bool bUseUsdMemoryOverloadsIfSupported = true;
-			bool bEnableUsdMemoryOverloads = EnableUsdSdk(Target) && bUseUsdMemoryOverloadsIfSupported &&
-				(Target.Platform == UnrealTargetPlatform.Win64 && Target.Architecture == UnrealArch.X64 ||
+			bool bEnableUsdMemoryOverloads = (EnableUsdSdk(Target) && bUseUsdMemoryOverloadsIfSupported) &&
+				((Target.Platform == UnrealTargetPlatform.Win64 && Target.Architecture == UnrealArch.X64) ||
 				Target.Platform == UnrealTargetPlatform.Mac ||
 				Target.Platform == UnrealTargetPlatform.Linux);
 			if (!bEnableUsdMemoryOverloads)
diff --git a/Engine/Plugins/Importers/USDImporterMDL/Source/USDImporterMDL/USDImporterMDL.Build.cs b/Engine/Plugins/Importers/USDImporterMDL/Source/USDImporterMDL/USDImporterMDL.Build.cs
--- a/Engine/Plugins/Importers/USDImporterMDL/Source/USDImporterMDL/USDImporterMDL.Build.cs
+++ b/Engine/Plugins/Importers/USDImporterMDL/Source/USDImporterMDL/USDImporterMDL.Build.cs
@@ -31,8 +31,8 @@ public class USDImporterMDL : ModuleRules
 		);
 
 		bool bUseUsdMemoryOverloadsIfSupported = true;
-		bool bEnableUsdMemoryOverloads = EnableUsdSdk(Target) && bUseUsdMemoryOverloadsIfSupported &&
-			(Target.Platform == UnrealTargetPlatform.Win64 && Target.Architecture == UnrealArch.X64 ||
+		bool bEnableUsdMemoryOverloads = (EnableUsdSdk(Target) && bUseUsdMemoryOverloadsIfSupported) &&
+			((Target.Platform == UnrealTargetPlatform.Win64 && Target.Architecture == UnrealArch.X64) ||
 			Target.Platform == UnrealTargetPlatform.Mac ||
 			Target.Platform == UnrealTargetPlatform.Linux);
 		if (!bEnableUsdMemoryOverloads)
diff --git a/Engine/Plugins/Runtime/USDCore/Source/USDUtilities/USDUtilities.Build.cs b/Engine/Plugins/Runtime/USDCore/Source/USDUtilities/USDUtilities.Build.cs
--- a/Engine/Plugins/Runtime/USDCore/Source/USDUtilities/USDUtilities.Build.cs
+++ b/Engine/Plugins/Runtime/USDCore/Source/USDUtilities/USDUtilities.Build.cs
@@ -86,8 +86,8 @@ namespace UnrealBuildTool.Rules
 			PublicDefinitions.Add("ENABLE_USD_MATERIALX=" + (bEnableUsdMaterialX ? "1" : "0"));
 
 			bool bUseUsdMemoryOverloadsIfSupported = true;
-			bool bEnableUsdMemoryOverloads = EnableUsdSdk(Target) && bUseUsdMemoryOverloadsIfSupported &&
-				(Target.Platform == UnrealTargetPlatform.Win64 && Target.Architecture == UnrealArch.X64 ||
+			bool bEnableUsdMemoryOverloads = (EnableUsdSdk(Target) && bUseUsdMemoryOverloadsIfSupported) &&
+				((Target.Platform == UnrealTargetPlatform.Win64 && Target.Architecture == UnrealArch.X64) ||
 				Target.Platform == UnrealTargetPlatform.Mac ||
 				Target.Platform == UnrealTargetPlatform.Linux);
 			if (!bEnableUsdMemoryOverloads)
diff --git a/Engine/Plugins/Runtime/USDCore/Source/UnrealUSDWrapper/UnrealUSDWrapper.Build.cs b/Engine/Plugins/Runtime/USDCore/Source/UnrealUSDWrapper/UnrealUSDWrapper.Build.cs
--- a/Engine/Plugins/Runtime/USDCore/Source/UnrealUSDWrapper/UnrealUSDWrapper.Build.cs
+++ b/Engine/Plugins/Runtime/USDCore/Source/UnrealUSDWrapper/UnrealUSDWrapper.Build.cs
@@ -25,8 +25,8 @@ namespace UnrealBuildTool.Rules
 			);
 
 			bool bUseUsdMemoryOverloadsIfSupported = true;
-			bool bEnableUsdMemoryOverloads = EnableUsdSdk(Target) && bUseUsdMemoryOverloadsIfSupported && 
-				(Target.Platform == UnrealTargetPlatform.Win64 && Target.Architecture == UnrealArch.X64 ||
+			bool bEnableUsdMemoryOverloads = (EnableUsdSdk(Target) && bUseUsdMemoryOverloadsIfSupported) && 
+				((Target.Platform == UnrealTargetPlatform.Win64 && Target.Architecture == UnrealArch.X64) ||
 				Target.Platform == UnrealTargetPlatform.Mac ||
 				Target.Platform == UnrealTargetPlatform.Linux);
 			PublicDefinitions.Add("ENABLE_USD_MEMORY_OVERLOADS=" + (bEnableUsdMemoryOverloads ? "1" : "0"));
