USD: Fix LevelSequences stomping each other when importing stages with multiple animated layers that have the same base filename.

diff --git a/Engine/Plugins/Importers/USDImporter/Source/USDStage/Private/USDLevelSequenceHelper.cpp b/Engine/Plugins/Importers/USDImporter/Source/USDStage/Private/USDLevelSequenceHelper.cpp
--- a/Engine/Plugins/Importers/USDImporter/Source/USDStage/Private/USDLevelSequenceHelper.cpp
+++ b/Engine/Plugins/Importers/USDImporter/Source/USDStage/Private/USDLevelSequenceHelper.cpp
@@ -864,7 +864,7 @@ private:
 																 // that each Sequence is only present once in the hierarchy.
 
 	// Sequence Name to Layer Identifier Map. Relationship: N Sequences to 1 Layer.
-	TMap<FName, FString> LayerIdentifierByLevelSequenceName;
+	TMap<FString, FString> LayerIdentifierByLevelSequenceName;
 
 	// Sections handling
 private:
@@ -1451,6 +1451,8 @@ ULevelSequence* FUsdLevelSequenceHelperImpl::FindOrAddSequenceForLayer(
 	const FString& SequenceDisplayName
 )
 {
+	using namespace UsdUnreal::ObjectUtils;
+
 	if (!Layer)
 	{
 		return nullptr;
@@ -1460,21 +1462,30 @@ ULevelSequence* FUsdLevelSequenceHelperImpl::FindOrAddSequenceForLayer(
 
 	if (!Sequence)
 	{
-		// This needs to be unique, or else when we reload the stage we will end up with a new ULevelSequence with the same class, outer and name as
-		// the previous one. Also note that the previous level sequence, even though unreferenced by the stage actor, is likely still alive and valid
-		// due to references from the transaction buffer, so we would basically end up creating a identical new object on top of an existing one (the
-		// new object has the same address as the existing one). When importing we don't actually want to do this though, because we want these assets
-		// name to conflict so that we can publish/replace old assets if desired. The stage importer will make these names unique later if needed. We
-		// only get a PrimLinkCache when importing (from UUsdStageImporter::ImportFromFile) or when BindToUsdStageActor is called, which also gives us
+		// We only get a PrimLinkCache when importing (from UUsdStageImporter::ImportFromFile) or when BindToUsdStageActor is called, which also gives us
 		// a stage actor. So if we don't have an actor but have a cache, we're importing
-
 		const bool bIsImporting = StageActor.IsExplicitlyNull() && PrimLinkCache.IsValid();
-		FName UniqueSequenceName = bIsImporting ? *UsdUnreal::ObjectUtils::SanitizeObjectName(FPaths::GetBaseFilename(SequenceDisplayName))
+
+		// Names need to be globally unique when opening the stage, or else when we reload the stage we will end up with a new ULevelSequence with the
+		// same class, outer and name as the previous one and would otherwise stomp it. The old LevelSequence is likely still alive and valid due to
+		// references from the transaction buffer, so we could run into trouble.
+		//
+		// When importing we don't want unique names across the whole transient package though, because we may randomly have old LevelSequences in the
+		// transient package already which would give us random suffixes. Instead, want these asset names to be consistent between different imports
+		// in order to replace previously imported assets, if desired. The stage importer will make the actual published asset names unique later if
+		// needed though.
+		//
+		// We still never want the LevelSequences to collide *within the same import* however, as we'd just end up stomping the same LevelSequences
+		// we're currently importing (see UE-283920), so here we use GetUniqueName instead
+		FName UniqueSequenceName = bIsImporting ? *GetUniqueName(
+													*SanitizeObjectName(FPaths::GetBaseFilename(SequenceDisplayName)),
+													LayerIdentifierByLevelSequenceName
+												)
 												: MakeUniqueObjectName(
-													  GetTransientPackage(),
-													  ULevelSequence::StaticClass(),
-													  *UsdUnreal::ObjectUtils::SanitizeObjectName(FPaths::GetBaseFilename(SequenceDisplayName))
-												  );
+													GetTransientPackage(),
+													ULevelSequence::StaticClass(),
+													*SanitizeObjectName(FPaths::GetBaseFilename(SequenceDisplayName))
+												);
 
 		Sequence = NewObject<ULevelSequence>(GetTransientPackage(), UniqueSequenceName, FUsdLevelSequenceHelperImpl::DefaultObjFlags);
 		Sequence->Initialize();
@@ -1485,7 +1496,7 @@ ULevelSequence* FUsdLevelSequenceHelperImpl::FindOrAddSequenceForLayer(
 			return nullptr;
 		}
 
-		LayerIdentifierByLevelSequenceName.Add(Sequence->GetFName(), Layer.GetIdentifier());
+		LayerIdentifierByLevelSequenceName.Add(Sequence->GetName(), Layer.GetIdentifier());
 		LevelSequencesByIdentifier.Add(SequenceIdentifier, Sequence);
 		IdentifierByLevelSequence.Add(Sequence, SequenceIdentifier);
 
@@ -1573,8 +1584,8 @@ void FUsdLevelSequenceHelperImpl::CreateSubSequenceSection(ULevelSequence& Seque
 		SubTrack = MovieScene->AddTrack<UMovieSceneSubTrack>();
 	}
 
-	const FString* LayerIdentifier = LayerIdentifierByLevelSequenceName.Find(Sequence.GetFName());
-	const FString* SubLayerIdentifier = LayerIdentifierByLevelSequenceName.Find(SubSequence.GetFName());
+	const FString* LayerIdentifier = LayerIdentifierByLevelSequenceName.Find(Sequence.GetName());
+	const FString* SubLayerIdentifier = LayerIdentifierByLevelSequenceName.Find(SubSequence.GetName());
 
 	if (!LayerIdentifier || !SubLayerIdentifier)
 	{
@@ -3914,8 +3925,8 @@ void FUsdLevelSequenceHelperImpl::UpdateUsdLayerOffsetFromSection(const UMovieSc
 		return;
 	}
 
-	const FString* LayerIdentifier = LayerIdentifierByLevelSequenceName.Find(Sequence->GetFName());
-	const FString* SubLayerIdentifier = LayerIdentifierByLevelSequenceName.Find(SubSequence->GetFName());
+	const FString* LayerIdentifier = LayerIdentifierByLevelSequenceName.Find(Sequence->GetName());
+	const FString* SubLayerIdentifier = LayerIdentifierByLevelSequenceName.Find(SubSequence->GetName());
 
 	if (!LayerIdentifier || !SubLayerIdentifier)
 	{
@@ -4381,7 +4392,7 @@ void FUsdLevelSequenceHelperImpl::HandleMovieSceneChange(UMovieScene& MovieScene
 		return;
 	}
 
-	const FString LayerIdentifier = LayerIdentifierByLevelSequenceName.FindRef(Sequence->GetFName());
+	const FString LayerIdentifier = LayerIdentifierByLevelSequenceName.FindRef(Sequence->GetName());
 	FLayerTimeInfo* LayerTimeInfo = LayerTimeInfosByLayerIdentifier.Find(LayerIdentifier);
 	if (!LayerTimeInfo)
 	{
diff --git a/Engine/Plugins/Runtime/USDCore/Source/USDClasses/Private/USDObjectUtils.cpp b/Engine/Plugins/Runtime/USDCore/Source/USDClasses/Private/USDObjectUtils.cpp
--- a/Engine/Plugins/Runtime/USDCore/Source/USDClasses/Private/USDObjectUtils.cpp
+++ b/Engine/Plugins/Runtime/USDCore/Source/USDClasses/Private/USDObjectUtils.cpp
@@ -404,28 +404,3 @@ bool UsdUnreal::ObjectUtils::RemoveNumberedSuffix(FString& Prefix)
 
 	return bRemoved;
 }
-
-FString UsdUnreal::ObjectUtils::GetUniqueName(FString Name, const TSet<FString>& UsedNames)
-{
-	if (!UsedNames.Contains(Name))
-	{
-		return Name;
-	}
-
-	const bool bRemoved = RemoveNumberedSuffix(Name);
-
-	// Its possible that removing the suffix made it into a unique name already
-	if (bRemoved && !UsedNames.Contains(Name))
-	{
-		return Name;
-	}
-
-	int32 Suffix = 0;
-	FString Result;
-	do
-	{
-		Result = FString::Printf(TEXT("%s_%d"), *Name, Suffix++);
-	} while (UsedNames.Contains(Result));
-
-	return Result;
-}
diff --git a/Engine/Plugins/Runtime/USDCore/Source/USDClasses/Public/USDObjectUtils.h b/Engine/Plugins/Runtime/USDCore/Source/USDClasses/Public/USDObjectUtils.h
--- a/Engine/Plugins/Runtime/USDCore/Source/USDClasses/Public/USDObjectUtils.h
+++ b/Engine/Plugins/Runtime/USDCore/Source/USDClasses/Public/USDObjectUtils.h
@@ -70,8 +70,32 @@ namespace UsdUnreal::ObjectUtils
 	 * Appends numbered suffixes to Name until the result is not contained in UsedNames, and returns it.
 	 * Does not add the result to UsedNames before returning (as it is const).
 	 * @param Name - Received string to make unique (e.g. "MyName")
-	 * @param UsedNames - Strings that cannot be used for the result
+	 * @param UsedNames - Strings that cannot be used for the result. Should have a .Contains() function that receives FString
 	 * @return Modified Name so that it doesn't match anything in UsedNames (e.g. "MyName" again, or "MyName_0" or "MyName_423")
 	 */
-	USDCLASSES_API FString GetUniqueName(FString Name, const TSet<FString>& UsedNames);
+	template<typename StringContainer>
+	inline FString GetUniqueName(FString Name, const StringContainer& UsedNames)
+	{
+		if (!UsedNames.Contains(Name))
+		{
+			return Name;
+		}
+
+		const bool bRemoved = RemoveNumberedSuffix(Name);
+
+		// It's possible that removing the suffix made it into a unique name already
+		if (bRemoved && !UsedNames.Contains(Name))
+		{
+			return Name;
+		}
+
+		int32 Suffix = 0;
+		FString Result;
+		do
+		{
+			Result = FString::Printf(TEXT("%s_%d"), *Name, Suffix++);
+		} while (UsedNames.Contains(Result));
+
+		return Result;
+	}
 };	  // namespace UsdUnreal::ObjectUtils
