USD: Discard empty LevelSequences when importing;

Fix imported LevelSequences being left as ReadOnly when importing.

diff --git a/Engine/Plugins/Importers/USDImporter/Source/USDStageImporter/Private/USDStageImporter.cpp b/Engine/Plugins/Importers/USDImporter/Source/USDStageImporter/Private/USDStageImporter.cpp
--- a/Engine/Plugins/Importers/USDImporter/Source/USDStageImporter/Private/USDStageImporter.cpp
+++ b/Engine/Plugins/Importers/USDImporter/Source/USDStageImporter/Private/USDStageImporter.cpp
@@ -56,13 +56,16 @@
 #include "Misc/RedirectCollector.h"
 #include "Misc/ScopedSlowTask.h"
 #include "Modules/ModuleManager.h"
+#include "MovieScene.h"
 #include "PackageTools.h"
 #include "PhysicsEngine/PhysicsAsset.h"
+#include "Sections/MovieSceneSubSection.h"
 #include "Serialization/ArchiveReplaceObjectRef.h"
 #include "Sound/SoundBase.h"
 #include "Sound/SoundWave.h"
 #include "SparseVolumeTexture/SparseVolumeTexture.h"
 #include "Subsystems/AssetEditorSubsystem.h"
+#include "Tracks/MovieSceneSubTrack.h"
 #include "UObject/UObjectIterator.h"
 
 #define LOCTEXT_NAMESPACE "USDStageImporter"
@@ -767,6 +770,106 @@ namespace UE::USDStageImporter::Private
 		}
 	}
 
+	void DiscardEmptyLevelSequences(TArray<UObject*>& InOutLevelSequences)
+	{
+		// Maps from a visited LevelSequence to whether it has useful data or not
+		TMap<UObject*, bool> SequenceHasData;
+		SequenceHasData.Reserve(InOutLevelSequences.Num());
+
+		TFunction<bool(UMovieSceneSequence*)> TraverseSequence = nullptr;
+		TraverseSequence = [&SequenceHasData, &TraverseSequence](UMovieSceneSequence* Sequence) -> bool
+		{
+			if (!Sequence)
+			{
+				return false;
+			}
+
+			if (bool* FoundResult = SequenceHasData.Find(Sequence))
+			{
+				return *FoundResult;
+			}
+
+			bool bHasData = false;
+
+			// Preemptively add an entry to the map in order to prevent ourselves from getting stuck traversing
+			// loops, if those are possible
+			SequenceHasData.Add(Sequence, bHasData);
+
+			if (UMovieScene* MovieScene = Sequence->GetMovieScene())
+			{
+				// When opening the stage we may mark some non-local-layer-stack sequences as read-only.
+				// There's no point to doing this when importing, however. It's easier to strip this here than to 
+				// prevent the LevelSequences from being made read-only when first importing
+				const bool bReadOnly = false;
+				MovieScene->SetReadOnly(bReadOnly);
+
+				bool bHasSubsectionWithData = false;
+				if (UMovieSceneSubTrack* SubTrack = MovieScene->FindTrack<UMovieSceneSubTrack>())
+				{
+					TArray<UMovieSceneSection*> Sections = SubTrack->GetAllSections();
+					for (UMovieSceneSection* Section : Sections)
+					{
+						UMovieSceneSubSection* SubSection = Cast<UMovieSceneSubSection>(Section);
+						if (!SubSection)
+						{
+							// This should never happen, but if it does we have no idea what this is, so lets not remove it
+							bHasSubsectionWithData = true;
+							continue;
+						}
+
+						UMovieSceneSequence* Dependency = SubSection->GetSequence();
+						const bool bDependencyHasData = TraverseSequence(Dependency);
+						if (bDependencyHasData)
+						{
+							bHasSubsectionWithData = true;
+							continue;
+						}
+						else
+						{
+							SubSection->SetSequence(nullptr);
+							SubTrack->RemoveSection(*Section);
+						}
+					}
+
+					if (SubTrack->GetAllSections().Num() == 0)
+					{
+						MovieScene->RemoveTrack(*SubTrack);
+					}
+				}
+
+				// TODO: Remove empty bindings too?
+				bool bHasBindingsWithTracks = false;
+				TArray<FMovieSceneBinding> Bindings = MovieScene->GetBindings();
+				for (FMovieSceneBinding& Binding : Bindings)
+				{
+					if (Binding.GetTracks().Num() != 0)
+					{
+						bHasBindingsWithTracks = true;
+						break;
+					}
+				}
+
+				bHasData = bHasSubsectionWithData || bHasBindingsWithTracks;
+			}
+
+			SequenceHasData.Add(Sequence, bHasData);
+			return bHasData;
+		};
+
+		for (UObject* Object : InOutLevelSequences)
+		{
+			TraverseSequence(Cast<ULevelSequence>(Object));
+		}
+
+		InOutLevelSequences.RemoveAll(
+			[&SequenceHasData](UObject* Sequence)
+			{
+				// FindRef returns default bool (false) if Sequence is not in the map
+				return !SequenceHasData.FindRef(Sequence);
+			}
+		);
+	}
+
 	// Move imported assets from transient folder to their final package, updating AssetCache to point to the moved assets
 	void PublishAssets(
 		FUsdStageImportContext& ImportContext,
@@ -947,6 +1050,8 @@ namespace UE::USDStageImporter::Private
 			}
 		};
 
+		DiscardEmptyLevelSequences(LevelSequences);
+
 		// We should publish them in the order of dependents before dependencies, because publishing will involve
 		// taking these assets from the asset cache, which can only happen in that order. If we don't follow this order
 		// we may end up unnecessarily duplicating assets
