USD: Refactor of LevelSequenceHelper in order to properly support Subsections due to references and payload arcs;

Fix animation being lost when importing stages with multiple layers that have the same filename;
Fix placement of geometry cache sections within LevelSequences;
Fix subsections having incorrect start/end offsets when animated prims did not have animated transforms;

diff --git a/Engine/Plugins/Importers/USDImporter/Source/USDSchemas/Private/USDGeometryCacheTranslator.cpp b/Engine/Plugins/Importers/USDImporter/Source/USDSchemas/Private/USDGeometryCacheTranslator.cpp
--- a/Engine/Plugins/Importers/USDImporter/Source/USDSchemas/Private/USDGeometryCacheTranslator.cpp
+++ b/Engine/Plugins/Importers/USDImporter/Source/USDSchemas/Private/USDGeometryCacheTranslator.cpp
@@ -382,6 +382,8 @@ namespace UsdGeometryCacheTranslatorImpl
 		FString RootPrimPath = RootPrim.GetPrimPath().GetString();
 		FReadMeshDataArgs Args(GetReadMeshDataArgs(Context, RootPrimPath));
 
+		StartOffsetTime = static_cast<float>(Args.StartFrame) / Args.FramesPerSecond;
+
 		// Compute the asset hash from the merged mesh description
 		FSHA1 SHA1;
 		FSHAHash MeshHash = FStaticMeshOperations::ComputeSHAHash(MeshDescription);
@@ -464,11 +466,7 @@ namespace UsdGeometryCacheTranslatorImpl
 					}
 				);
 			}
-			if (!bIsImporting)
-			{
-				// StartOffsetTime is the offset applied to the GeometryCache section on the sequencer track, so not relevant when importing
-				StartOffsetTime = static_cast<float>(Args.StartFrame) / Args.FramesPerSecond;
-			}
+
 			GeometryCache->SetFrameStartEnd(Args.StartFrame, Args.EndFrame - 1);
 
 			// Create a track for each mesh to be processed and add it to the GeometryCache
@@ -861,8 +859,10 @@ void FGeometryCacheCreateAssetsTaskChain::SetupTasks()
 				);
 				if (UserData)
 				{
+					// TODO: We're overwriting some data even if we reuse the asset from the asset cache,
+					// but then again we do that for the other asset types too...
 					UserData->PrimvarToUVIndex = LODIndexToMaterialInfo[0].PrimvarToUVIndex;	// We use the same primvar mapping for all LODs
-					UserData->LayerStartOffsetSeconds = StartTimeOffset;
+					UserData->StageStartOffsetSeconds = StartTimeOffset;
 					UserData->PrimPaths.AddUnique(PrimPath.GetString());
 
 					if (Context->MetadataOptions.bCollectMetadata)
@@ -1111,19 +1111,19 @@ void FUsdGeometryCacheTranslator::UpdateComponents(USceneComponent* SceneCompone
 			const bool bIsLooping = false;
 			GeometryCacheComponent->SetLooping(bIsLooping);
 
-			float LayerStartOffsetSeconds = 0.0f;
+			float StageStartOffsetSeconds = 0.0f;
 			if (GeometryCache)
 			{
 				if (UUsdGeometryCacheAssetUserData* UserData = GeometryCache->GetAssetUserData<UUsdGeometryCacheAssetUserData>())
 				{
-					LayerStartOffsetSeconds = UserData->LayerStartOffsetSeconds;
+					StageStartOffsetSeconds = UserData->StageStartOffsetSeconds;
 				}
 			}
 
 			// The Time from the stage has to be adjusted to be relative to the time range of the geometry cache
 			// by applying the start offset. Thus, the adjusted time has to be clamped between 0 and the duration.
 			const double FramesPerSecond = Context->Stage.GetTimeCodesPerSecond();
-			float AdjustedTime = static_cast<float>(TimeCode / FramesPerSecond - LayerStartOffsetSeconds);
+			float AdjustedTime = static_cast<float>(TimeCode / FramesPerSecond - StageStartOffsetSeconds);
 			if (GeometryCache)
 			{
 				const float Duration = GeometryCache->CalculateDuration();
diff --git a/Engine/Plugins/Importers/USDImporter/Source/USDSchemas/Private/USDSkelRootTranslator.cpp b/Engine/Plugins/Importers/USDImporter/Source/USDSchemas/Private/USDSkelRootTranslator.cpp
--- a/Engine/Plugins/Importers/USDImporter/Source/USDSchemas/Private/USDSkelRootTranslator.cpp
+++ b/Engine/Plugins/Importers/USDImporter/Source/USDSchemas/Private/USDSkelRootTranslator.cpp
@@ -1633,7 +1633,9 @@ namespace UsdSkelRootTranslatorImpl
 									TEXT("USDAssetUserData")
 								);
 								UserData->PrimPaths = {SkelAnimationPrimPath};
+								PRAGMA_DISABLE_DEPRECATION_WARNINGS
 								UserData->LayerStartOffsetSeconds = LayerStartOffsetSeconds;
+								PRAGMA_ENABLE_DEPRECATION_WARNINGS
 
 								AnimSequence->AddAssetUserData(UserData);
 
diff --git a/Engine/Plugins/Importers/USDImporter/Source/USDSchemas/Private/USDSkelSkeletonTranslator.cpp b/Engine/Plugins/Importers/USDImporter/Source/USDSchemas/Private/USDSkelSkeletonTranslator.cpp
--- a/Engine/Plugins/Importers/USDImporter/Source/USDSchemas/Private/USDSkelSkeletonTranslator.cpp
+++ b/Engine/Plugins/Importers/USDImporter/Source/USDSchemas/Private/USDSkelSkeletonTranslator.cpp
@@ -1833,7 +1833,8 @@ namespace UsdSkelSkeletonTranslatorImpl
 					&bAnimSequenceIsNew
 				);
 
-				TOptional<float> LayerStartOffsetSeconds;
+				TOptional<float> StageStartOffsetSeconds;
+				TOptional<float> StageScaleSeconds;
 				if (bAnimSequenceIsNew || AnimSequence->GetSkeleton() != SkeletalMesh->GetSkeleton())
 				{
 					FScopedUnrealAllocs UEAllocs;
@@ -1844,7 +1845,8 @@ namespace UsdSkelSkeletonTranslatorImpl
 					AnimSequence->SetSkeleton(SkeletalMesh->GetSkeleton());
 
 					TUsdStore<pxr::VtArray<pxr::UsdSkelSkinningQuery>> SkinningTargets = SkeletonBinding.Get().GetSkinningTargets();
-					LayerStartOffsetSeconds = 0.0f;
+					StageStartOffsetSeconds = 0.0f;
+					StageScaleSeconds = 1.0f;
 					bool bSuccess = UsdToUnreal::ConvertSkelAnim(
 						SkeletonQuery.Get(),
 						&SkinningTargets.Get(),
@@ -1852,7 +1854,8 @@ namespace UsdSkelSkeletonTranslatorImpl
 						Context->bAllowInterpretingLODs,
 						RootMotionPrim,
 						AnimSequence,
-						&LayerStartOffsetSeconds.GetValue()
+						&StageStartOffsetSeconds.GetValue(),
+						&StageScaleSeconds.GetValue()
 					);
 
 					if (!bSuccess
@@ -1870,13 +1873,17 @@ namespace UsdSkelSkeletonTranslatorImpl
 				{
 					UserData->PrimPaths.AddUnique(SkelAnimationPrimPath);
 
-					// It should be fine that we won't fetch/set this again in case we're reusing an AnimSequence from the asset cache
-					// because the LayerStartOffsetSeconds value is influenced by the animation start time codes, which should always
-					// affect the asset hash. This means that if the correct value for this were to change, we'd end up generating a
-					// new AnimSequence and computing it anyway
-					if (LayerStartOffsetSeconds.IsSet())
+					// It should be fine that we won't fetch/set these again in case we're reusing an AnimSequence from the asset cache
+					// because they are influenced by the animation start/end time codes, which should always affect the asset hash.
+					// This means that if the correct value for this were to change, we'd end up generating a new AnimSequence and computing 
+					// it anyway
+					if (StageStartOffsetSeconds.IsSet())
 					{
-						UserData->LayerStartOffsetSeconds = LayerStartOffsetSeconds.GetValue();
+						UserData->StageStartOffsetSeconds = StageStartOffsetSeconds.GetValue();
+					}
+					if (StageScaleSeconds.IsSet())
+					{
+						UserData->StageScaleSeconds = StageScaleSeconds.GetValue();
 					}
 
 					if (Context->MetadataOptions.bCollectMetadata)
@@ -2152,32 +2159,23 @@ void FUsdSkelSkeletonTranslator::UpdateComponents(USceneComponent* SceneComponen
 	{
 		if (UAnimSequence* AnimSequence = Cast<UAnimSequence>(SkeletalMeshComponent->AnimationData.AnimToPlay.Get()))
 		{
-			UE::FSdfLayerOffset CombinedOffset;
-			if (SkelAnimPrim)
-			{
-				CombinedOffset = UsdUtils::GetPrimToStageOffset(SkelAnimPrim);
-			}
-
-			double LayerStartOffsetSeconds = 0.0f;
+			double StageStartOffsetSeconds = 0.0f;
+			double StageScaleSeconds = 1.0f;
 			if (UUsdAnimSequenceAssetUserData* UserData = AnimSequence->GetAssetUserData<UUsdAnimSequenceAssetUserData>())
 			{
-				LayerStartOffsetSeconds = UserData->LayerStartOffsetSeconds;
+				StageStartOffsetSeconds = UserData->StageStartOffsetSeconds;
+				StageScaleSeconds = UserData->StageScaleSeconds;
 			}
 
 			// Always change the mode here because the sequencer will change it back to AnimationCustomMode when animating
 			SkeletalMeshComponent->SetAnimationMode(EAnimationMode::AnimationSingleNode);
 
-			// Part of the CombinedOffset will be due to a framerate difference. We don't care about that part here though, so remove it
-			const double TimeCodesPerSecondDifference = Context->Stage.GetTimeCodesPerSecond() / AnimSequence->ImportFileFramerate;
-			CombinedOffset.Scale /= TimeCodesPerSecondDifference;
-
-			// Always use the sequence's framerate here because we need to sample the UAnimSequence with in seconds, and that
+			// Always use the sequence's framerate here because we need to sample the UAnimSequence with seconds, and that
 			// asset may have been created when the stage had a different framesPerSecond (and was reused by the assets cache)
 			// Use the import framerate here because we will need to change the sampling framerate of the sequence in order to get it
 			// to match the target duration in seconds and the number of source frames.
-			const double LayerTimeCode = ((Context->Time - CombinedOffset.Offset) / CombinedOffset.Scale);
-			const double AnimSequenceTime = LayerTimeCode / AnimSequence->ImportFileFramerate;
-			SkeletalMeshComponent->SetPosition(AnimSequenceTime - LayerStartOffsetSeconds);
+			const double AnimSequenceTime = Context->Time / AnimSequence->ImportFileFramerate;
+			SkeletalMeshComponent->SetPosition((AnimSequenceTime - StageStartOffsetSeconds) / StageScaleSeconds);
 
 			SkeletalMeshComponent->TickAnimation(0.f, false);
 			SkeletalMeshComponent->RefreshBoneTransforms();
diff --git a/Engine/Plugins/Importers/USDImporter/Source/USDStage/Private/USDLevelSequenceHelper.cpp b/Engine/Plugins/Importers/USDImporter/Source/USDStage/Private/USDLevelSequenceHelper.cpp
--- a/Engine/Plugins/Importers/USDImporter/Source/USDStage/Private/USDLevelSequenceHelper.cpp
+++ b/Engine/Plugins/Importers/USDImporter/Source/USDStage/Private/USDLevelSequenceHelper.cpp
@@ -20,13 +20,21 @@
 #include "USDTypesConversion.h"
 #include "USDValueConversion.h"
 
+#include "UsdWrappers/PcpLayerStack.h"
+#include "UsdWrappers/PcpLayerStackIdentifier.h"
+#include "UsdWrappers/PcpMapExpression.h"
+#include "UsdWrappers/PcpNodeRef.h"
 #include "UsdWrappers/SdfChangeBlock.h"
 #include "UsdWrappers/SdfLayer.h"
+#include "UsdWrappers/SdfPayloadsProxy.h"
+#include "UsdWrappers/SdfPrimSpec.h"
+#include "UsdWrappers/SdfReferencesProxy.h"
 #include "UsdWrappers/UsdAttribute.h"
 #include "UsdWrappers/UsdEditContext.h"
 #include "UsdWrappers/UsdGeomBBoxCache.h"
 #include "UsdWrappers/UsdGeomXformable.h"
 #include "UsdWrappers/UsdPrim.h"
+#include "UsdWrappers/UsdResolveInfo.h"
 #include "UsdWrappers/UsdStage.h"
 #include "UsdWrappers/UsdTyped.h"
 
@@ -780,18 +788,20 @@ public:
 	void Clear();
 
 private:
-	struct FLayerOffsetInfo
-	{
-		FString LayerIdentifier;
-		UE::FSdfLayerOffset LayerOffset;
-	};
-
 	struct FLayerTimeInfo
 	{
 		FString Identifier;
 		FString FilePath;
 
-		TArray<FLayerOffsetInfo> SubLayersOffsets;
+		// Time offsets for LevelSequence subsections that correspond to sublayers
+		TMap<FString, UE::FSdfLayerOffset> SublayerIdentifierToOffsets;
+
+		// Time offsets for LevelSequence subsections that correspond to prim reference composition arcs
+		// TODO: Instead of layer identifier, use a struct with layer identifier + referencer prim path
+		TMap<FString, UE::FSdfLayerOffset> ReferenceIdentifierToOffsets;
+
+		// Time offsets for LevelSequence subsections that correspond to prim payload composition arcs
+		TMap<FString, UE::FSdfLayerOffset> PayloadIdentifierToOffsets;
 
 		TOptional<double> StartTimeCode;
 		TOptional<double> EndTimeCode;
@@ -813,8 +823,7 @@ protected:
 	// Sequences handling
 public:
 
-	/** Creates a Level Sequence and its SubSequenceSection for each layer in the local layer stack (root layer and sub layers) */
-	void CreateLocalLayersSequences();
+	void CreateSequencesForLayerStack(const UE::FSdfLayer& RootLayer);
 	void BindToUsdStageActor(AUsdStageActor* InStageActor);
 	void UnbindFromUsdStageActor();
 	EUsdRootMotionHandling GetRootMotionHandling() const;
@@ -842,11 +851,19 @@ public:
 		return OnSkelAnimationBaked;
 	}
 
+	void DirtySequenceHierarchyCache();
+	const FMovieSceneSequenceTransform& GetSubsequenceTransform(const ULevelSequence* Sequence);
+
 private:
-	ULevelSequence* FindSequenceForAttribute(const UE::FUsdAttribute& Attribute);
 	ULevelSequence* FindOrAddSequenceForAttribute(const UE::FUsdAttribute& Attribute, UE::FSdfLayer* OutSequenceLayer = nullptr);
-	ULevelSequence* FindSequenceForIdentifier(const FString& SequenceIdentitifer);
-	ULevelSequence* FindOrAddSequenceForLayer(const UE::FSdfLayer& Layer, const FString& SequenceIdentifier, const FString& SequenceDisplayName);
+
+	// WARNING: These two functions will not create the required parent LevelSequences and SubSequence sections, and are only 
+	// meant for edge cases/exceptions or internal use
+	ULevelSequence* FindSequenceForIdentifier(const FString& SequenceIdentitifer) const;
+	ULevelSequence* FindOrAddSequenceForLayer(const UE::FSdfLayer& Layer, bool* bOutCreatedNew = nullptr);
+
+	ULevelSequence* CreateSequenceAndParentsRecursively(const UE::FPcpNodeRef& PcpNode, const UE::FSdfLayer& Layer);
+
 	UE::FSdfLayer FindEditTargetForSubsequence(ULevelSequence* Sequence);
 
 	/** Removes PrimTwin as a user of Sequence. If Sequence is now unused, remove its subsection and itself. */
@@ -859,9 +876,12 @@ private:
 
 	TSet<FName> LocalLayersSequences;							 // List of sequences associated with sublayers
 
-	FMovieSceneSequenceHierarchy SequenceHierarchyCache;		 // Cache for the hierarchy of level sequences and subsections
+	bool bSequenceHierarchyCacheIsDirty = false;
+	FMovieSceneSequenceHierarchy SequenceHierarchyCache;		 // Cache for the hierarchy of level sequences and subsections. Go through GetSubsequenceTransform()
+																 // instead of using this directly.
 	TMap<ULevelSequence*, FMovieSceneSequenceID> SequencesID;	 // Tracks the FMovieSceneSequenceID for each Sequence in the hierarchy. We assume
 																 // that each Sequence is only present once in the hierarchy.
+																 // TODO: This will need to change if we support multiple subsections for the same sequence
 
 	// Sequence Name to Layer Identifier Map. Relationship: N Sequences to 1 Layer.
 	TMap<FString, FString> LayerIdentifierByLevelSequenceName;
@@ -870,7 +890,7 @@ private:
 private:
 	/** Returns the UMovieSceneSubSection associated with SubSequence on the Sequence UMovieSceneSubTrack if it exists */
 	UMovieSceneSubSection* FindSubSequenceSection(ULevelSequence& Sequence, ULevelSequence& SubSequence);
-	void CreateSubSequenceSection(ULevelSequence& Sequence, ULevelSequence& Subsequence);
+	void CreateSubSequenceSection(ULevelSequence& Sequence, ULevelSequence& Subsequence, const UE::FSdfLayerOffset& SubLayerOffset);
 	void RemoveSubSequenceSection(ULevelSequence& Sequence, ULevelSequence& SubSequence);
 
 	/**
@@ -917,7 +937,8 @@ public:
 	void UpdateControlRigTracks(UUsdPrimTwin& PrimTwin);
 
 private:
-	// Sequence Name to Prim Path. Relationship: 1 Sequence to N Prim Path.
+	// Used to track which LevelSequences are being used by which prim.
+	// TODO: This was mostly used by a legacy way of guessing reference/payload offsets, and could maybe be removed?
 	TMultiMap<FName, FString> PrimPathByLevelSequenceName;
 
 	struct FPrimTwinBindings
@@ -1059,14 +1080,19 @@ FUsdLevelSequenceHelperImpl::~FUsdLevelSequenceHelperImpl()
 ULevelSequence* FUsdLevelSequenceHelperImpl::Init(const UE::FUsdStage& InUsdStage)
 {
 	UsdStage = InUsdStage;
+	const UE::FSdfLayer RootLayer = UsdStage.GetRootLayer();
 
 	Clear();
 
-	CreateLocalLayersSequences();
+	// Set this first as we'll need our MainLevelSequence setup before our call to CreateSubSequenceSection, as
+	// that involves calling UMovieSceneCompiledDataManager::CompileHierarchy() with it
+	MainLevelSequence = FindOrAddSequenceForLayer(RootLayer);
+
+	CreateSequencesForLayerStack(RootLayer);
 
 	// We call Init with a default (invalid) stage to "clear", so only subscribe to these events
 	// if we've actually been given a valid stage
-	if (InUsdStage)
+	if (UsdStage)
 	{
 		SubscribeToEditorEvents();
 	}
@@ -1093,12 +1119,6 @@ bool FUsdLevelSequenceHelperImpl::Serialize(FArchive& Ar)
 	Ar << RootMotionHandling;
 	Ar << StageActorBinding;
 
-	// Always keep SequenceHierarchyCache up-to-date given that it can't be serialized itself
-	if (Ar.IsLoading() && MainLevelSequence && MainLevelSequence->GetMovieScene())
-	{
-		UMovieSceneCompiledDataManager::CompileHierarchy(MainLevelSequence, &SequenceHierarchyCache, EMovieSceneServerClientMask::All);
-	}
-
 	return true;
 }
 
@@ -1171,66 +1191,67 @@ void FUsdLevelSequenceHelperImpl::AddReferencedObjects(FReferenceCollector& Coll
 	Collector.AddReferencedObjects(LevelSequencesByIdentifier);
 }
 
-void FUsdLevelSequenceHelperImpl::CreateLocalLayersSequences()
+void FUsdLevelSequenceHelperImpl::CreateSequencesForLayerStack(const UE::FSdfLayer& RootLayer)
 {
 	if (!UsdStage)
 	{
 		return;
 	}
 
-	UE::FSdfLayer RootLayer = UsdStage.GetRootLayer();
-	const FLayerTimeInfo& RootLayerInfo = FindOrAddLayerTimeInfo(RootLayer);
-
-	// Create main level sequence for root layer
-	MainLevelSequence = FindOrAddSequenceForLayer(RootLayer, RootLayer.GetIdentifier(), RootLayer.GetDisplayName());
-	if (!MainLevelSequence)
-	{
-		return;
-	}
-
-	UMovieScene* MovieScene = MainLevelSequence->GetMovieScene();
-	if (!MovieScene)
+	ULevelSequence* RootSequence = FindOrAddSequenceForLayer(RootLayer);
+	if (!RootSequence)
 	{
 		return;
 	}
 
-	SequencesID.Add(MainLevelSequence) = MovieSceneSequenceID::Root;
-
-	LocalLayersSequences.Add(MainLevelSequence->GetFName());
+	LocalLayersSequences.Add(RootSequence->GetFName());
 
-	TFunction<void(const FLayerTimeInfo* LayerTimeInfo, ULevelSequence& ParentSequence)> RecursivelyCreateSequencesForLayer;
-	RecursivelyCreateSequencesForLayer =
-		[&RecursivelyCreateSequencesForLayer, this](const FLayerTimeInfo* LayerTimeInfo, ULevelSequence& ParentSequence)
+	TFunction<void(const UE::FSdfLayer&, ULevelSequence&)> RecursivelyCreateSequencesForLayer;
+	RecursivelyCreateSequencesForLayer = [&RecursivelyCreateSequencesForLayer, this](const UE::FSdfLayer& Layer, ULevelSequence& Sequence)
 	{
-		if (!LayerTimeInfo)
+		UMovieScene* MovieScene = Sequence.GetMovieScene();
+		if (!MovieScene)
 		{
 			return;
 		}
 
-		if (UE::FSdfLayer Layer = UE::FSdfLayer::FindOrOpen(*LayerTimeInfo->Identifier))
+		for (const FString& SubLayerPath : Layer.GetSubLayerPaths())
 		{
-			for (const FString& SubLayerPath : Layer.GetSubLayerPaths())
+			UE::FSdfLayer SubLayer = UsdUtils::FindLayerForSubLayerPath(Layer, SubLayerPath);
+			if (!SubLayer)
 			{
-				if (UE::FSdfLayer SubLayer = UsdUtils::FindLayerForSubLayerPath(Layer, SubLayerPath))
-				{
-					if (ULevelSequence* SubSequence = FindOrAddSequenceForLayer(SubLayer, SubLayer.GetIdentifier(), SubLayer.GetDisplayName()))
-					{
-						if (!LocalLayersSequences.Contains(SubSequence->GetFName()))	// Make sure we don't parse an already parsed layer
-						{
-							LocalLayersSequences.Add(SubSequence->GetFName());
+				continue;
+			}
 
-							CreateSubSequenceSection(ParentSequence, *SubSequence);
+			ULevelSequence* SubSequence = FindOrAddSequenceForLayer(SubLayer);
+			if (!SubSequence)
+			{
+				continue;
+			}
 
-							RecursivelyCreateSequencesForLayer(FindLayerTimeInfo(SubLayer), *SubSequence);
-						}
-					}
-				}
+			// Make sure we don't parse an already parsed layer
+			if (LocalLayersSequences.Contains(SubSequence->GetFName()))
+			{
+				continue;
+			}
+			LocalLayersSequences.Add(SubSequence->GetFName());
+
+			// We can't reuse these between calls as we may invalidate these two pointers with the recursive calls
+			const FLayerTimeInfo* LayerTimeInfo = FindLayerTimeInfo(Layer);
+			const UE::FSdfLayerOffset* Offset = LayerTimeInfo->SublayerIdentifierToOffsets.Find(SubLayer.GetIdentifier());
+			if (!Offset)
+			{
+				continue;
 			}
+
+			CreateSubSequenceSection(Sequence, *SubSequence, *Offset);
+
+			RecursivelyCreateSequencesForLayer(SubLayer, *SubSequence);
 		}
 	};
 
 	// Create level sequences for all sub layers (accessible via the main level sequence but otherwise hidden)
-	RecursivelyCreateSequencesForLayer(&RootLayerInfo, *MainLevelSequence);
+	RecursivelyCreateSequencesForLayer(RootLayer, *RootSequence);
 }
 
 void FUsdLevelSequenceHelperImpl::BindToUsdStageActor(AUsdStageActor* InStageActor)
@@ -1385,71 +1406,85 @@ void FUsdLevelSequenceHelperImpl::UnsubscribeToEditorEvents()
 #endif	  // WITH_EDITOR
 }
 
-ULevelSequence* FUsdLevelSequenceHelperImpl::FindSequenceForAttribute(const UE::FUsdAttribute& Attribute)
+void FUsdLevelSequenceHelperImpl::DirtySequenceHierarchyCache()
+{
+	bSequenceHierarchyCacheIsDirty = true;
+}
+
+const FMovieSceneSequenceTransform& FUsdLevelSequenceHelperImpl::GetSubsequenceTransform(const ULevelSequence* Sequence)
 {
-	if (!Attribute || !Attribute.GetPrim())
+	const static FMovieSceneSequenceTransform IdentityTransform;
+
+	if (!Sequence)
 	{
-		return nullptr;
+		return IdentityTransform;
 	}
 
-	if (!UsdStage)
+	if (bSequenceHierarchyCacheIsDirty)
 	{
-		return nullptr;
+		bSequenceHierarchyCacheIsDirty = false;
+
+		SequenceHierarchyCache = FMovieSceneSequenceHierarchy();
+		UMovieSceneCompiledDataManager::CompileHierarchy(MainLevelSequence, &SequenceHierarchyCache, EMovieSceneServerClientMask::All);
+
+		SequencesID.Reset();
+		SequencesID.Add(MainLevelSequence) = MovieSceneSequenceID::Root;
+		for (const TTuple<FMovieSceneSequenceID, FMovieSceneSubSequenceData>& Pair : SequenceHierarchyCache.AllSubSequenceData())
+		{
+			if (ULevelSequence* LevelSequence = Cast<ULevelSequence>(Pair.Value.GetSequence()))
+			{
+				SequencesID.Add(LevelSequence, Pair.Key);
+			}
+		}
 	}
 
-	UE::FSdfLayer AttributeLayer = UsdUtils::FindLayerForAttribute(Attribute, 0.0);
-	if (!AttributeLayer)
+	// Early out here to also catch the case where we just have MainLevelSequence and no subsequences. In that case our SequencesID
+	// will be empty (as we never created a subsection) and so will our SequenceHierarchyCache
+	if (Sequence == MainLevelSequence)
 	{
-		return nullptr;
+		return IdentityTransform;
 	}
 
-	UE::FUsdPrim Prim = Attribute.GetPrim();
-
-	ULevelSequence* Sequence = nullptr;
-
-	// If the attribute is on the Root or a SubLayer, return the Sequence associated with that layer
-	if (AttributeLayer.HasSpec(Prim.GetPrimPath()) && UsdStage.HasLocalLayer(AttributeLayer))
+	FMovieSceneSequenceID SequenceID = SequencesID.FindRef(Sequence);
+	if (const FMovieSceneSubSequenceData* FoundData = SequenceHierarchyCache.FindSubData(SequenceID))
 	{
-		Sequence = FindSequenceForIdentifier(AttributeLayer.GetIdentifier());
+		return FoundData->RootToSequenceTransform;
 	}
-	// The prim should have its own sequence, return that
 	else
 	{
-		Sequence = FindSequenceForIdentifier(Prim.GetPrimPath().GetString());
+		// If our main LevelSequence doesn't have any subsequence sections, the SequenceHierarchyCache will be fully empty and we will fail.
+		// Even in that case though, don't show a warning when returning the request for the root as that is meant to be the identity transform
+		// anyway.
+		USD_LOG_WARNING(TEXT("Failed to find LevelSequence '%s' in SequenceHierarchyCache"), *Sequence->GetPathName());
 	}
 
-	return Sequence;
+	return IdentityTransform;
 }
 
 ULevelSequence* FUsdLevelSequenceHelperImpl::FindOrAddSequenceForAttribute(const UE::FUsdAttribute& Attribute, UE::FSdfLayer* OutSequenceLayer)
 {
-	if (!Attribute || !Attribute.GetPrim())
+	using namespace UE;
+
+	// Get layer with strongest attribute spec
+	FSdfLayer AttributeLayer = UsdUtils::FindLayerForAttribute(Attribute, UsdUtils::GetEarliestTimeCode());
+	if (!AttributeLayer)
 	{
 		return nullptr;
 	}
 
-	ULevelSequence* Sequence = FindSequenceForAttribute(Attribute);
-	if (!Sequence)
+	// This gets the leafmost composition arc node for the strongest opinion, corresponding to the opinion that also led us to AttributeLayer
+	FUsdResolveInfo ResolveInfo = Attribute.GetResolveInfo(UsdUtils::GetEarliestTimeCode());
+	FPcpNodeRef Node = ResolveInfo.GetNode();
+
+	ULevelSequence* AttributeSequence = CreateSequenceAndParentsRecursively(Node, AttributeLayer);
+	if (AttributeSequence && OutSequenceLayer)
 	{
-		if (UE::FSdfLayer AttributeLayer = UsdUtils::FindLayerForAttribute(Attribute, 0.0))
-		{
-			const FString SequenceIdentifier = AttributeLayer.GetIdentifier();
-			Sequence = FindOrAddSequenceForLayer(AttributeLayer, SequenceIdentifier, SequenceIdentifier);
-			if (OutSequenceLayer)
-			{
-				*OutSequenceLayer = AttributeLayer;
-			}
-		}
+		*OutSequenceLayer = AttributeLayer;
 	}
-
-	return Sequence;
+	return AttributeSequence;
 }
 
-ULevelSequence* FUsdLevelSequenceHelperImpl::FindOrAddSequenceForLayer(
-	const UE::FSdfLayer& Layer,
-	const FString& SequenceIdentifier,
-	const FString& SequenceDisplayName
-)
+ULevelSequence* FUsdLevelSequenceHelperImpl::FindOrAddSequenceForLayer(const UE::FSdfLayer& Layer, bool* bOutCreatedNew)
 {
 	using namespace UsdUnreal::ObjectUtils;
 
@@ -1458,7 +1493,10 @@ ULevelSequence* FUsdLevelSequenceHelperImpl::FindOrAddSequenceForLayer(
 		return nullptr;
 	}
 
-	ULevelSequence* Sequence = FindSequenceForIdentifier(SequenceIdentifier);
+	const FString LayerIdentifier = Layer.GetIdentifier();
+	const FString LayerDisplayName = Layer.GetDisplayName();
+
+	ULevelSequence* Sequence = FindSequenceForIdentifier(LayerIdentifier);
 
 	if (!Sequence)
 	{
@@ -1478,13 +1516,13 @@ ULevelSequence* FUsdLevelSequenceHelperImpl::FindOrAddSequenceForLayer(
 		// We still never want the LevelSequences to collide *within the same import* however, as we'd just end up stomping the same LevelSequences
 		// we're currently importing (see UE-283920), so here we use GetUniqueName instead
 		FName UniqueSequenceName = bIsImporting ? *GetUniqueName(
-													*SanitizeObjectName(FPaths::GetBaseFilename(SequenceDisplayName)),
+													*SanitizeObjectName(FPaths::GetBaseFilename(LayerDisplayName)),
 													LayerIdentifierByLevelSequenceName
 												)
 												: MakeUniqueObjectName(
 													GetTransientPackage(),
 													ULevelSequence::StaticClass(),
-													*SanitizeObjectName(FPaths::GetBaseFilename(SequenceDisplayName))
+													*SanitizeObjectName(FPaths::GetBaseFilename(LayerDisplayName))
 												);
 
 		Sequence = NewObject<ULevelSequence>(GetTransientPackage(), UniqueSequenceName, FUsdLevelSequenceHelperImpl::DefaultObjFlags);
@@ -1496,22 +1534,210 @@ ULevelSequence* FUsdLevelSequenceHelperImpl::FindOrAddSequenceForLayer(
 			return nullptr;
 		}
 
-		LayerIdentifierByLevelSequenceName.Add(Sequence->GetName(), Layer.GetIdentifier());
-		LevelSequencesByIdentifier.Add(SequenceIdentifier, Sequence);
-		IdentifierByLevelSequence.Add(Sequence, SequenceIdentifier);
+		LayerIdentifierByLevelSequenceName.Add(Sequence->GetName(), LayerIdentifier);
+		LevelSequencesByIdentifier.Add(LayerIdentifier, Sequence);
+		IdentifierByLevelSequence.Add(Sequence, LayerIdentifier);
 
 		// Here we abuse the FScopedReadOnlyDisable so that we can use the code in its destructor to set
 		// Sequence to ReadOnly if Layer doesn't belong to UsdStage's local layer stack
 		UE::USDLevelSequenceHelper::Private::FScopedReadOnlyDisable SetToReadOnly{MovieScene, Layer, UsdStage};
 
-		const FLayerTimeInfo LayerTimeInfo = FindOrAddLayerTimeInfo(Layer);
-
+		const FLayerTimeInfo& LayerTimeInfo = FindOrAddLayerTimeInfo(Layer);
 		UpdateMovieSceneTimeRanges(*MovieScene, LayerTimeInfo);
+
+		if (bOutCreatedNew)
+		{
+			*bOutCreatedNew = true;
+		}
 	}
 
 	return Sequence;
 }
 
+ULevelSequence* FUsdLevelSequenceHelperImpl::CreateSequenceAndParentsRecursively(const UE::FPcpNodeRef& CurrentNode, const UE::FSdfLayer& CurrentLayer)
+{
+	using namespace UE;
+
+	bool bIsNewSequence = false;
+	ULevelSequence* CurrentSequence = FindOrAddSequenceForLayer(CurrentLayer, &bIsNewSequence);
+	if (!bIsNewSequence)
+	{
+		// For now we still assume a layer only shows up once in the whole layer stack, so if its LevelSequence was created
+		// before, it was also already added as a subsequence to another LevelSequence somewhere, and we're fully done here
+		//
+		// TODO: Support multiple subsections of the same LevelSequence, which can happen e.g. when we have two referencer
+		// prims may reference the same layer with different offset/scale
+		return CurrentSequence;
+	}
+
+	FPcpLayerStack NodeStack = CurrentNode.GetLayerStack();
+	TArray<FSdfLayer> NodeStackLayers = NodeStack.GetLayers();
+	FPcpLayerStackIdentifier LayerStackIdentifier = NodeStack.GetIdentifier();
+	FSdfLayer NodeStackRootLayer = LayerStackIdentifier.RootLayer();
+	FString NodeStackRootLayerIdentifier = NodeStackRootLayer.GetIdentifier();
+
+	// This means the site we're chasing is not the immediately referenced/payload layer, but something that is part
+	// of its internal composition, and we must dig further to connect our LevelSequences
+	if (NodeStackRootLayer != CurrentLayer)
+	{
+		// Note that at this point these LevelSequences are not necessarily added to the main LevelSequence
+		// tree of the stage's root layer. We'll do that next, when we find the proper referencer layer
+		// to recurse upwards with
+		CreateSequencesForLayerStack(NodeStackRootLayer);
+	}
+
+	// Find next composition arc to climb up
+	FPcpNodeRef ChildNode = CurrentNode;
+	FPcpNodeRef ParentNode = CurrentNode.GetParentNode();
+	while (ParentNode && ParentNode.GetLayerStack() == NodeStack)
+	{
+		ChildNode = ParentNode;
+		ParentNode = ParentNode.GetParentNode();
+	}
+
+	// If we have no parent node there is no composition arc between us and the root,
+	// and so CurrentLayer is part of the local layer stack and must have already been created at Init().
+	// This means we should have found our CurrentSequence via the FindOrAddSequenceForLayer call at
+	// the top of this function, and something went wrong
+	if (!ensure(ParentNode))
+	{
+		return CurrentSequence;
+	}
+
+	FPcpLayerStack ParentStack = ParentNode.GetLayerStack();
+	FSdfPath PrimPathInParent = ParentNode.GetPath();
+
+	auto HandleReferenceOrPayload = [&ParentStack,
+									 &PrimPathInParent,
+									 &NodeStackRootLayer,
+									 this,
+									 &NodeStackRootLayerIdentifier,
+									 &ParentNode,
+									 CurrentSequence,
+									 &NodeStackLayers,
+									 &CurrentLayer]<typename ProxyType>()
+	{
+		// Need to find the strongest spec in the parent stack that authors an additive reference Op to our child stack layer,
+		// as that is the layer where we want to add our subsequence section for our LocalSequence
+		FSdfLayer FoundReferencerLayer;
+		FSdfLayerOffset ReferencedLayerOffset;
+
+		for (const FSdfLayer& ParentLayer : ParentStack.GetLayers())
+		{
+			// Check for a spec in this layer
+			FSdfPath ParentPath = PrimPathInParent;
+			FSdfPrimSpec SpecInParentLayer = ParentLayer.GetPrimAtPath(ParentPath);
+			while (!SpecInParentLayer && !ParentPath.IsAbsoluteRootPath())
+			{
+				// It could be that we don't have a spec for this exact prim in this layer but it was instead
+				// a parent prim that authored the reference, so let's climb up.
+				// This is probably safe to do as we're hunting for an exact reference to NodeStackRootLayer anyway
+				ParentPath = ParentPath.GetParentPath();
+				SpecInParentLayer = ParentLayer.GetPrimAtPath(ParentPath);
+			}
+			if (!SpecInParentLayer)
+			{
+				continue;
+			}
+
+			ProxyType ReferenceOrPayloadProxy;
+			if constexpr (std::is_same_v<ProxyType, FSdfReferencesProxy>)
+			{
+				ReferenceOrPayloadProxy = SpecInParentLayer.GetReferenceList();
+			}
+			else
+			{
+				ReferenceOrPayloadProxy = SpecInParentLayer.GetPayloadList();
+			}
+
+			for (const auto& ReferenceOrPayload : ReferenceOrPayloadProxy.GetAppliedItems())
+			{
+				FString ReferenceTarget = ReferenceOrPayload.AssetPath;
+
+				// TODO: Use asset resolver here to fetch a real filesystem path. For now we're assuming ReferenceTarget already is one
+
+				if (FPaths::IsRelative(ReferenceTarget))
+				{
+					ReferenceTarget = FPaths::ConvertRelativePathToFull(FPaths::GetPath(ParentLayer.GetRealPath()), ReferenceTarget);
+				}
+
+				FSdfLayer TargetLayer = FSdfLayer::FindOrOpen(*ReferenceTarget);
+
+				// The target of the reference we want is always the root layer of the referenced layer stack
+				if (TargetLayer == NodeStackRootLayer)
+				{
+					FoundReferencerLayer = ParentLayer;
+					ReferencedLayerOffset = ReferenceOrPayload.LayerOffset;
+					break;
+				}
+			}
+
+			if (FoundReferencerLayer)
+			{
+				break;
+			}
+		}
+
+		if (FoundReferencerLayer)
+		{
+			// Store our offset into our FLayerOffsetInfo structs
+			FLayerTimeInfo& ParentInfo = FindOrAddLayerTimeInfo(FoundReferencerLayer);
+			if constexpr (std::is_same_v<ProxyType, FSdfReferencesProxy>)
+			{
+				ParentInfo.ReferenceIdentifierToOffsets.Add(NodeStackRootLayerIdentifier, ReferencedLayerOffset);
+			}
+			else
+			{
+				ParentInfo.PayloadIdentifierToOffsets.Add(NodeStackRootLayerIdentifier, ReferencedLayerOffset);
+			}
+
+			ULevelSequence* ParentSequence = CreateSequenceAndParentsRecursively(ParentNode, FoundReferencerLayer);
+			ULevelSequence* ChildSequence = FindSequenceForIdentifier(NodeStackRootLayerIdentifier);
+			if (!ensure(ParentSequence && ChildSequence))
+			{
+				return;
+			}
+
+			// Our parent referencer is always referencing the root of our layer stack: Not necessarily our current layer
+			CreateSubSequenceSection(*ParentSequence, *ChildSequence, ReferencedLayerOffset);
+
+			// This is a failsafe case: It seems that for value clips the assetPaths layers don't count as sublayers,
+			// payloads, references, or any other composition arc... We'll still have to handle LevelSequences created
+			// for these layers however, as FindOrAddSequenceForLayer() and friends will actually end up finding the attribute
+			// specs inside of them. Luckily they don't have any sort of offset or scale, so here we'll just add the assetPath
+			// sequence to its referencer layer LevelSequence with a subsection starting at zero
+			if (!NodeStackLayers.Contains(CurrentLayer))
+			{
+				UE::FSdfLayerOffset NoOffset;
+				CreateSubSequenceSection(*ChildSequence, *CurrentSequence, NoOffset);
+			}
+		}
+	};
+
+	switch (ChildNode.GetArcType())
+	{
+		case EPcpArcType::PcpArcTypeReference:
+		{
+			// Weird syntax but we're just calling the lamda specifying the template type here
+			HandleReferenceOrPayload.template operator()<FSdfReferencesProxy>();
+			break;
+		}
+		case EPcpArcType::PcpArcTypePayload:
+		{
+			HandleReferenceOrPayload.template operator()<FSdfPayloadsProxy>();
+			break;
+		}
+		default:
+		{
+			// Reference and payload arcs are the only composition arcs that can jump to a different layer stack, so they're the
+			// only ones we care about for now
+			break;
+		}
+	}
+
+	return CurrentSequence;
+}
+
 UE::FSdfLayer FUsdLevelSequenceHelperImpl::FindEditTargetForSubsequence(ULevelSequence* Sequence)
 {
 	FString LayerIdentifier = IdentifierByLevelSequence.FindRef(Sequence);
@@ -1558,7 +1784,7 @@ UMovieSceneSubSection* FUsdLevelSequenceHelperImpl::FindSubSequenceSection(ULeve
 	}
 }
 
-void FUsdLevelSequenceHelperImpl::CreateSubSequenceSection(ULevelSequence& Sequence, ULevelSequence& SubSequence)
+void FUsdLevelSequenceHelperImpl::CreateSubSequenceSection(ULevelSequence& Sequence, ULevelSequence& SubSequence, const UE::FSdfLayerOffset& SubLayerOffset)
 {
 	if (&Sequence == &SubSequence)
 	{
@@ -1586,54 +1812,15 @@ void FUsdLevelSequenceHelperImpl::CreateSubSequenceSection(ULevelSequence& Seque
 
 	const FString* LayerIdentifier = LayerIdentifierByLevelSequenceName.Find(Sequence.GetName());
 	const FString* SubLayerIdentifier = LayerIdentifierByLevelSequenceName.Find(SubSequence.GetName());
-
 	if (!LayerIdentifier || !SubLayerIdentifier)
 	{
 		return;
 	}
 
-	FLayerTimeInfo* LayerTimeInfo = LayerTimeInfosByLayerIdentifier.Find(*LayerIdentifier);
-	FLayerTimeInfo* SubLayerTimeInfo = LayerTimeInfosByLayerIdentifier.Find(*SubLayerIdentifier);
-
-	if (!LayerTimeInfo || !SubLayerTimeInfo)
-	{
-		return;
-	}
-
-	UE::FSdfLayerOffset SubLayerOffset;
-
 	UE::FSdfLayer Layer = UE::FSdfLayer::FindOrOpen(**LayerIdentifier);
 	UE::FSdfLayer SubLayer = UE::FSdfLayer::FindOrOpen(**SubLayerIdentifier);
 
-	TArray<FString> PrimPathsForSequence;
-	PrimPathByLevelSequenceName.MultiFind(SubSequence.GetFName(), PrimPathsForSequence);
-
-	if (PrimPathsForSequence.Num() > 0)
-	{
-		if (UE::FUsdPrim SequencePrim = UsdStage.GetPrimAtPath(UE::FSdfPath(*PrimPathsForSequence[0])))
-		{
-			TArray<UE::FUsdAttribute> Attrs = UsdUtils::GetAttributesForProperty(SequencePrim, UnrealIdentifiers::TransformPropertyName);
-			if (Attrs.Num() > 0)
-			{
-				SubLayerOffset = UsdUtils::GetLayerToStageOffset(Attrs[0]);
-			}
-		}
-	}
-	else if (UsdStage.HasLocalLayer(SubLayer))
-	{
-		const FLayerOffsetInfo* SubLayerOffsetPtr = Algo::FindByPredicate(
-			LayerTimeInfo->SubLayersOffsets,
-			[&SubLayerIdentifier](const FLayerOffsetInfo& Other)
-			{
-				return (Other.LayerIdentifier == *SubLayerIdentifier);
-			}
-		);
-
-		if (SubLayerOffsetPtr)
-		{
-			SubLayerOffset = SubLayerOffsetPtr->LayerOffset;
-		}
-	}
+	UE::USDLevelSequenceHelper::Private::FScopedReadOnlyDisable ReadOnlyGuard{MovieScene, Layer, UsdStage};
 
 	const double TimeCodesPerSecond = Layer.GetTimeCodesPerSecond();
 	const bool bIsAlembicSublayer = SubLayerIdentifier->EndsWith(TEXT(".abc"));
@@ -1698,6 +1885,7 @@ void FUsdLevelSequenceHelperImpl::CreateSubSequenceSection(ULevelSequence& Seque
 			SubSectionRange.Size<FFrameNumber>().Value,
 			RowIndex
 		);
+		DirtySequenceHierarchyCache();
 	}
 
 	const double TimeCodesPerSecondDifference = TimeCodesPerSecond / SubLayer.GetTimeCodesPerSecond();
@@ -1711,23 +1899,6 @@ void FUsdLevelSequenceHelperImpl::CreateSubSequenceSection(ULevelSequence& Seque
 	// These will be edited together by the Sequencer (so e.g. if a user drags the inner playback range it will also update the start frame offset),
 	// we just have to set them up properly the first time around it seems.
 	SubSection->Parameters.StartFrameOffset = -SubSequence.GetMovieScene()->GetPlaybackRange().GetLowerBoundValue();
-
-	if (MainLevelSequence)
-	{
-		UMovieSceneCompiledDataManager::CompileHierarchy(MainLevelSequence, &SequenceHierarchyCache, EMovieSceneServerClientMask::All);
-
-		for (const TTuple<FMovieSceneSequenceID, FMovieSceneSubSequenceData>& Pair : SequenceHierarchyCache.AllSubSequenceData())
-		{
-			if (UMovieSceneSequence* CachedSubSequence = Pair.Value.GetSequence())
-			{
-				if (CachedSubSequence == &SubSequence)
-				{
-					SequencesID.Add(&SubSequence, Pair.Key);
-					break;
-				}
-			}
-		}
-	}
 }
 
 void FUsdLevelSequenceHelperImpl::RemoveSubSequenceSection(ULevelSequence& Sequence, ULevelSequence& SubSequence)
@@ -1736,14 +1907,10 @@ void FUsdLevelSequenceHelperImpl::RemoveSubSequenceSection(ULevelSequence& Seque
 	{
 		if (UMovieSceneSection* SubSection = FindSubSequenceSection(Sequence, SubSequence))
 		{
-			SequencesID.Remove(&SubSequence);
 			SubTrack->Modify();
 			SubTrack->RemoveSection(*SubSection);
 
-			if (MainLevelSequence)
-			{
-				UMovieSceneCompiledDataManager::CompileHierarchy(MainLevelSequence, &SequenceHierarchyCache, EMovieSceneServerClientMask::All);
-			}
+			DirtySequenceHierarchyCache();
 		}
 	}
 }
@@ -1792,71 +1959,78 @@ void FUsdLevelSequenceHelperImpl::UpdateSubSectionTimeRanges(bool bShowResizedSe
 		double TimeCodesPerSecond = GetTimeCodesPerSecond();
 		const FFrameNumber StageEndFrame = UE::USDLevelSequenceHelper::Private::RoundAsFrameNumber(TickResolution, StageEndTimeSeconds);
 
-		FLayerTimeInfo* LayerTimeInfo = LayerTimeInfosByLayerIdentifier.Find(*IdentifierAndSeq.Key);
+		FLayerTimeInfo* LayerTimeInfo = LayerTimeInfosByLayerIdentifier.Find(IdentifierAndSeq.Key);
 
 		for (UMovieSceneSection* Section : SubTrack->GetAllSections())
 		{
-			if (UMovieSceneSubSection* SubSection = Cast<UMovieSceneSubSection>(Section))
+			UMovieSceneSubSection* SubSection = Cast<UMovieSceneSubSection>(Section);
+			if (!SubSection)
 			{
-				UMovieSceneSequence* SubSequence = SubSection->GetSequence();
-				if (!SubSequence)
-				{
-					continue;
-				}
+				continue;
+			}
 
-				FString* SubLayerIdentifier = IdentifierByLevelSequence.Find(Cast<ULevelSequence>(SubSequence));
-				if (!SubLayerIdentifier)
-				{
-					continue;
-				}
+			UMovieSceneSequence* SubSequence = SubSection->GetSequence();
+			if (!SubSequence)
+			{
+				continue;
+			}
 
-				UE::FSdfLayer SubLayer = UE::FSdfLayer::FindOrOpen(**SubLayerIdentifier);
-				if (!SubLayer)
-				{
-					continue;
-				}
+			FString* SubSequenceLayerIdentifier = IdentifierByLevelSequence.Find(Cast<ULevelSequence>(SubSequence));
+			if (!SubSequenceLayerIdentifier)
+			{
+				continue;
+			}
 
-				// Find our current sublayer offset for this sublayer
-				FLayerOffsetInfo* SubLayerOffset = Algo::FindByPredicate(
-					LayerTimeInfo->SubLayersOffsets,
-					[SubLayerIdentifier](const FLayerOffsetInfo& Other)
-					{
-						return Other.LayerIdentifier == *SubLayerIdentifier;
-					}
-				);
-				if (!SubLayerOffset)
-				{
-					continue;
-				}
+			// TODO: This will likely fail in edge cases where a layer is both a sublayer and a reference (or a reference and a payload, etc.)
+			// We need to better track the actual sections to know what they represent purely on the UE side (maybe some metadata, or tracking the
+			// pointers?)
+			UE::FSdfLayerOffset* SubLayerOffset = LayerTimeInfo->SublayerIdentifierToOffsets.Find(*SubSequenceLayerIdentifier);
+			if (!SubLayerOffset)
+			{
+				SubLayerOffset = LayerTimeInfo->ReferenceIdentifierToOffsets.Find(*SubSequenceLayerIdentifier);
+			}
+			if (!SubLayerOffset)
+			{
+				SubLayerOffset = LayerTimeInfo->PayloadIdentifierToOffsets.Find(*SubSequenceLayerIdentifier);
+			}
+			if (!SubLayerOffset)
+			{
+				continue;
+			}
 
-				SubSection->Modify();
+			UE::FSdfLayer SubSequenceLayer = UE::FSdfLayer::FindOrOpen(**SubSequenceLayerIdentifier);
+			if (!SubSequenceLayer)
+			{
+				continue;
+			}
 
-				// StartFrameOffset: Should match the sublayer's playback range start so that the inner sequence starts playing at its zero
-				SubSection->Parameters.StartFrameOffset = -SubSequence->GetMovieScene()->GetPlaybackRange().GetLowerBoundValue();
+			SubSection->Modify();
 
-				// TimeScale: Should match the sublayer's scale (or the inverse of it, rather)
-				const double TimeCodesPerSecondDifference = TimeCodesPerSecond / SubLayer.GetTimeCodesPerSecond();
-				SubSection->Parameters.TimeScale = FMath::IsNearlyZero(SubLayerOffset->LayerOffset.Scale)
-													   ? 0.f
-													   : 1.f / (SubLayerOffset->LayerOffset.Scale / TimeCodesPerSecondDifference);
+			// StartFrameOffset: Should match the sublayer's playback range start so that the inner sequence starts playing at its zero
+			SubSection->Parameters.StartFrameOffset = -SubSequence->GetMovieScene()->GetPlaybackRange().GetLowerBoundValue();
 
-				// Section start/end: The start should match the sublayer's offset, and the end should always clip to the stage's endTimeCode
-				TRange<FFrameNumber> Range = SubSection->GetRange();
-				FFrameNumber NewLower = UE::USDLevelSequenceHelper::Private::RoundAsFrameNumber(
-					TickResolution,
-					SubLayerOffset->LayerOffset.Offset / TimeCodesPerSecond
-				);
-				FFrameNumber NewUpper = FMath::Max(StageEndFrame, NewLower);
-				if (NewLower == Range.GetLowerBoundValue() && NewUpper == Range.GetUpperBoundValue())
-				{
-					continue;
-				}
-				Range.SetLowerBoundValue(NewLower);
-				Range.SetUpperBoundValue(NewUpper);
-				SubSection->SetRange(Range);
+			// TimeScale: Should match the sublayer's scale (or the inverse of it, rather)
+			const double TimeCodesPerSecondDifference = TimeCodesPerSecond / SubSequenceLayer.GetTimeCodesPerSecond();
+			SubSection->Parameters.TimeScale = FMath::IsNearlyZero(SubLayerOffset->Scale)
+												   ? 0.f
+												   : 1.f / (SubLayerOffset->Scale / TimeCodesPerSecondDifference);
 
-				bChangedSectionSize = true;
+			// Section start/end: The start should match the sublayer's offset, and the end should always clip to the stage's endTimeCode
+			TRange<FFrameNumber> Range = SubSection->GetRange();
+			FFrameNumber NewLower = UE::USDLevelSequenceHelper::Private::RoundAsFrameNumber(
+				TickResolution,
+				SubLayerOffset->Offset / TimeCodesPerSecond
+			);
+			FFrameNumber NewUpper = FMath::Max(StageEndFrame, NewLower);
+			if (NewLower == Range.GetLowerBoundValue() && NewUpper == Range.GetUpperBoundValue())
+			{
+				continue;
 			}
+			Range.SetLowerBoundValue(NewLower);
+			Range.SetUpperBoundValue(NewUpper);
+			SubSection->SetRange(Range);
+
+			bChangedSectionSize = true;
 		}
 	}
 
@@ -2047,9 +2221,6 @@ void FUsdLevelSequenceHelperImpl::AddCommonTracks(const UUsdPrimTwin& PrimTwin,
 		return;
 	}
 
-	UE::FSdfLayer PrimLayer = UsdUtils::FindLayerForPrim(Prim);
-	ULevelSequence* PrimSequence = FindSequenceForIdentifier(PrimLayer.GetIdentifier());
-
 	UE::FUsdGeomXformable Xformable(Prim);
 
 	// If this xformable has an op to reset the xform stack and one of its ancestors is animated, then we need to make
@@ -2144,28 +2315,23 @@ void FUsdLevelSequenceHelperImpl::AddCommonTracks(const UUsdPrimTwin& PrimTwin,
 			}
 		}
 
-		// Find the strongest layer where any of these is authored. The TimeCode here is only for handling Value Clips, which
-		// we largely don't support anyway
-		const double TimeCode = 0.0;
-		const bool bIncludeSessionLayers = false;
-		UE::FSdfLayer Layer = UsdUtils::FindLayerForAttributes(Attrs, TimeCode, bIncludeSessionLayers);
+		UE::FUsdAttribute StrongestAttr = UsdUtils::GetStrongestAttribute(Attrs, UsdUtils::GetEarliestTimeCode());
+
+		UE::FSdfLayer Layer;
+		ULevelSequence* AttributeSequence = FindOrAddSequenceForAttribute(StrongestAttr, &Layer);
 
 		// If we're creating a brand new transform track to compensate resetXformOp we may not have any animated attribute
 		// already, but we still need to do this
-		if (!Layer && bNeedTrackToCompensateResetXformOp)
+		if (!AttributeSequence && bNeedTrackToCompensateResetXformOp)
 		{
 			Layer = UsdUtils::FindLayerForPrim(Prim);
+			AttributeSequence = FindOrAddSequenceForLayer(Layer);
 		}
 
 		// Get the Subsequence where we should create our track according to that Layer
-		if (ULevelSequence* AttributeSequence = FindOrAddSequenceForLayer(Layer, Layer.GetIdentifier(), Layer.GetDisplayName()))
+		if (AttributeSequence)
 		{
-			FMovieSceneSequenceTransform SequenceTransform;
-			FMovieSceneSequenceID SequenceID = SequencesID.FindRef(AttributeSequence);
-			if (FMovieSceneSubSequenceData* SubSequenceData = SequenceHierarchyCache.FindSubData(SequenceID))
-			{
-				SequenceTransform = SubSequenceData->RootToSequenceTransform;
-			}
+			const FMovieSceneSequenceTransform& SequenceTransform = GetSubsequenceTransform(AttributeSequence);
 
 			if (UMovieScene* MovieScene = AttributeSequence->GetMovieScene())
 			{
@@ -2274,13 +2440,6 @@ void FUsdLevelSequenceHelperImpl::AddCommonTracks(const UUsdPrimTwin& PrimTwin,
 				{
 					const bool bIsMuted = UsdUtils::IsAttributeMuted(AttributeForSequence, UsdStage);
 
-					FMovieSceneSequenceTransform SequenceTransform;
-					FMovieSceneSequenceID SequenceID = SequencesID.FindRef(AttributeSequence);
-					if (FMovieSceneSubSequenceData* SubSequenceData = SequenceHierarchyCache.FindSubData(SequenceID))
-					{
-						SequenceTransform = SubSequenceData->RootToSequenceTransform;
-					}
-
 					if (UMovieScene* MovieScene = AttributeSequence->GetMovieScene())
 					{
 						UE::USDLevelSequenceHelper::Private::FScopedReadOnlyDisable ReadOnlyGuard{MovieScene, SequenceLayer, UsdStage};
@@ -2293,6 +2452,7 @@ void FUsdLevelSequenceHelperImpl::AddCommonTracks(const UUsdPrimTwin& PrimTwin,
 								bIsMuted
 							))
 						{
+							const FMovieSceneSequenceTransform& SequenceTransform = GetSubsequenceTransform(AttributeSequence);
 							UsdToUnreal::FPropertyTrackReader Reader = UsdToUnreal::CreatePropertyTrackReader(
 								Prim,
 								UnrealIdentifiers::HiddenInGamePropertyName
@@ -2383,8 +2543,10 @@ void FUsdLevelSequenceHelperImpl::AddBoundsTracks(const UUsdPrimTwin& PrimTwin,
 		// For these tracks however we may not have any authored `extent` or `extentsHint` yet (and would only author them on-demand), so
 		// we may need to create tracks that correspond purely to computed, non-authored values in USD.
 		// If the user manually modifies these, we'll author these as `extent` or `extentsHint` depending on the prim, but only on-demand.
+		//
+		// TODO: These may not be connected to the LevelSequence hierarchy in edge cases
 		UE::FSdfLayer PrimLayer = UsdUtils::FindLayerForPrim(Prim);
-		TargetSequence = FindSequenceForIdentifier(PrimLayer.GetIdentifier());
+		TargetSequence = FindOrAddSequenceForLayer(PrimLayer);
 		TargetLayer = PrimLayer;
 	}
 	if (!TargetSequence)
@@ -2392,12 +2554,7 @@ void FUsdLevelSequenceHelperImpl::AddBoundsTracks(const UUsdPrimTwin& PrimTwin,
 		return;
 	}
 
-	FMovieSceneSequenceTransform SequenceTransform;
-	FMovieSceneSequenceID SequenceID = SequencesID.FindRef(TargetSequence);
-	if (FMovieSceneSubSequenceData* SubSequenceData = SequenceHierarchyCache.FindSubData(SequenceID))
-	{
-		SequenceTransform = SubSequenceData->RootToSequenceTransform;
-	}
+	const FMovieSceneSequenceTransform& SequenceTransform = GetSubsequenceTransform(TargetSequence);
 
 	UMovieScene* MovieScene = TargetSequence->GetMovieScene();
 	if (!MovieScene)
@@ -2503,13 +2660,6 @@ void FUsdLevelSequenceHelperImpl::AddCameraTracks(const UUsdPrimTwin& PrimTwin,
 		{
 			const bool bIsMuted = UsdUtils::IsAttributeMuted(Attr, UsdStage);
 
-			FMovieSceneSequenceTransform SequenceTransform;
-			FMovieSceneSequenceID SequenceID = SequencesID.FindRef(AttributeSequence);
-			if (FMovieSceneSubSequenceData* SubSequenceData = SequenceHierarchyCache.FindSubData(SequenceID))
-			{
-				SequenceTransform = SubSequenceData->RootToSequenceTransform;
-			}
-
 			UMovieScene* MovieScene = AttributeSequence->GetMovieScene();
 			if (!MovieScene)
 			{
@@ -2527,6 +2677,7 @@ void FUsdLevelSequenceHelperImpl::AddCameraTracks(const UUsdPrimTwin& PrimTwin,
 			if (UMovieSceneFloatTrack* FloatTrack = AddTrack<
 					UMovieSceneFloatTrack>(PropertyName, PrimTwin, *ComponentToBind, *AttributeSequence, bIsMuted))
 			{
+				const FMovieSceneSequenceTransform& SequenceTransform = GetSubsequenceTransform(AttributeSequence);
 				UsdToUnreal::FPropertyTrackReader Reader = UsdToUnreal::CreatePropertyTrackReader(Prim, PropertyName);
 				UsdToUnreal::ConvertFloatTimeSamples(UsdStage, TimeSamples, Reader.FloatReader, *FloatTrack, SequenceTransform);
 			}
@@ -2596,13 +2747,6 @@ void FUsdLevelSequenceHelperImpl::AddLightTracks(const UUsdPrimTwin& PrimTwin, c
 		}
 	}
 
-	UE::FSdfLayer PrimLayer = UsdUtils::FindLayerForPrim(Prim);
-	ULevelSequence* PrimSequence = FindSequenceForIdentifier(PrimLayer.GetIdentifier());
-	if (!PrimSequence)
-	{
-		return;
-	}
-
 	for (const TPair<FName, ETrackType>& Pair : PropertyPathToTrackType)
 	{
 		const FName& PropertyPath = Pair.Key;
@@ -2638,20 +2782,22 @@ void FUsdLevelSequenceHelperImpl::AddLightTracks(const UUsdPrimTwin& PrimTwin, c
 			continue;
 		}
 
-		FMovieSceneSequenceTransform SequenceTransform;
-		FMovieSceneSequenceID SequenceID = SequencesID.FindRef(PrimSequence);
-		if (FMovieSceneSubSequenceData* SubSequenceData = SequenceHierarchyCache.FindSubData(SequenceID))
+		UE::FSdfLayer AttributeLayer;
+		ULevelSequence* AttributeSequence = FindOrAddSequenceForAttribute(MainAttr, &AttributeLayer);
+		if (!AttributeSequence)
 		{
-			SequenceTransform = SubSequenceData->RootToSequenceTransform;
+			return;
 		}
 
-		UMovieScene* MovieScene = PrimSequence->GetMovieScene();
+		UMovieScene* MovieScene = AttributeSequence->GetMovieScene();
 		if (!MovieScene)
 		{
 			continue;
 		}
 
-		UE::USDLevelSequenceHelper::Private::FScopedReadOnlyDisable ReadOnlyGuard{MovieScene, PrimLayer, UsdStage};
+		UE::USDLevelSequenceHelper::Private::FScopedReadOnlyDisable ReadOnlyGuard{MovieScene, AttributeLayer, UsdStage};
+
+		const FMovieSceneSequenceTransform& SequenceTransform = GetSubsequenceTransform(AttributeSequence);
 
 		UsdToUnreal::FPropertyTrackReader Reader = UsdToUnreal::CreatePropertyTrackReader(Prim, PropertyPath);
 
@@ -2660,7 +2806,7 @@ void FUsdLevelSequenceHelperImpl::AddLightTracks(const UUsdPrimTwin& PrimTwin, c
 			case ETrackType::Bool:
 			{
 				if (UMovieSceneBoolTrack* BoolTrack = AddTrack<
-						UMovieSceneBoolTrack>(PropertyPath, PrimTwin, *ComponentToBind, *PrimSequence, bIsMuted))
+						UMovieSceneBoolTrack>(PropertyPath, PrimTwin, *ComponentToBind, *AttributeSequence, bIsMuted))
 				{
 					UsdToUnreal::ConvertBoolTimeSamples(UsdStage, UnionedTimeSamples, Reader.BoolReader, *BoolTrack, SequenceTransform);
 				}
@@ -2669,7 +2815,7 @@ void FUsdLevelSequenceHelperImpl::AddLightTracks(const UUsdPrimTwin& PrimTwin, c
 			case ETrackType::Float:
 			{
 				if (UMovieSceneFloatTrack* FloatTrack = AddTrack<
-						UMovieSceneFloatTrack>(PropertyPath, PrimTwin, *ComponentToBind, *PrimSequence, bIsMuted))
+						UMovieSceneFloatTrack>(PropertyPath, PrimTwin, *ComponentToBind, *AttributeSequence, bIsMuted))
 				{
 					UsdToUnreal::ConvertFloatTimeSamples(UsdStage, UnionedTimeSamples, Reader.FloatReader, *FloatTrack, SequenceTransform);
 				}
@@ -2678,7 +2824,7 @@ void FUsdLevelSequenceHelperImpl::AddLightTracks(const UUsdPrimTwin& PrimTwin, c
 			case ETrackType::Color:
 			{
 				if (UMovieSceneColorTrack* ColorTrack = AddTrack<
-						UMovieSceneColorTrack>(PropertyPath, PrimTwin, *ComponentToBind, *PrimSequence, bIsMuted))
+						UMovieSceneColorTrack>(PropertyPath, PrimTwin, *ComponentToBind, *AttributeSequence, bIsMuted))
 				{
 					UsdToUnreal::ConvertColorTimeSamples(UsdStage, UnionedTimeSamples, Reader.ColorReader, *ColorTrack, SequenceTransform);
 				}
@@ -2688,7 +2834,7 @@ void FUsdLevelSequenceHelperImpl::AddLightTracks(const UUsdPrimTwin& PrimTwin, c
 				continue;
 		}
 
-		PrimPathByLevelSequenceName.AddUnique(PrimSequence->GetFName(), Prim.GetPrimPath().GetString());
+		PrimPathByLevelSequenceName.AddUnique(AttributeSequence->GetFName(), Prim.GetPrimPath().GetString());
 	}
 }
 
@@ -2731,22 +2877,17 @@ void FUsdLevelSequenceHelperImpl::AddSkeletalTracks(const UUsdPrimTwin& PrimTwin
 	UE::FUsdAttribute ScalesAttr = SkelAnimationPrim.GetAttribute(TEXT("scales"));
 	UE::FUsdAttribute BlendShapeWeightsAttr = SkelAnimationPrim.GetAttribute(TEXT("blendShapeWeights"));
 
-	const bool bIncludeSessionLayers = false;
-	UE::FSdfLayer SkelAnimationLayer = UsdUtils::FindLayerForAttributes(
+	UE::FUsdAttribute StrongestAttribute = UsdUtils::GetStrongestAttribute(
 		{TranslationsAttr, RotationsAttr, ScalesAttr, BlendShapeWeightsAttr},
-		0.0,
-		bIncludeSessionLayers
+		UsdUtils::GetEarliestTimeCode()
 	);
-	if (!SkelAnimationLayer)
+	if (!StrongestAttribute)
 	{
 		return;
 	}
 
-	ULevelSequence* SkelAnimationSequence = FindOrAddSequenceForLayer(
-		SkelAnimationLayer,
-		SkelAnimationLayer.GetIdentifier(),
-		SkelAnimationLayer.GetDisplayName()
-	);
+	UE::FSdfLayer SkelAnimationLayer;
+	ULevelSequence* SkelAnimationSequence = FindOrAddSequenceForAttribute(StrongestAttribute, &SkelAnimationLayer);
 	if (!SkelAnimationSequence)
 	{
 		return;
@@ -2771,19 +2912,27 @@ void FUsdLevelSequenceHelperImpl::AddSkeletalTracks(const UUsdPrimTwin& PrimTwin
 	if (UMovieSceneSkeletalAnimationTrack* SkeletalTrack = AddTrack<
 			UMovieSceneSkeletalAnimationTrack>(SkelAnimationPrim.GetName(), PrimTwin, *ComponentToBind, *SkelAnimationSequence, bIsMuted))
 	{
-		double LayerStartOffsetSeconds = 0.0f;
+		UE::FUsdResolveInfo Info = StrongestAttribute.GetResolveInfo(UsdUtils::GetEarliestTimeCode());
+		UE::FPcpNodeRef PcpNode = Info.GetNode();
+		UE::FPcpMapExpression MapToRoot = PcpNode.GetMapToRoot();
+		UE::FSdfLayerOffset TotalOffset = MapToRoot.GetTimeOffset();
+
+		double StartSecondsInStage = 0.0f;
 		if (UUsdAnimSequenceAssetUserData* UserData = Sequence->GetAssetUserData<UUsdAnimSequenceAssetUserData>())
 		{
-			LayerStartOffsetSeconds = UserData->LayerStartOffsetSeconds;
+			StartSecondsInStage = UserData->StageStartOffsetSeconds;
 		}
 
-		FFrameNumber StartOffsetTick = FFrameTime::FromDecimal(LayerStartOffsetSeconds * MovieScene->GetTickResolution().AsDecimal()).RoundToFrame();
+		const double StartTimeCodeInStage = StartSecondsInStage * UsdStage.GetTimeCodesPerSecond();
+		const double StartTimeCodeInLayer = (StartTimeCodeInStage - TotalOffset.Offset) / TotalOffset.Scale;
+		const double StartSecondsInLayer = StartTimeCodeInLayer / SkelAnimationLayer.GetTimeCodesPerSecond();
+		FFrameNumber StartTickInLayer = FFrameTime::FromDecimal(StartSecondsInLayer * MovieScene->GetTickResolution().AsDecimal()).RoundToFrame();
 
 		SkeletalTrack->Modify();
 		SkeletalTrack->RemoveAllAnimationData();
 
 		UMovieSceneSkeletalAnimationSection* NewSection = Cast<UMovieSceneSkeletalAnimationSection>(
-			SkeletalTrack->AddNewAnimation(StartOffsetTick, Sequence)
+			SkeletalTrack->AddNewAnimation(StartTickInLayer, Sequence)
 		);
 		NewSection->EvalOptions.CompletionMode = EMovieSceneCompletionMode::KeepState;
 	}
@@ -2813,18 +2962,12 @@ void FUsdLevelSequenceHelperImpl::AddGeometryCacheTracks(const UUsdPrimTwin& Pri
 		return;
 	}
 
-	UE::FSdfLayer GeometryCacheLayer = UsdUtils::FindLayerForPrim(Prim);
-	if (!GeometryCacheLayer)
-	{
-		return;
-	}
+	TArray<UE::FUsdAttribute> AnimatedMeshAttributes = UsdUtils::GetAnimatedMeshAttributes(Prim);
+	UE::FUsdAttribute StrongestAttribute = UsdUtils::GetStrongestAttribute(AnimatedMeshAttributes, UsdUtils::GetEarliestTimeCode());
 
-	ULevelSequence* GeometryCacheSequence = FindOrAddSequenceForLayer(
-		GeometryCacheLayer,
-		GeometryCacheLayer.GetIdentifier(),
-		GeometryCacheLayer.GetDisplayName()
-	);
-	if (!GeometryCacheSequence)
+	UE::FSdfLayer GeometryCacheLayer;
+	ULevelSequence* GeometryCacheSequence = FindOrAddSequenceForAttribute(StrongestAttribute, &GeometryCacheLayer);
+	if (!ensure(GeometryCacheSequence)) // We should always be able to find a place to put our section
 	{
 		return;
 	}
@@ -2841,26 +2984,48 @@ void FUsdLevelSequenceHelperImpl::AddGeometryCacheTracks(const UUsdPrimTwin& Pri
 	if (UMovieSceneGeometryCacheTrack* GeometryCacheTrack = AddTrack<
 			UMovieSceneGeometryCacheTrack>(Prim.GetName(), PrimTwin, *ComponentToBind, *GeometryCacheSequence, bIsMuted))
 	{
-		GeometryCacheTrack->Modify();
-		GeometryCacheTrack->RemoveAllAnimationData();
+		UE::FUsdResolveInfo Info = StrongestAttribute.GetResolveInfo(UsdUtils::GetEarliestTimeCode());
+		UE::FPcpNodeRef PcpNode = Info.GetNode();
+		UE::FPcpMapExpression MapToRoot = PcpNode.GetMapToRoot();
+		UE::FSdfLayerOffset TotalOffset = MapToRoot.GetTimeOffset();
 
-		double LayerStartOffsetSeconds = 0;
-		if (GeometryCacheSequence == MainLevelSequence)
+		double StartSecondsInStage = 0;
+		if (UUsdGeometryCacheAssetUserData* UserData = GeometryCache->GetAssetUserData<UUsdGeometryCacheAssetUserData>())
 		{
-			// The LayerStartOffset needs to be applied only for the track in the main sequence. For subsequences, it's the subsequence section that
-			// is offset.
-			if (UUsdGeometryCacheAssetUserData* UserData = GeometryCache->GetAssetUserData<UUsdGeometryCacheAssetUserData>())
-			{
-				LayerStartOffsetSeconds = UserData->LayerStartOffsetSeconds;
-			}
+			StartSecondsInStage = UserData->StageStartOffsetSeconds;
 		}
 
-		const FFrameNumber StartOffsetTick = FFrameTime::FromDecimal(LayerStartOffsetSeconds * MovieScene->GetTickResolution().AsDecimal())
-												 .RoundToFrame();
+		const double StartTimeCodeInStage = StartSecondsInStage * UsdStage.GetTimeCodesPerSecond();
+		const double StartTimeCodeInLayer = (StartTimeCodeInStage - TotalOffset.Offset) / TotalOffset.Scale;
+		const double StartSecondsInLayer = StartTimeCodeInLayer / GeometryCacheLayer.GetTimeCodesPerSecond();
+		FFrameNumber StartTickInLayer = FFrameTime::FromDecimal(StartSecondsInLayer * MovieScene->GetTickResolution().AsDecimal()).RoundToFrame();
+
+		GeometryCacheTrack->Modify();
+		GeometryCacheTrack->RemoveAllAnimationData();
+
 		UMovieSceneGeometryCacheSection* NewSection = Cast<UMovieSceneGeometryCacheSection>(
-			GeometryCacheTrack->AddNewAnimation(StartOffsetTick, ComponentToBind)
+			GeometryCacheTrack->AddNewAnimation(StartTickInLayer, ComponentToBind)
 		);
 		NewSection->EvalOptions.CompletionMode = EMovieSceneCompletionMode::KeepState;
+
+		// The geometry cache assumes a playback rate respective to the stage, but we need to make this work with a playback
+		// rate respective to the LevelSequence the section is being placed in, as we'll end up putting timeScale on
+		// our parent sequences. E.g. if we have a sublayer with a scale of 2.0, we'll put a timeScale of 0.5 on the
+		// subsequence section, and here we'll set PlayRate to 2.0 as well: When playing the parent LevelSequence
+		// we'll see the "native" playback rate of the geometry cache, but if we played the child LevelSequence alone
+		// we'd see the geometry cache in the rate that matches the actual layer (even though the asset was produced for the 
+		// stage as a whole)
+		if (!FMath::IsNearlyEqual(TotalOffset.Scale, 1.0))
+		{
+			NewSection->Params.PlayRate = TotalOffset.Scale;
+
+			// See UMovieSceneGeometryCacheSection::PostEditChangeProperty
+			FPropertyChangedEvent PropertyChangeEvent{
+				FMovieSceneGeometryCacheParams::StaticStruct()->FindPropertyByName(GET_MEMBER_NAME_CHECKED(FMovieSceneGeometryCacheParams, PlayRate))
+			};
+			NewSection->Modify();
+			NewSection->PostEditChangeProperty(PropertyChangeEvent);
+		}
 	}
 
 	PrimPathByLevelSequenceName.AddUnique(GeometryCacheSequence->GetFName(), PrimPath.GetString());
@@ -2894,13 +3059,11 @@ void FUsdLevelSequenceHelperImpl::AddGroomTracks(const UUsdPrimTwin& PrimTwin, c
 		return;
 	}
 
-	UE::FSdfLayer GroomLayer = UsdUtils::FindLayerForPrim(Prim);
-	if (!GroomLayer)
-	{
-		return;
-	}
+	TArray<UE::FUsdAttribute> AnimatedAttributes = UsdUtils::GetAnimatedAttributes(Prim);
+	UE::FUsdAttribute StrongestAttribute = UsdUtils::GetStrongestAttribute(AnimatedAttributes, UsdUtils::GetEarliestTimeCode());
 
-	ULevelSequence* GroomAnimationSequence = FindOrAddSequenceForLayer(GroomLayer, GroomLayer.GetIdentifier(), GroomLayer.GetDisplayName());
+	UE::FSdfLayer GroomLayer;
+	ULevelSequence* GroomAnimationSequence = FindOrAddSequenceForAttribute(StrongestAttribute, &GroomLayer);
 	if (!GroomAnimationSequence)
 	{
 		return;
@@ -2939,21 +3102,6 @@ void FUsdLevelSequenceHelperImpl::AddVolumeTracks(const UUsdPrimTwin& PrimTwin,
 		return;
 	}
 
-	UE::FSdfLayer PrimLayer = UsdUtils::FindLayerForPrim(Prim);
-	ULevelSequence* PrimSequence = FindSequenceForIdentifier(PrimLayer.GetIdentifier());
-	if (!PrimSequence)
-	{
-		return;
-	}
-
-	UMovieScene* MovieScene = PrimSequence->GetMovieScene();
-	if (!MovieScene)
-	{
-		return;
-	}
-
-	UE::USDLevelSequenceHelper::Private::FScopedReadOnlyDisable ReadOnlyGuard{MovieScene, PrimLayer, UsdStage};
-
 	const FName PropertyPath = GET_MEMBER_NAME_CHECKED(UHeterogeneousVolumeComponent, Frame);
 
 	// Here we'll just get *any* of the filePath attrs from this Volume prim to check for the
@@ -2967,6 +3115,21 @@ void FUsdLevelSequenceHelperImpl::AddVolumeTracks(const UUsdPrimTwin& PrimTwin,
 	UE::FUsdAttribute MainAttr = Attrs[0];
 	const bool bIsMuted = MainAttr && MainAttr.GetNumTimeSamples() > 0 && UsdUtils::IsAttributeMuted(MainAttr, UsdStage);
 
+	UE::FSdfLayer PrimLayer;
+	ULevelSequence* FrameAttributeSequence = FindOrAddSequenceForAttribute(MainAttr, &PrimLayer);
+	if (!FrameAttributeSequence)
+	{
+		return;
+	}
+
+	UMovieScene* MovieScene = FrameAttributeSequence->GetMovieScene();
+	if (!MovieScene)
+	{
+		return;
+	}
+
+	UE::USDLevelSequenceHelper::Private::FScopedReadOnlyDisable ReadOnlyGuard{MovieScene, PrimLayer, UsdStage};
+
 	UUsdSparseVolumeTextureAssetUserData* UserData = nullptr;
 	const int32 ElementIndex = 0;
 	UMaterialInterface* CurrentMaterial = VolumeComponent->GetMaterial(ElementIndex);
@@ -3042,19 +3205,7 @@ void FUsdLevelSequenceHelperImpl::AddVolumeTracks(const UUsdPrimTwin& PrimTwin,
 
 	// We still have a Sequence transform though, as that converts the TimeSamples from being global
 	// to the stage to the particular subsequence where they are going to be added to.
-	//
-	// This may seem like it undoes the calculation in ConvertedTimeSamples, and it really does: For cases
-	// where we have a sublayer with an offset and scale, we'll end up adding the layer-local time samples
-	// to the subsequence, like we want. Using the CombinedOffset AND the SequenceTransform is needed for
-	// a different case however: Prim references with sublayer and offsets. In that case the prim itself
-	// may have a sublayer and offset, but its track will be placed in the levelsequence for the *referencer*
-	// layer: This means we want to see the keys on that layer instead, at times relative to it
-	FMovieSceneSequenceTransform SequenceTransform;
-	FMovieSceneSequenceID SequenceID = SequencesID.FindRef(PrimSequence);
-	if (FMovieSceneSubSequenceData* SubSequenceData = SequenceHierarchyCache.FindSubData(SequenceID))
-	{
-		SequenceTransform = SubSequenceData->RootToSequenceTransform;
-	}
+	const FMovieSceneSequenceTransform& SequenceTransform = GetSubsequenceTransform(FrameAttributeSequence);
 
 	// Unlike the other cases we can create our Reader right here, because the only thing we need to generate
 	// the track are the TimeSamples
@@ -3079,7 +3230,7 @@ void FUsdLevelSequenceHelperImpl::AddVolumeTracks(const UUsdPrimTwin& PrimTwin,
 		return static_cast<float>(TargetIndex);
 	};
 
-	if (UMovieSceneFloatTrack* FloatTrack = AddTrack<UMovieSceneFloatTrack>(PropertyPath, PrimTwin, *VolumeComponent, *PrimSequence, bIsMuted))
+	if (UMovieSceneFloatTrack* FloatTrack = AddTrack<UMovieSceneFloatTrack>(PropertyPath, PrimTwin, *VolumeComponent, *FrameAttributeSequence, bIsMuted))
 	{
 		// The component won't really linearly interpolate anything and will just do the analogous of constant interpolation,
 		// so it would be nice if our keys showed that too
@@ -3094,7 +3245,7 @@ void FUsdLevelSequenceHelperImpl::AddVolumeTracks(const UUsdPrimTwin& PrimTwin,
 		);
 	}
 
-	PrimPathByLevelSequenceName.AddUnique(PrimSequence->GetFName(), Prim.GetPrimPath().GetString());
+	PrimPathByLevelSequenceName.AddUnique(FrameAttributeSequence->GetFName(), Prim.GetPrimPath().GetString());
 }
 
 namespace UE::USDLevelSequenceHelper::Private
@@ -3266,12 +3417,7 @@ void FUsdLevelSequenceHelperImpl::AddAudioTracks(const UUsdPrimTwin& PrimTwin, c
 		TOptional<FFrameNumber> EndFrameNumber;
 		FFrameNumber MediaFrameNumberOffset;
 		{
-			FMovieSceneSequenceTransform SequenceTransform;
-			FMovieSceneSequenceID SequenceID = SequencesID.FindRef(AttributeSequence);
-			if (FMovieSceneSubSequenceData* SubSequenceData = SequenceHierarchyCache.FindSubData(SequenceID))
-			{
-				SequenceTransform = SubSequenceData->RootToSequenceTransform;
-			}
+			const FMovieSceneSequenceTransform& SequenceTransform = GetSubsequenceTransform(AttributeSequence);
 
 			const FFrameRate Resolution = MovieScene->GetTickResolution();
 			const FFrameRate DisplayRate = MovieScene->GetDisplayRate();
@@ -3360,17 +3506,11 @@ void FUsdLevelSequenceHelperImpl::AddAudioTracks(const UUsdPrimTwin& PrimTwin, c
 							}
 						}
 
-						FMovieSceneSequenceTransform SequenceTransform;
-						FMovieSceneSequenceID SequenceID = SequencesID.FindRef(AttributeSequence);
-						if (FMovieSceneSubSequenceData* SubSequenceData = SequenceHierarchyCache.FindSubData(SequenceID))
-						{
-							SequenceTransform = SubSequenceData->RootToSequenceTransform;
-						}
-
 						// Parse time-sampled animation into the volume channel
 						TArray<double> TimeSamples;
 						if (GainAttr.GetTimeSamples(TimeSamples) && TimeSamples.Num() > 0)
 						{
+							const FMovieSceneSequenceTransform& SequenceTransform = GetSubsequenceTransform(AttributeSequence);
 							UsdToUnreal::FPropertyTrackReader Reader = UsdToUnreal::CreatePropertyTrackReader(Prim, TEXT("Volume"));
 							UsdToUnreal::ConvertFloatTimeSamples(
 								UsdStage,
@@ -3420,38 +3560,11 @@ void FUsdLevelSequenceHelperImpl::AddPrim(UUsdPrimTwin& PrimTwin, bool bForceVis
 		return;
 	}
 
-	UE::FSdfPath PrimPath(*PrimTwin.PrimPath);
-	UE::FUsdPrim UsdPrim(UsdStage.GetPrimAtPath(PrimPath));
-
-	UE::FSdfLayer PrimLayer = UsdUtils::FindLayerForPrim(UsdPrim);
-	ULevelSequence* PrimSequence = FindSequenceForIdentifier(PrimLayer.GetIdentifier());
-
-	// PrimSequence is needed for a subsequence section to be created so fall back to the main sequence
-	if (!PrimSequence)
-	{
-		PrimSequence = MainLevelSequence;
-	}
-
-	TArray<UE::FUsdAttribute> PrimAttributes = UsdPrim.GetAttributes();
-
-	for (const UE::FUsdAttribute& PrimAttribute : PrimAttributes)
+	UE::FSdfPath PrimPath{*PrimTwin.PrimPath};
+	UE::FUsdPrim UsdPrim{UsdStage.GetPrimAtPath(PrimPath)};
+	if (!UsdPrim)
 	{
-		if (PrimAttribute.GetNumTimeSamples() > 0)
-		{
-			if (ULevelSequence* AttributeSequence = FindOrAddSequenceForAttribute(PrimAttribute))
-			{
-				PrimPathByLevelSequenceName.AddUnique(AttributeSequence->GetFName(), PrimTwin.PrimPath);
-
-				if (!SequencesID.Contains(AttributeSequence))
-				{
-					if (PrimSequence)
-					{
-						// Create new subsequence section for this referencing prim
-						CreateSubSequenceSection(*PrimSequence, *AttributeSequence);
-					}
-				}
-			}
-		}
+		return;
 	}
 
 	static IConsoleVariable* CVar = IConsoleManager::Get().FindConsoleVariable(TEXT("USD.DisableGeoCacheTracks"));
@@ -3682,11 +3795,7 @@ void FUsdLevelSequenceHelperImpl::UpdateControlRigTracks(UUsdPrimTwin& PrimTwin)
 	UE::FUsdAttribute ScalesAttr = SkelAnimationPrim.CreateAttribute(TEXT("scales"), TEXT("half3[]"));
 	UE::FUsdAttribute BlendShapeWeightsAttr = SkelAnimationPrim.GetAttribute(TEXT("blendShapeWeights"));
 
-	ULevelSequence* SkelAnimationSequence = FindOrAddSequenceForLayer(
-		SkelAnimationLayer,
-		SkelAnimationLayer.GetIdentifier(),
-		SkelAnimationLayer.GetDisplayName()
-	);
+	ULevelSequence* SkelAnimationSequence = FindOrAddSequenceForAttribute(JointsAttr, &SkelAnimationLayer);
 	if (!SkelAnimationSequence)
 	{
 		return;
@@ -3823,7 +3932,7 @@ void FUsdLevelSequenceHelperImpl::RemoveSequenceForPrim(ULevelSequence& Sequence
 
 			LevelSequencesByIdentifier.Remove(PrimTwin.PrimPath);
 			IdentifierByLevelSequence.Remove(&Sequence);
-			SequencesID.Remove(&Sequence);
+			DirtySequenceHierarchyCache();
 		}
 	}
 }
@@ -3926,24 +4035,24 @@ void FUsdLevelSequenceHelperImpl::UpdateUsdLayerOffsetFromSection(const UMovieSc
 	}
 
 	const FString* LayerIdentifier = LayerIdentifierByLevelSequenceName.Find(Sequence->GetName());
-	const FString* SubLayerIdentifier = LayerIdentifierByLevelSequenceName.Find(SubSequence->GetName());
+	const FString* SubSequenceLayerIdentifier = LayerIdentifierByLevelSequenceName.Find(SubSequence->GetName());
 
-	if (!LayerIdentifier || !SubLayerIdentifier)
+	if (!LayerIdentifier || !SubSequenceLayerIdentifier)
 	{
 		return;
 	}
 
 	FLayerTimeInfo* LayerTimeInfo = LayerTimeInfosByLayerIdentifier.Find(*LayerIdentifier);
-	FLayerTimeInfo* SubLayerTimeInfo = LayerTimeInfosByLayerIdentifier.Find(*SubLayerIdentifier);
+	FLayerTimeInfo* SubSequenceLayerTimeInfo = LayerTimeInfosByLayerIdentifier.Find(*SubSequenceLayerIdentifier);
 
-	if (!LayerTimeInfo || !SubLayerTimeInfo)
+	if (!LayerTimeInfo || !SubSequenceLayerTimeInfo)
 	{
 		return;
 	}
 
 	const double TimeCodesPerSecond = GetTimeCodesPerSecond();
-	const double SubStartTimeCode = SubLayerTimeInfo->StartTimeCode.Get(0.0);
-	const double SubEndTimeCode = SubLayerTimeInfo->EndTimeCode.Get(0.0);
+	const double SubStartTimeCode = SubSequenceLayerTimeInfo->StartTimeCode.Get(0.0);
+	const double SubEndTimeCode = SubSequenceLayerTimeInfo->EndTimeCode.Get(0.0);
 
 	FFrameRate TickResolution = MovieScene->GetTickResolution();
 	FFrameNumber ModifiedStartFrame = Section->GetInclusiveStartFrame();
@@ -3972,13 +4081,13 @@ void FUsdLevelSequenceHelperImpl::UpdateUsdLayerOffsetFromSection(const UMovieSc
 
 	// Prevent twins from being rebuilt when we update the layer offsets
 	TOptional<FScopedBlockNoticeListening> BlockNotices;
-
 	if (StageActor.IsValid())
 	{
 		BlockNotices.Emplace(StageActor.Get());
 	}
 
-	if (LocalLayersSequences.Contains(SubSequence->GetFName()))
+	// Check if we have a subsection due to a sublayer
+	if (UE::FSdfLayerOffset* SublayerOffset = LayerTimeInfo->SublayerIdentifierToOffsets.Find(*SubSequenceLayerIdentifier))
 	{
 		UE::FSdfLayer Layer = UE::FSdfLayer::FindOrOpen(*LayerTimeInfo->Identifier);
 		if (!Layer)
@@ -3987,17 +4096,25 @@ void FUsdLevelSequenceHelperImpl::UpdateUsdLayerOffsetFromSection(const UMovieSc
 			return;
 		}
 
+		const TArray<FString>& SubLayerPaths = Layer.GetSubLayerPaths();
+		const TArray<UE::FSdfLayerOffset>& SubLayerOffsets = Layer.GetSubLayerOffsets();
+		if (SubLayerPaths.Num() != SubLayerOffsets.Num())
+		{
+			return;
+		}
+
 		int32 SubLayerIndex = INDEX_NONE;
-		FLayerOffsetInfo* SubLayerOffset = Algo::FindByPredicate(
-			LayerTimeInfo->SubLayersOffsets,
-			[&SubLayerIndex, &SubLayerIdentifier = SubLayerTimeInfo->Identifier](const FLayerOffsetInfo& Other)
+		for (int32 Index = 0; Index < SubLayerPaths.Num(); ++Index)
+		{
+			if (UE::FSdfLayer SubLayer = UsdUtils::FindLayerForSubLayerPath(Layer, SubLayerPaths[Index]))
 			{
-				bool bFound = (Other.LayerIdentifier == SubLayerIdentifier);
-				++SubLayerIndex;
-
-				return bFound;
+				if (SubLayer.GetIdentifier() == *SubSequenceLayerIdentifier)
+				{
+					SubLayerIndex = Index;
+					break;
+				}
 			}
-		);
+		}
 
 		if (SubLayerIndex != INDEX_NONE)
 		{
@@ -4005,15 +4122,19 @@ void FUsdLevelSequenceHelperImpl::UpdateUsdLayerOffsetFromSection(const UMovieSc
 			UpdateLayerTimeInfoFromLayer(*LayerTimeInfo, Layer);
 		}
 	}
-	else
+	// Check if we have a subsection due to a reference
+	else if (UE::FSdfLayerOffset* ReferencedOffset = LayerTimeInfo->ReferenceIdentifierToOffsets.Find(*SubSequenceLayerIdentifier))
 	{
-		TArray<FString> PrimPathsForSequence;
-		PrimPathByLevelSequenceName.MultiFind(Section->GetSequence()->GetFName(), PrimPathsForSequence);
-
-		for (const FString& PrimPath : PrimPathsForSequence)
-		{
-			UsdUtils::SetRefOrPayloadLayerOffset(UsdStage.GetPrimAtPath(UE::FSdfPath(*PrimPath)), NewLayerOffset);
-		}
+		// TODO: Handle updates to subsections created due to references and payloads. We'd need the referencer prim path at this
+		// point, which we do not record
+		return;
+	}
+	// Check if we have a subsection due to a payload
+	else if (UE::FSdfLayerOffset* PayloadOffset = LayerTimeInfo->PayloadIdentifierToOffsets.Find(*SubSequenceLayerIdentifier))
+	{
+		// TODO: Handle updates to subsections created due to references and payloads. We'd need the referencer prim path at this
+		// point, which we do not record
+		return;
 	}
 }
 
@@ -4314,7 +4435,11 @@ FGuid FUsdLevelSequenceHelperImpl::GetOrCreateComponentBinding(
 
 	FPrimTwinBindings& Bindings = PrimTwinToBindings.FindOrAdd(&PrimTwin);
 
-	ensure(Bindings.Sequence == nullptr || Bindings.Sequence == &Sequence);
+	// TODO: Better support having multiple bindings for the same same on different sequences.
+	// We used to ensure "Bindings.Sequence == &Sequence" here, but we shouldn't do that as it's very simple and valid to
+	// construct a stage where a prim's strongest opinion for two animated attributes comes from different layers, and it
+	// would have triggered this ensure
+
 	Bindings.Sequence = &Sequence;
 
 	if (FGuid* ExistingGuid = Bindings.ObjectClassToBindingGuid.Find(ComponentToBind.GetClass()))
@@ -4751,12 +4876,7 @@ void FUsdLevelSequenceHelperImpl::HandleControlRigSectionChange(UMovieSceneContr
 		}
 	}
 
-	FMovieSceneSequenceTransform SequenceTransform;
-	FMovieSceneSequenceID SequenceID = SequencesID.FindRef(LevelSequence);
-	if (FMovieSceneSubSequenceData* SubSequenceData = SequenceHierarchyCache.FindSubData(SequenceID))
-	{
-		SequenceTransform = SubSequenceData->RootToSequenceTransform;
-	}
+	const FMovieSceneSequenceTransform& SequenceTransform = GetSubsequenceTransform(LevelSequence);
 
 	// Actually bake inside the UsdUtilities module as we need to manipulate USD arrays a lot
 	const UsdUtils::FBlendShapeMap& BlendShapeMap = StageActorValue->GetBlendShapeMap();
@@ -5002,16 +5122,6 @@ void FUsdLevelSequenceHelperImpl::HandleTrackChange(UMovieSceneTrack& Track, boo
 		}
 		else
 		{
-			FMovieSceneSequenceTransform SequenceTransform;
-
-			if (const FMovieSceneSequenceID* SequenceID = SequencesID.Find(Sequence))
-			{
-				if (FMovieSceneSubSequenceData* SubSequenceData = SequenceHierarchyCache.FindSubData(*SequenceID))
-				{
-					SequenceTransform = SubSequenceData->RootToSequenceTransform;
-				}
-			}
-
 			// Right now we don't write out changes to SkeletalAnimation tracks, and only property tracks... the UAnimSequence
 			// asset can't be modified all that much in UE anyway. Later on we may want to enable writing it out anyway though,
 			// and pick up on changes to the section offset or play rate and bake out the UAnimSequence again
@@ -5021,6 +5131,8 @@ void FUsdLevelSequenceHelperImpl::HandleTrackChange(UMovieSceneTrack& Track, boo
 				UE::USDLevelSequenceHelper::Private::ShowVisibilityWarningIfNeeded(PropertyTrack, UsdPrim);
 #endif
 
+				FMovieSceneSequenceTransform SequenceTransform = GetSubsequenceTransform(Sequence);
+
 				TSet<FName> PropertyPathsToRefresh;
 				UnrealToUsd::FPropertyTrackWriter
 					Writer = UnrealToUsd::CreatePropertyTrackWriter(*BoundSceneComponent, *PropertyTrack, UsdPrim, PropertyPathsToRefresh);
@@ -5186,28 +5298,26 @@ void FUsdLevelSequenceHelperImpl::UpdateLayerTimeInfoFromLayer(FLayerTimeInfo& L
 	}
 
 	const TArray<FString>& SubLayerPaths = Layer.GetSubLayerPaths();
-	LayerTimeInfo.SubLayersOffsets.Empty(SubLayerPaths.Num());
+	const TArray<UE::FSdfLayerOffset>& SubLayerOffsets = Layer.GetSubLayerOffsets();
+	if (SubLayerPaths.Num() != SubLayerOffsets.Num())
+	{
+		return;
+	}
 
-	int32 SubLayerIndex = 0;
-	for (const UE::FSdfLayerOffset& SubLayerOffset : Layer.GetSubLayerOffsets())
+	LayerTimeInfo.SublayerIdentifierToOffsets.Empty(SubLayerPaths.Num());
+	for (int32 Index = 0; Index < SubLayerPaths.Num(); ++Index)
 	{
-		if (SubLayerPaths.IsValidIndex(SubLayerIndex))
-		{
-			if (UE::FSdfLayer SubLayer = UsdUtils::FindLayerForSubLayerPath(Layer, SubLayerPaths[SubLayerIndex]))
-			{
-				FLayerOffsetInfo SubLayerOffsetInfo;
-				SubLayerOffsetInfo.LayerIdentifier = SubLayer.GetIdentifier();
-				SubLayerOffsetInfo.LayerOffset = SubLayerOffset;
+		const FString& SubLayerPath = SubLayerPaths[Index];
+		const UE::FSdfLayerOffset& SubLayerOffset = SubLayerOffsets[Index];
 
-				LayerTimeInfo.SubLayersOffsets.Add(MoveTemp(SubLayerOffsetInfo));
-			}
+		if (UE::FSdfLayer SubLayer = UsdUtils::FindLayerForSubLayerPath(Layer, SubLayerPath))
+		{
+			LayerTimeInfo.SublayerIdentifierToOffsets.Add(SubLayer.GetIdentifier(), SubLayerOffset);
 		}
-
-		++SubLayerIndex;
 	}
 }
 
-ULevelSequence* FUsdLevelSequenceHelperImpl::FindSequenceForIdentifier(const FString& SequenceIdentitifer)
+ULevelSequence* FUsdLevelSequenceHelperImpl::FindSequenceForIdentifier(const FString& SequenceIdentitifer) const
 {
 	ULevelSequence* Sequence = nullptr;
 	if (auto* FoundSequence = LevelSequencesByIdentifier.Find(SequenceIdentitifer))
diff --git a/Engine/Plugins/Runtime/GeometryCache/Source/GeometryCacheTracks/Classes/MovieSceneGeometryCacheSection.h b/Engine/Plugins/Runtime/GeometryCache/Source/GeometryCacheTracks/Classes/MovieSceneGeometryCacheSection.h
--- a/Engine/Plugins/Runtime/GeometryCache/Source/GeometryCacheTracks/Classes/MovieSceneGeometryCacheSection.h
+++ b/Engine/Plugins/Runtime/GeometryCache/Source/GeometryCacheTracks/Classes/MovieSceneGeometryCacheSection.h
@@ -84,7 +84,7 @@ protected:
 	virtual void PostLoad() override;
 	virtual void Serialize(FArchive& Ar) override;
 
-private:
+public:
 
 	//~ UObject interface
 
diff --git a/Engine/Plugins/Runtime/USDCore/Source/USDClasses/Public/USDAssetUserData.h b/Engine/Plugins/Runtime/USDCore/Source/USDClasses/Public/USDAssetUserData.h
--- a/Engine/Plugins/Runtime/USDCore/Source/USDClasses/Public/USDAssetUserData.h
+++ b/Engine/Plugins/Runtime/USDCore/Source/USDClasses/Public/USDAssetUserData.h
@@ -36,13 +36,28 @@ class UUsdAnimSequenceAssetUserData : public UUsdAssetUserData
 	GENERATED_BODY()
 
 public:
-	// Offset into the stage in seconds to when this AnimSequence should start playing to match the
-	// skeletal animation in that stage. This is required because UAnimSequences just range from the
-	// first skeletal timeSample to the last, and the first sample is not necessarily the startTimeCode
-	// for the stage. Note that this is wrt. the startTimeCode of the *layer*, and not the composed stage.
-	// This is only used for animating USkeletalMeshComponents via the AUsdStageActor, and it is in
-	// seconds since it will need to drive the UAnimSequence with a time value in the [0, LenghtSeconds] range.
-	// This should be applied *after* any offset/scale conversions on the time coordinate.
+	// Offset into the stage in seconds to when this AnimSequence should start
+	// playing to match the skeletal animation in that stage. This is required
+	// because UAnimSequences just range from the first skeletal timeSample to the
+	// last, and the first sample is not necessarily the startTimeCode for the
+	// stage. Note that this is wrt. the startTimeCode of the *stage*, and not the
+	// individual layer. This is used for animating USkeletalMeshComponents via
+	// the AUsdStageActor and building the LevelSequences, and it is in seconds 
+	// since it will need to drive the UAnimSequence with a time value in the 
+	// [0, LenghtSeconds] range. This should be applied *after* any offset/scale 
+	// conversions on the time coordinate.
+	UPROPERTY()
+	float StageStartOffsetSeconds = 0.0f;
+
+	// Total scale difference (as in, "MapToRoot") between the baked AnimSequence
+	// playback speed and the stage playback speed.
+	// E.g. If we have root.usda referencing ref1.usda with scale 2.0, and ref1.usda has
+	// sublayer animseq.usda with scale 3.0 (where the SkelAnimation prim is authored),
+	// then this property will hold 6.0
+	UPROPERTY()
+	float StageScaleSeconds = 1.0f;
+
+	UE_DEPRECATED(5.7, "Use StageStartOffsetSeconds instead.")
 	UPROPERTY()
 	float LayerStartOffsetSeconds = 0.0f;
 };
@@ -113,7 +128,11 @@ class UUsdGeometryCacheAssetUserData : public UUsdMeshAssetUserData
 
 public:
 	// Check analogous comment on UUsdAnimSequenceAssetUserData
-	UPROPERTY(EditAnywhere, BlueprintReadWrite, Category = "USD")
+	UPROPERTY()
+	float StageStartOffsetSeconds = 0.0f;
+
+	UE_DEPRECATED(5.7, "Use StageStartOffsetSeconds instead.")
+	UPROPERTY()
 	float LayerStartOffsetSeconds = 0.0f;
 };
 
diff --git a/Engine/Plugins/Runtime/USDCore/Source/USDUtilities/Private/USDAttributeUtils.cpp b/Engine/Plugins/Runtime/USDCore/Source/USDUtilities/Private/USDAttributeUtils.cpp
--- a/Engine/Plugins/Runtime/USDCore/Source/USDUtilities/Private/USDAttributeUtils.cpp
+++ b/Engine/Plugins/Runtime/USDCore/Source/USDUtilities/Private/USDAttributeUtils.cpp
@@ -9,8 +9,10 @@
 #include "USDProjectSettings.h"
 #include "USDTypesConversion.h"
 
+#include "UsdWrappers/PcpPrimIndex.h"
 #include "UsdWrappers/SdfLayer.h"
 #include "UsdWrappers/SdfPath.h"
+#include "UsdWrappers/SdfPrimSpec.h"
 #include "UsdWrappers/UsdAttribute.h"
 #include "UsdWrappers/UsdPrim.h"
 #include "UsdWrappers/UsdStage.h"
@@ -507,6 +509,154 @@ bool UsdUtils::NotifyIfInstanceProxy(const pxr::UsdPrim& Prim)
 
 	return false;
 }
+
+UE::FUsdAttribute UsdUtils::GetStrongestAttribute(const TArray<UE::FUsdAttribute>& Attributes, double TimeCode)
+{
+	if (Attributes.Num() == 0)
+	{
+		return {};
+	}
+
+	TMap<UE::FSdfLayer, TArray<UE::FUsdAttribute>> LayerToAttribute;
+	LayerToAttribute.Reserve(Attributes.Num());
+
+	UE::FUsdPrim CommonPrim = Attributes[0].GetPrim();
+
+	// Collect the layers where we have prim stacks
+	TArray<UE::FSdfPrimSpec> SpecStack = CommonPrim.GetPrimStack();
+	TArray<UE::FSdfLayer> PrimStackLayers;
+	PrimStackLayers.Reserve(SpecStack.Num());
+	for (const UE::FSdfPrimSpec& Spec : SpecStack)
+	{
+		if (!Spec)
+		{
+			continue;
+		}
+
+		UE::FSdfLayer Layer = Spec.GetLayer();
+		if (!Layer)
+		{
+			continue;
+		}
+
+		PrimStackLayers.Add(Layer);
+		LayerToAttribute.Add(Layer);
+	}
+
+	// Collect all the layers where we have attribute specs, and in which layers those attributes have their strongest specs
+	for (const UE::FUsdAttribute& Attribute : Attributes)
+	{
+		UE::FUsdPrim Prim = Attributes[0].GetPrim();
+		if (Prim != CommonPrim)
+		{
+			USD_LOG_WARNING(
+				TEXT("Found different prims '%s' and '%s' when trying to collect strongest attributes"),
+				*Prim.GetPrimPath().GetString(),
+				*CommonPrim.GetPrimPath().GetString()
+			);
+			return {};
+		}
+
+		FScopedUsdAllocs UsdAllocs;
+
+		TArray<UE::FSdfLayer> AttributeStackLayers;
+		std::vector<pxr::SdfPropertySpecHandle> PropertyStack = pxr::UsdAttribute{Attribute}.GetPropertyStack(TimeCode);
+		for (const pxr::SdfPropertySpecHandle& PropertySpec : PropertyStack)
+		{
+			if (!PropertySpec)
+			{
+				continue;
+			}
+
+			if (!PropertySpec->HasDefaultValue() && PropertySpec->GetLayer()->GetNumTimeSamplesForPath(PropertySpec->GetPath()) == 0)
+			{
+				continue;
+			}
+
+			pxr::SdfLayerHandle SpecLayer = PropertySpec->GetLayer();
+			if (!SpecLayer)
+			{
+				continue;
+			}
+
+			AttributeStackLayers.Add(UE::FSdfLayer{SpecLayer});
+		}
+		if (AttributeStackLayers.Num() == 0)
+		{
+			continue;
+		}
+
+		if (TArray<UE::FUsdAttribute>* FoundAttributes = LayerToAttribute.Find(AttributeStackLayers[0]))
+		{
+			FoundAttributes->Add(Attribute);
+		}
+		else
+		{
+			// In the general case if we have an attribute spec for something, we'll have a prim spec for its prim in the
+			// same layer. For value clips this seems to not be the case: We can have an attribute spec without an owning
+			// prim spec in the same layer somehow...
+			//
+			// Here we try to work around that by merging the layers we know about: If we find an attribute spec layer that we
+			// don't know about, we pretend we have a corresponding prim spec with the same strength
+			for (int32 AttributeIndex = 0, PrimIndex = 0; AttributeIndex < AttributeStackLayers.Num(); ++AttributeIndex, ++PrimIndex)
+			{
+				if (!ensure(PrimStackLayers.IsValidIndex(PrimIndex)))
+				{
+					// Shouldn't ever happen but it's hard to predict all USD scenarios
+					break;
+				}
+
+				const UE::FSdfLayer& AttributeLayer = AttributeStackLayers[AttributeIndex];
+				const UE::FSdfLayer& PrimLayer = PrimStackLayers[PrimIndex];
+
+				if (LayerToAttribute.Contains(AttributeLayer))
+				{
+					// We already have this layer in the prim stack, resync our indices here so we continue traversing
+					// prim and layer stacks in tandem until the next different layer
+					if (PrimLayer != AttributeLayer)
+					{
+						AttributeIndex = PrimStackLayers.IndexOfByKey(PrimLayer);
+					}
+					continue;
+				}
+				else
+				{
+					PrimStackLayers.Insert(AttributeLayer, PrimIndex);
+					LayerToAttribute.FindOrAdd(AttributeLayer).Add(Attribute);
+
+					// We only really care about the strongest spec anyway
+					break;
+				}
+			}
+		}
+	}
+
+	// Early out here in case all our attributes come from the same layer anyway
+	if (LayerToAttribute.Num() == 1)
+	{
+		const TArray<UE::FUsdAttribute>& FoundAttrs = (LayerToAttribute.CreateIterator())->Value;
+		if (FoundAttrs.Num() > 0)
+		{
+			// Defer to the order the attributes were in the Attributes array in case of ties
+			return FoundAttrs[0];
+		}
+
+		return {};
+	}
+
+	// Return the attribute that has a spec on our strongest prim spec layer
+	for (const UE::FSdfLayer& Layer : PrimStackLayers)
+	{
+		TArray<UE::FUsdAttribute>* FoundAttribute = LayerToAttribute.Find(Layer);
+		if (FoundAttribute && FoundAttribute->Num() > 0)
+		{
+			return (*FoundAttribute)[0];
+		}
+	}
+
+	return {};
+}
+
 #endif	  // #if USE_USD_SDK
 
 #undef LOCTEXT_NAMESPACE
diff --git a/Engine/Plugins/Runtime/USDCore/Source/USDUtilities/Private/USDConversionUtils.cpp b/Engine/Plugins/Runtime/USDCore/Source/USDUtilities/Private/USDConversionUtils.cpp
--- a/Engine/Plugins/Runtime/USDCore/Source/USDUtilities/Private/USDConversionUtils.cpp
+++ b/Engine/Plugins/Runtime/USDCore/Source/USDUtilities/Private/USDConversionUtils.cpp
@@ -1451,6 +1451,30 @@ bool UsdUtils::GetAnimatedBoundsTimeSamples(
 	return bHasAnimatedBounds;
 }
 
+TArray<UE::FUsdAttribute> UsdUtils::GetAnimatedAttributes(const pxr::UsdPrim& Prim)
+{
+	if (!Prim || !Prim.IsActive())
+	{
+		return {};
+	}
+
+	TUsdStore<std::vector<pxr::UsdAttribute>> UsdAttributes = Prim.GetAttributes();
+
+	TArray<UE::FUsdAttribute> Attributes;
+	Attributes.Reserve(UsdAttributes->size());
+
+	for (const pxr::UsdAttribute& UsdAttribute : *UsdAttributes)
+	{
+		TUsdStore<std::vector<double>> TimeSamples;
+		if (UsdAttribute.GetTimeSamples(&TimeSamples.Get()) && TimeSamples->size() > 0)
+		{
+			Attributes.Add(UE::FUsdAttribute{UsdAttribute});
+		}
+	}
+
+	return Attributes;
+}
+
 bool UsdUtils::HasAuthoredKind(const pxr::UsdPrim& Prim)
 {
 	FScopedUsdAllocs Allocs;
diff --git a/Engine/Plugins/Runtime/USDCore/Source/USDUtilities/Private/USDGeomMeshConversion.cpp b/Engine/Plugins/Runtime/USDCore/Source/USDUtilities/Private/USDGeomMeshConversion.cpp
--- a/Engine/Plugins/Runtime/USDCore/Source/USDUtilities/Private/USDGeomMeshConversion.cpp
+++ b/Engine/Plugins/Runtime/USDCore/Source/USDUtilities/Private/USDGeomMeshConversion.cpp
@@ -20,6 +20,7 @@
 #include "USDTypesConversion.h"
 
 #include "UsdWrappers/SdfPath.h"
+#include "UsdWrappers/UsdAttribute.h"
 #include "UsdWrappers/UsdPrim.h"
 #include "UsdWrappers/UsdStage.h"
 
@@ -7028,14 +7029,43 @@ void UsdUtils::GetAnimatedMeshTimeCodes(const UE::FUsdStage& Stage, const FStrin
 		return;
 	}
 
-	FScopedUsdAllocs Allocs;
-
-	pxr::UsdPrim UsdPrim = pxr::UsdPrim(Stage.GetPrimAtPath(UE::FSdfPath(*PrimPath)));
+	UE::FUsdPrim UsdPrim = (Stage.GetPrimAtPath(UE::FSdfPath(*PrimPath)));
 	if (!UsdPrim)
 	{
 		return;
 	}
 
+	TArray<UE::FUsdAttribute> Attributes = UsdUtils::GetAnimatedMeshAttributes(UsdPrim);
+
+	TOptional<double> MinStartTimeCode;
+	TOptional<double> MaxEndTimeCode;
+
+	for (const UE::FUsdAttribute& Attribute : Attributes)
+	{
+		TArray<double> TimeSamples;
+		if (Attribute.GetTimeSamples(TimeSamples) && TimeSamples.Num() > 0)
+		{
+			MinStartTimeCode = FMath::Min(MinStartTimeCode.Get(TNumericLimits<double>::Max()), TimeSamples[0]);
+			MaxEndTimeCode = FMath::Max(MaxEndTimeCode.Get(TNumericLimits<double>::Lowest()), TimeSamples[TimeSamples.Num() - 1]);
+		}
+	}
+
+	if (MinStartTimeCode.IsSet() && MaxEndTimeCode.IsSet())
+	{
+		OutStartFrame = FMath::FloorToInt(MinStartTimeCode.GetValue());
+		OutEndFrame = FMath::CeilToInt(MaxEndTimeCode.GetValue());
+	}
+}
+
+TArray<UE::FUsdAttribute> UsdUtils::GetAnimatedMeshAttributes(const pxr::UsdPrim& Prim)
+{
+	FScopedUsdAllocs Allocs;
+
+	if (!Prim)
+	{
+		return {};
+	}
+
 	constexpr bool bIncludeInherited = false;
 	pxr::TfTokenVector GeomMeshAttributeNames = pxr::UsdGeomMesh::GetSchemaAttributeNames(bIncludeInherited);
 	pxr::TfTokenVector GeomPointBasedAttributeNames = pxr::UsdGeomPointBased::GetSchemaAttributeNames(bIncludeInherited);
@@ -7043,27 +7073,22 @@ void UsdUtils::GetAnimatedMeshTimeCodes(const UE::FUsdStage& Stage, const FStrin
 	GeomMeshAttributeNames.reserve(GeomMeshAttributeNames.size() + GeomPointBasedAttributeNames.size());
 	GeomMeshAttributeNames.insert(GeomMeshAttributeNames.end(), GeomPointBasedAttributeNames.begin(), GeomPointBasedAttributeNames.end());
 
-	TOptional<double> MinStartTimeCode;
-	TOptional<double> MaxEndTimeCode;
+	TArray<UE::FUsdAttribute> AnimatedAttributes;
+	AnimatedAttributes.Reserve(GeomMeshAttributeNames.size());
 
 	for (const pxr::TfToken& AttributeName : GeomMeshAttributeNames)
 	{
-		if (const pxr::UsdAttribute& Attribute = UsdPrim.GetAttribute(AttributeName))
+		if (const pxr::UsdAttribute& Attribute = Prim.GetAttribute(AttributeName))
 		{
 			std::vector<double> TimeSamples;
 			if (Attribute.GetTimeSamples(&TimeSamples) && TimeSamples.size() > 0)
 			{
-				MinStartTimeCode = FMath::Min(MinStartTimeCode.Get(TNumericLimits<double>::Max()), TimeSamples[0]);
-				MaxEndTimeCode = FMath::Max(MaxEndTimeCode.Get(TNumericLimits<double>::Lowest()), TimeSamples[TimeSamples.size() - 1]);
+				AnimatedAttributes.Add(UE::FUsdAttribute{Attribute});
 			}
 		}
 	}
 
-	if (MinStartTimeCode.IsSet() && MaxEndTimeCode.IsSet())
-	{
-		OutStartFrame = FMath::FloorToInt(MinStartTimeCode.GetValue());
-		OutEndFrame = FMath::CeilToInt(MaxEndTimeCode.GetValue());
-	}
+	return AnimatedAttributes;
 }
 
 #undef LOCTEXT_NAMESPACE
diff --git a/Engine/Plugins/Runtime/USDCore/Source/USDUtilities/Private/USDLayerUtils.cpp b/Engine/Plugins/Runtime/USDCore/Source/USDUtilities/Private/USDLayerUtils.cpp
--- a/Engine/Plugins/Runtime/USDCore/Source/USDUtilities/Private/USDLayerUtils.cpp
+++ b/Engine/Plugins/Runtime/USDCore/Source/USDUtilities/Private/USDLayerUtils.cpp
@@ -3,6 +3,7 @@
 #include "USDLayerUtils.h"
 
 #include "UnrealUSDWrapper.h"
+#include "USDAttributeUtils.h"
 #include "USDErrorUtils.h"
 #include "USDMemory.h"
 #include "USDTypesConversion.h"
@@ -22,7 +23,6 @@
 #endif	  // WITH_EDITOR
 
 #if USE_USD_SDK
-
 #include "USDIncludesStart.h"
 #include "pxr/usd/pcp/layerStack.h"
 #include "pxr/usd/sdf/fileFormat.h"
@@ -517,7 +517,7 @@ UE::FSdfLayer UsdUtils::FindLayerForAttribute(const pxr::UsdAttribute& Attribute
 
 	for (const pxr::SdfPropertySpecHandle& PropertySpec : Attribute.GetPropertyStack(TimeCode))
 	{
-		if (PropertySpec->HasDefaultValue() || PropertySpec->GetLayer()->GetNumTimeSamplesForPath(PropertySpec->GetPath()) > 0)
+		if (PropertySpec)
 		{
 			return UE::FSdfLayer(PropertySpec->GetLayer());
 		}
diff --git a/Engine/Plugins/Runtime/USDCore/Source/USDUtilities/Private/USDSkeletalDataConversion.cpp b/Engine/Plugins/Runtime/USDCore/Source/USDUtilities/Private/USDSkeletalDataConversion.cpp
--- a/Engine/Plugins/Runtime/USDCore/Source/USDUtilities/Private/USDSkeletalDataConversion.cpp
+++ b/Engine/Plugins/Runtime/USDCore/Source/USDUtilities/Private/USDSkeletalDataConversion.cpp
@@ -2062,7 +2062,8 @@ bool UsdToUnreal::ConvertSkelAnim(
 	bool bInInterpretLODs,
 	const pxr::UsdPrim& RootMotionPrim,
 	UAnimSequence* OutSkeletalAnimationAsset,
-	float* OutStartOffsetSeconds
+	float* OutStartOffsetSeconds,
+	float* OutScaleSeconds
 )
 {
 	TRACE_CPUPROFILER_EVENT_SCOPE(UsdToUnreal::ConvertSkelAnim);
@@ -2186,23 +2187,23 @@ bool UsdToUnreal::ConvertSkelAnim(
 	// which is important because later our composition of tracks and subsections within a LevelSequence will reapply analogous
 	// offsets and scalings anyway
 
-	const double StageStartTimeCode = FMath::Min(
+	const double StartTimeCodeInStage = FMath::Min(
 		FirstJointSampleTimeCode.Get(TNumericLimits<double>::Max()),
 		FMath::Min(FirstBlendShapeSampleTimeCode.Get(TNumericLimits<double>::Max()), FirstRootMotionTimeCode.Get(TNumericLimits<double>::Max()))
 	);
-	const double StageEndTimeCode = FMath::Max(
+	const double EndTimeCodeInStage = FMath::Max(
 		LastJointSampleTimeCode.Get(TNumericLimits<double>::Lowest()),
 		FMath::Max(LastBlendShapeSampleTimeCode.Get(TNumericLimits<double>::Lowest()), LastRootMotionTimeCode.Get(TNumericLimits<double>::Lowest()))
 	);
-	const double StageStartSeconds = StageStartTimeCode / StageTimeCodesPerSecond;
-	const double StageSequenceLengthTimeCodes = StageEndTimeCode - StageStartTimeCode;
-	const double LayerSequenceLengthTimeCodes = StageSequenceLengthTimeCodes / Offset.Scale;
-	const double LayerSequenceLengthSeconds = FMath::Max<double>(LayerSequenceLengthTimeCodes / LayerTimeCodesPerSecond, MINIMUM_ANIMATION_LENGTH);
-	const double LayerStartTimeCode = (StageStartTimeCode - Offset.Offset) / Offset.Scale;
-	const double LayerStartSeconds = LayerStartTimeCode / LayerTimeCodesPerSecond;
+	const double StartSecondsInStage = StartTimeCodeInStage / StageTimeCodesPerSecond;
+	const double SequenceLengthTimeCodesInStage = EndTimeCodeInStage - StartTimeCodeInStage;
+	const double SequenceLengthTimeCodesInLayer = SequenceLengthTimeCodesInStage / Offset.Scale;
+	const double SequenceLengthSecondsInLayer = FMath::Max<double>(SequenceLengthTimeCodesInLayer / LayerTimeCodesPerSecond, MINIMUM_ANIMATION_LENGTH);
+	const double StartTimeCodeInLayer = (StartTimeCodeInStage - Offset.Offset) / Offset.Scale;
+	const double StartSecondsInLayer = StartTimeCodeInLayer / LayerTimeCodesPerSecond;
 
 	// Just bake each time code in the source layer as a frame
-	const int32 NumBakedFrames = FMath::RoundToInt(FMath::Max(LayerSequenceLengthSeconds * LayerTimeCodesPerSecond + 1.0, 1.0));
+	const int32 NumBakedFrames = FMath::RoundToInt(FMath::Max(SequenceLengthSecondsInLayer * LayerTimeCodesPerSecond + 1.0, 1.0));
 	const double StageBakeIntervalTimeCodes = 1.0 * Offset.Scale;
 
 	IAnimationDataController& Controller = OutSkeletalAnimationAsset->GetController();
@@ -2240,7 +2241,7 @@ bool UsdToUnreal::ConvertSkelAnim(
 		pxr::VtArray<pxr::GfMatrix4d> UsdJointTransforms;
 		for (int32 FrameIndex = 0; FrameIndex < NumBakedFrames; ++FrameIndex)
 		{
-			const double StageFrameTimeCodes = StageStartTimeCode + FrameIndex * StageBakeIntervalTimeCodes;
+			const double StageFrameTimeCodes = StartTimeCodeInStage + FrameIndex * StageBakeIntervalTimeCodes;
 
 			InUsdSkeletonQuery.ComputeJointLocalTransforms(&UsdJointTransforms, StageFrameTimeCodes);
 			for (int32 BoneIndex = 0; BoneIndex < NumBonesInUE; ++BoneIndex)
@@ -2319,9 +2320,9 @@ bool UsdToUnreal::ConvertSkelAnim(
 			pxr::VtArray<float> WeightsForFrame;
 			for (int32 FrameIndex = 0; FrameIndex < NumBakedFrames; ++FrameIndex)
 			{
-				const double StageFrameTimeCodes = StageStartTimeCode + FrameIndex * StageBakeIntervalTimeCodes;
+				const double StageFrameTimeCodes = StartTimeCodeInStage + FrameIndex * StageBakeIntervalTimeCodes;
 				const double LayerFrameTimeCodes = (StageFrameTimeCodes - Offset.Offset) / Offset.Scale;
-				const double LayerFrameSeconds = LayerFrameTimeCodes / LayerTimeCodesPerSecond - LayerStartSeconds;
+				const double LayerFrameSeconds = LayerFrameTimeCodes / LayerTimeCodesPerSecond - StartSecondsInLayer;
 
 				UsdAnimQuery.ComputeBlendShapeWeights(&WeightsForFrame, pxr::UsdTimeCode(StageFrameTimeCodes));
 
@@ -2497,7 +2498,7 @@ bool UsdToUnreal::ConvertSkelAnim(
 
 	const FFrameRate FrameRate(LayerTimeCodesPerSecond, 1);
 	Controller.SetFrameRate(FrameRate, bShouldTransact);
-	const FFrameNumber FrameNumber = FrameRate.AsFrameNumber(LayerSequenceLengthSeconds);
+	const FFrameNumber FrameNumber = FrameRate.AsFrameNumber(SequenceLengthSecondsInLayer);
 	Controller.SetNumberOfFrames(FrameNumber, bShouldTransact);
 	Controller.NotifyPopulated();	 // This call is important to get the controller to not use the sampling frequency as framerate
 	Controller.CloseBracket(bShouldTransact);
@@ -2507,11 +2508,12 @@ bool UsdToUnreal::ConvertSkelAnim(
 
 	if (OutStartOffsetSeconds)
 	{
-		// We don't want to store just StartSeconds here, because part of that may be because the layer itself
-		// has an offset/scale within the stage. In OutStartOffsetSeconds we need to store the start of the animation
-		// in seconds *with respect to its own layer*. The layer's offset/scale can be retrieved later at any time by
-		// just looking at the FSdfLayerOffset for the SkelAnimation prim (like what is done in FUsdSkelSkeletonTranslator::UpdateComponents)
-		*OutStartOffsetSeconds = LayerStartSeconds;
+		*OutStartOffsetSeconds = StartSecondsInStage;
+	}
+
+	if (OutScaleSeconds)
+	{
+		*OutScaleSeconds = Offset.Scale;
 	}
 
 	return true;
diff --git a/Engine/Plugins/Runtime/USDCore/Source/USDUtilities/Public/USDAttributeUtils.h b/Engine/Plugins/Runtime/USDCore/Source/USDUtilities/Public/USDAttributeUtils.h
--- a/Engine/Plugins/Runtime/USDCore/Source/USDUtilities/Public/USDAttributeUtils.h
+++ b/Engine/Plugins/Runtime/USDCore/Source/USDUtilities/Public/USDAttributeUtils.h
@@ -56,5 +56,8 @@ namespace UsdUtils
 	 * This will show a warning and return true if Prim is an instance, letting us avoid those scenarios
 	 */
 	USDUTILITIES_API bool NotifyIfInstanceProxy(const pxr::UsdPrim& Prim);
+
+	/** Returns the attribute that has the strongest spec, acrosss the provided Attributes array */
+	USDUTILITIES_API UE::FUsdAttribute GetStrongestAttribute(const TArray<UE::FUsdAttribute>& Attributes, double TimeCode);
 #endif	  // USE_USD_SDK
 }
diff --git a/Engine/Plugins/Runtime/USDCore/Source/USDUtilities/Public/USDConversionUtils.h b/Engine/Plugins/Runtime/USDCore/Source/USDUtilities/Public/USDConversionUtils.h
--- a/Engine/Plugins/Runtime/USDCore/Source/USDUtilities/Public/USDConversionUtils.h
+++ b/Engine/Plugins/Runtime/USDCore/Source/USDUtilities/Public/USDConversionUtils.h
@@ -312,6 +312,9 @@ namespace UsdUtils
 		bool bInIgnoreVisibility = false
 	);
 
+	/** Returns all of the Prim attributes that contain timeSamples */
+	USDUTILITIES_API TArray<UE::FUsdAttribute> GetAnimatedAttributes(const pxr::UsdPrim& Prim);
+
 	/** Returns true if this prim has any authored opinion for prim kind */
 	USDUTILITIES_API bool HasAuthoredKind(const pxr::UsdPrim& Prim);
 
diff --git a/Engine/Plugins/Runtime/USDCore/Source/USDUtilities/Public/USDGeomMeshConversion.h b/Engine/Plugins/Runtime/USDCore/Source/USDUtilities/Public/USDGeomMeshConversion.h
--- a/Engine/Plugins/Runtime/USDCore/Source/USDUtilities/Public/USDGeomMeshConversion.h
+++ b/Engine/Plugins/Runtime/USDCore/Source/USDUtilities/Public/USDGeomMeshConversion.h
@@ -526,6 +526,9 @@ namespace UsdUtils
 
 	USDUTILITIES_API void GetAnimatedMeshTimeCodes(const UE::FUsdStage& Stage, const FString& PrimPath, int32& OutStartFrame, int32& OutEndFrame);
 
+	/** Returns all of the Prim UsdGeomMesh-related attributes (points, faceVertexIndices, etc.) that contain timeSamples */
+	USDUTILITIES_API TArray<UE::FUsdAttribute> GetAnimatedMeshAttributes(const pxr::UsdPrim& Prim);
+
 	/** Mesh topology variance as defined by Alembic's MeshTopologyVariance */
 	enum class EMeshTopologyVariance : uint8
 	{
diff --git a/Engine/Plugins/Runtime/USDCore/Source/USDUtilities/Public/USDLayerUtils.h b/Engine/Plugins/Runtime/USDCore/Source/USDUtilities/Public/USDLayerUtils.h
--- a/Engine/Plugins/Runtime/USDCore/Source/USDUtilities/Public/USDLayerUtils.h
+++ b/Engine/Plugins/Runtime/USDCore/Source/USDUtilities/Public/USDLayerUtils.h
@@ -125,6 +125,7 @@ namespace UsdUtils
 	USDUTILITIES_API UE::FSdfLayer FindLayerForAttribute(const pxr::UsdAttribute& Attribute, double TimeCode);
 
 	/** Finds the strongest layer that contributes to a set of attributes. Assumes these attributes are on the same stage */
+	UE_DEPRECATED(5.7, "Use UsdUtils::GetStrongestAttribute() and then UsdUtils::FindLayerForAttribute()")
 	USDUTILITIES_API UE::FSdfLayer FindLayerForAttributes(
 		const TArray<UE::FUsdAttribute>& Attributes,
 		double TimeCode,
diff --git a/Engine/Plugins/Runtime/USDCore/Source/USDUtilities/Public/USDSkeletalDataConversion.h b/Engine/Plugins/Runtime/USDCore/Source/USDUtilities/Public/USDSkeletalDataConversion.h
--- a/Engine/Plugins/Runtime/USDCore/Source/USDUtilities/Public/USDSkeletalDataConversion.h
+++ b/Engine/Plugins/Runtime/USDCore/Source/USDUtilities/Public/USDSkeletalDataConversion.h
@@ -327,17 +327,19 @@ namespace UsdToUnreal
 	 * Warning: UAnimSequence must be previously set with a USkeleton generated from the skeletal data of the same UsdSkelSkeletonQuery.
 	 * @param InUsdSkeletonQuery - SkinningQuery with the data to convert
 	 * @param InSkinningTargets - Skinned meshes that use the skeleton of InUsdSkeletonQuery. Required to fetch the blend shape ordering of each mesh.
-	 *Optional (can be nullptr to ignore)
+	 * Optional (can be nullptr to ignore)
 	 * @param InBlendShapes - Converted blend shape data that will be used to interpret blend shape weights as morph target weight float curves.
-	 *Optional (can be nullptr to ignore)
+	 * Optional (can be nullptr to ignore)
 	 * @param InInterpretLODs - Whether we try parsing animation data from all LODs of skinning meshes that are inside LOD variant sets
 	 * @param RootMotionPrim - Optional prim whose transform animation will be concatenaded into the root bone's regular joint animation
 	 * @param OutSkeletalAnimationAsset - Output parameter that will be filled with the converted data
 	 * @param OutStartOffsetSeconds - Optional output parameter that will be filled with the offset in seconds of when this UAnimSequence asset should
-	 *be played since the start of its layer to match the intended composed animation. The baked UAnimSequence will only contain the range between the
-	 *first and last joint and/or blend shape timeSamples, and so the offset is needed to properly position the animation on the stage's timeline.
-	 *This is in seconds because the main use case is to use this offset when animating USkeletalMeshComponents, and those drive their UAnimSequences
-	 *with seconds.
+	 * be played since the start of the stage to match the intended composed animation. The baked UAnimSequence will only contain the range between the
+	 * first and last joint and/or blend shape timeSamples, and so the offset is needed to properly position the animation on the stage's timeline.
+	 * This is in seconds because the main use case is to use this offset when animating USkeletalMeshComponents, and those drive their UAnimSequences
+	 * with seconds.
+	 * @param OutScaleSeconds - Optional output parameter that will be filled with the total scale difference between the playback speed of the 
+	 * baked AnimSequence and the stage
 	 * @return Whether the conversion was successful or not.
 	 */
 	USDUTILITIES_API bool ConvertSkelAnim(
@@ -347,7 +349,8 @@ namespace UsdToUnreal
 		bool bInInterpretLODs,
 		const pxr::UsdPrim& RootMotionPrim,
 		UAnimSequence* OutSkeletalAnimationAsset,
-		float* OutStartOffsetSeconds = nullptr
+		float* OutStartOffsetSeconds = nullptr,
+		float* OutScaleSeconds = nullptr
 	);
 
 	/**
