USD: Speculative fix to prevent a race condition when async loading.

diff --git a/Engine/Plugins/Importers/USDImporter/Source/USDStage/Private/USDStageActor.cpp b/Engine/Plugins/Importers/USDImporter/Source/USDStage/Private/USDStageActor.cpp
--- a/Engine/Plugins/Importers/USDImporter/Source/USDStage/Private/USDStageActor.cpp
+++ b/Engine/Plugins/Importers/USDImporter/Source/USDStage/Private/USDStageActor.cpp
@@ -1049,17 +1049,7 @@ AUsdStageActor::AUsdStageActor()
 			);
 		}
 
-		FCoreUObjectDelegates::OnObjectPropertyChanged.AddUObject(this, &AUsdStageActor::OnObjectPropertyChanged);
-
-		// Also prevent standalone from doing this
-		if (GIsEditor && GEditor)
-		{
-			if (UBlueprintGeneratedClass* BPClass = Cast<UBlueprintGeneratedClass>(GetClass()))
-			{
-				FRecompilationTracker::SetupEvents();
-				FCoreUObjectDelegates::OnObjectsReplaced.AddUObject(this, &AUsdStageActor::OnObjectsReplaced);
-			}
-		}
+		// Note that we subscribe to some additional events from PostInitProperties()
 
 		LevelSequenceHelper.GetOnSkelAnimationBaked().AddUObject(this, &AUsdStageActor::OnSkelAnimationBaked);
 
@@ -4131,6 +4121,26 @@ void AUsdStageActor::BeginDestroy()
 void AUsdStageActor::PostInitProperties()
 {
 	Super::PostInitProperties();
+
+#if WITH_EDITOR
+	// Don't register FCoreUObjectDelegates from the constructor as that can lead to a race condition when async loading (see UE-294491)
+	if (!IsTemplate())
+	{
+		FCoreUObjectDelegates::OnObjectPropertyChanged.RemoveAll(this);
+		FCoreUObjectDelegates::OnObjectPropertyChanged.AddUObject(this, &AUsdStageActor::OnObjectPropertyChanged);
+
+		// Also prevent standalone from doing this
+		if (GIsEditor && GEditor)
+		{
+			if (UBlueprintGeneratedClass* BPClass = Cast<UBlueprintGeneratedClass>(GetClass()))
+			{
+				FRecompilationTracker::SetupEvents();
+				FCoreUObjectDelegates::OnObjectsReplaced.RemoveAll(this);
+				FCoreUObjectDelegates::OnObjectsReplaced.AddUObject(this, &AUsdStageActor::OnObjectsReplaced);
+			}
+		}
+	}
+#endif	  // WITH_EDITOR
 }
 
 void AUsdStageActor::PostRegisterAllComponents()
