[USD] Fix compile errors in the USDGeomMeshTranslator for non-editor configurations.

#rnx

diff --git a/Engine/Plugins/Importers/USDImporter/Source/USDSchemas/Private/USDGeomMeshTranslator.cpp b/Engine/Plugins/Importers/USDImporter/Source/USDSchemas/Private/USDGeomMeshTranslator.cpp
--- a/Engine/Plugins/Importers/USDImporter/Source/USDSchemas/Private/USDGeomMeshTranslator.cpp
+++ b/Engine/Plugins/Importers/USDImporter/Source/USDSchemas/Private/USDGeomMeshTranslator.cpp
@@ -703,8 +703,10 @@ namespace UsdGeomMeshTranslatorImpl
 		StaticMesh.CreateBodySetup();
 		StaticMesh.MarkAsNotHavingNavigationData();	   // Needed or else it will warn if we try cooking with body setup
 
+#if WITH_EDITOR
 		// If it's a Nanite Assembly, cache references
 		StaticMesh.RecacheNaniteAssemblyReferences();
+#endif
 	}
 
 	bool BuildStaticMesh(UStaticMesh& StaticMesh, const FStaticFeatureLevel& FeatureLevel, TArray<FMeshDescription>& LODIndexToMeshDescription)
@@ -787,7 +789,7 @@ namespace UsdGeomMeshTranslatorImpl
 #endif	  // WITH_EDITOR
 
 			// Include the Nanite mesh bounds as well. This is important for Nanite Assemblies
-			if (StaticMesh.IsNaniteEnabled())
+			if (RenderData.HasValidNaniteData())
 			{
 				const FBoxSphereBounds3f NaniteBounds3f = RenderData.NaniteResourcesPtr->MeshBounds;
 				FBoxSphereBounds NaniteBounds;
