USD: Fix wrong log message text placeholders.

diff --git a/Engine/Plugins/Importers/USDImporter/Source/USDSchemas/Private/MeshTranslationImpl.cpp b/Engine/Plugins/Importers/USDImporter/Source/USDSchemas/Private/MeshTranslationImpl.cpp
--- a/Engine/Plugins/Importers/USDImporter/Source/USDSchemas/Private/MeshTranslationImpl.cpp
+++ b/Engine/Plugins/Importers/USDImporter/Source/USDSchemas/Private/MeshTranslationImpl.cpp
@@ -729,7 +729,7 @@ TMap<const UsdUtils::FUsdPrimMaterialSlot*, UMaterialInterface*> MeshTranslation
 					else if (!Material->IsTwoSided() && Slot.bMeshIsDoubleSided)
 					{
 						USD_LOG_WARNING(
-							TEXT("Using one-sided UE material '{0}' for doubleSided prim '{1}'"),
+							TEXT("Using one-sided UE material '%s' for doubleSided prim '%s'"),
 							*Slot.MaterialSource,
 							*UsdToUnreal::ConvertPath(UsdPrim.GetPrimPath())
 						);
