Small refactor to expose GroomImportOptions helper function.

#rnx

diff --git a/Engine/Plugins/Importers/USDImporter/Source/USDSchemas/Private/USDGroomTranslator.cpp b/Engine/Plugins/Importers/USDImporter/Source/USDSchemas/Private/USDGroomTranslator.cpp
--- a/Engine/Plugins/Importers/USDImporter/Source/USDSchemas/Private/USDGroomTranslator.cpp
+++ b/Engine/Plugins/Importers/USDImporter/Source/USDSchemas/Private/USDGroomTranslator.cpp
@@ -42,40 +42,6 @@
 
 namespace UE::UsdGroomTranslator::Private
 {
-	UGroomImportOptions* CreateGroomImportOptions(
-		const FHairDescriptionGroups& GroupsDescription,
-		const TArray<FHairGroupsInterpolation>& BuildSettings
-	)
-	{
-		// Create a new groom import options and populate the interpolation settings based on the group count
-		UGroomImportOptions* ImportOptions = NewObject<UGroomImportOptions>();
-		const uint32 GroupCount = GroupsDescription.HairGroups.Num();
-		if (GroupCount != uint32(ImportOptions->InterpolationSettings.Num()))
-		{
-			ImportOptions->InterpolationSettings.Init(FHairGroupsInterpolation(), GroupCount);
-		}
-
-		TOptional<FHairGroupsInterpolation> LastBuildSettings;
-		if (BuildSettings.Num() > 0)
-		{
-			LastBuildSettings = BuildSettings.Last();
-		}
-
-		// If there are less build settings than groups, use the last one that was specified by the user
-		for (uint32 Index = 0; Index < GroupCount; ++Index)
-		{
-			if (BuildSettings.IsValidIndex(Index))
-			{
-				ImportOptions->InterpolationSettings[Index] = BuildSettings[Index];
-			}
-			else if (LastBuildSettings.IsSet())
-			{
-				ImportOptions->InterpolationSettings[Index] = LastBuildSettings.GetValue();
-			}
-		}
-		return ImportOptions;
-	}
-
 	FSHAHash ComputeHairDescriptionHash(FHairDescription& HairDescription, TArray<FHairGroupsInterpolation>& BuildSettings)
 	{
 		// The computed hash takes the hair description and group build settings into account since the groom builder
@@ -189,9 +155,7 @@ protected:
 				FHairDescriptionGroups GroupsDescription;
 				FGroomBuilder::BuildHairDescriptionGroups(HairDescription, GroupsDescription);
 
-				ImportOptions.Reset(	//
-					UE::UsdGroomTranslator::Private::CreateGroomImportOptions(GroupsDescription, Context->GroomInterpolationSettings)
-				);
+				ImportOptions.Reset(CreateGroomImportOptions(GroupsDescription, Context->GroomInterpolationSettings));
 
 				FSHAHash SHAHash = UE::UsdGroomTranslator::Private::ComputeHairDescriptionHash(HairDescription, ImportOptions->InterpolationSettings);
 
diff --git a/Engine/Plugins/Runtime/HairStrands/Source/HairStrandsCore/Private/GroomImportOptions.cpp b/Engine/Plugins/Runtime/HairStrands/Source/HairStrandsCore/Private/GroomImportOptions.cpp
--- a/Engine/Plugins/Runtime/HairStrands/Source/HairStrandsCore/Private/GroomImportOptions.cpp
+++ b/Engine/Plugins/Runtime/HairStrands/Source/HairStrandsCore/Private/GroomImportOptions.cpp
@@ -104,4 +104,35 @@ void UGroomHairGroupsMapping::SetIndex(int32 InNewIndex, int32 InOldIndex)
 	{
 		OldToNewGroupIndexMapping[InOldIndex] = NewGroupNames.IsValidIndex(InNewIndex) ? InNewIndex : -1;
 	}
-}
\ No newline at end of file
+}
+
+UGroomImportOptions* CreateGroomImportOptions(const FHairDescriptionGroups& GroupsDescription, const TArray<FHairGroupsInterpolation>& BuildSettings)
+{
+	// Create a new groom import options and populate the interpolation settings based on the group count
+	UGroomImportOptions* ImportOptions = NewObject<UGroomImportOptions>();
+	const uint32 GroupCount = GroupsDescription.HairGroups.Num();
+	if (GroupCount != uint32(ImportOptions->InterpolationSettings.Num()))
+	{
+		ImportOptions->InterpolationSettings.Init(FHairGroupsInterpolation(), GroupCount);
+	}
+
+	TOptional<FHairGroupsInterpolation> LastBuildSettings;
+	if (BuildSettings.Num() > 0)
+	{
+		LastBuildSettings = BuildSettings.Last();
+	}
+
+	// If there are less build settings than groups, use the last one that was specified by the user
+	for (uint32 Index = 0; Index < GroupCount; ++Index)
+	{
+		if (BuildSettings.IsValidIndex(Index))
+		{
+			ImportOptions->InterpolationSettings[Index] = BuildSettings[Index];
+		}
+		else if (LastBuildSettings.IsSet())
+		{
+			ImportOptions->InterpolationSettings[Index] = LastBuildSettings.GetValue();
+		}
+	}
+	return ImportOptions;
+}
diff --git a/Engine/Plugins/Runtime/HairStrands/Source/HairStrandsCore/Public/GroomImportOptions.h b/Engine/Plugins/Runtime/HairStrands/Source/HairStrandsCore/Public/GroomImportOptions.h
--- a/Engine/Plugins/Runtime/HairStrands/Source/HairStrandsCore/Public/GroomImportOptions.h
+++ b/Engine/Plugins/Runtime/HairStrands/Source/HairStrandsCore/Public/GroomImportOptions.h
@@ -105,4 +105,6 @@ public:
 	static UE_API void RemapHairDescriptionGroups(const FHairDescriptionGroups& SrcGroups, const FHairDescriptionGroups& DstGroups, TArray<int32>& Out);
 };
 
+UE_API UGroomImportOptions* CreateGroomImportOptions(const FHairDescriptionGroups& GroupsDescription, const TArray<FHairGroupsInterpolation>& BuildSettings);
+
 #undef UE_API
diff --git a/Engine/Plugins/Runtime/HairStrands/Source/HairStrandsEditor/Private/HairStrandsImporter.cpp b/Engine/Plugins/Runtime/HairStrands/Source/HairStrandsEditor/Private/HairStrandsImporter.cpp
--- a/Engine/Plugins/Runtime/HairStrands/Source/HairStrandsEditor/Private/HairStrandsImporter.cpp
+++ b/Engine/Plugins/Runtime/HairStrands/Source/HairStrandsEditor/Private/HairStrandsImporter.cpp
@@ -170,7 +170,7 @@ UGroomAsset* FHairStrandsImporter::ImportHair(const FHairImportContext& ImportCo
 	}
 
 	const bool bUseGroupRemapping = GHairStrandsReimportGroupRemapping > 0;
-	if (OldExistingHair && bUseGroupRemapping)
+	if (OldExistingHair && OldExistingHair->CanRebuildFromDescription() && bUseGroupRemapping) // make sure OldExistingHair is not a new empty groom asset
 	{
 		FHairDescriptionGroups NewHairDescriptionGroups;
 		FGroomBuilder::BuildHairDescriptionGroups(NewHairDescription, NewHairDescriptionGroups);
