USD: Restore the legacy option to export to nested folders.

diff --git a/Engine/Plugins/Importers/USDImporter/Source/USDExporter/Private/LevelExporterUSDOptions.cpp b/Engine/Plugins/Importers/USDImporter/Source/USDExporter/Private/LevelExporterUSDOptions.cpp
--- a/Engine/Plugins/Importers/USDImporter/Source/USDExporter/Private/LevelExporterUSDOptions.cpp
+++ b/Engine/Plugins/Importers/USDImporter/Source/USDExporter/Private/LevelExporterUSDOptions.cpp
@@ -23,6 +23,7 @@ void UsdUtils::AddAnalyticsAttributes(const FLevelExporterUSDOptionsInner& Optio
 	InOutAttributes.Emplace(TEXT("ExportActorFolders"), Options.bExportActorFolders);
 	InOutAttributes.Emplace(TEXT("IgnoreSequencerAnimations"), Options.bIgnoreSequencerAnimations);
 	InOutAttributes.Emplace(TEXT("ExportFoliageOnActorsLayer"), Options.bExportFoliageOnActorsLayer);
+	InOutAttributes.Emplace(TEXT("bExportToNestedSubFolders"), Options.bExportToNestedSubFolders);
 	UsdUtils::AddAnalyticsAttributes(Options.AssetOptions, InOutAttributes);
 	UsdUtils::AddAnalyticsAttributes(Options.MetadataOptions, InOutAttributes);
 	InOutAttributes.Emplace(TEXT("LowestLandscapeLOD"), LexToString(Options.LowestLandscapeLOD));
@@ -57,6 +58,8 @@ void UsdUtils::HashForLevelExport(const FLevelExporterUSDOptionsInner& Options,
 	// If we changed where we want the assets exported we need to re-export them and update the reference paths
 	HashToUpdate.UpdateWithString(*Options.AssetFolder.Path, Options.AssetFolder.Path.Len());
 
+	HashToUpdate.Update(reinterpret_cast<const uint8*>(&Options.bExportToNestedSubFolders), sizeof(Options.bExportToNestedSubFolders));
+
 	// This affects how we author material overrides
 	HashToUpdate.Update(reinterpret_cast<const uint8*>(&Options.AssetOptions.bBakeMaterials), sizeof(Options.AssetOptions.bBakeMaterials));
 
diff --git a/Engine/Plugins/Importers/USDImporter/Source/USDExporter/Public/LevelExporterUSDOptions.h b/Engine/Plugins/Importers/USDImporter/Source/USDExporter/Public/LevelExporterUSDOptions.h
--- a/Engine/Plugins/Importers/USDImporter/Source/USDExporter/Public/LevelExporterUSDOptions.h
+++ b/Engine/Plugins/Importers/USDImporter/Source/USDExporter/Public/LevelExporterUSDOptions.h
@@ -49,6 +49,13 @@ struct FLevelExporterUSDOptionsInner
 	UPROPERTY(EditAnywhere, config, BlueprintReadWrite, Category = "Asset options")
 	FDirectoryPath AssetFolder;
 
+	/**
+	 * Export the asset files in a folder structure that mirrors their Unreal project content folder structure.
+	 * Warning: This can produce very long file paths that could cause issues depending on the operating system settings.
+	 */
+	UPROPERTY(EditAnywhere, config, BlueprintReadWrite, Category = "Asset options", AdvancedDisplay)
+	bool bExportToNestedSubFolders = false;
+
 	/** Options to use for all exported assets when appropriate (e.g. static and skeletal meshes, materials, etc.) */
 	UPROPERTY(EditAnywhere, config, BlueprintReadWrite, Category = "Asset options", meta = (ShowOnlyInnerProperties))
 	FUsdMeshAssetOptions AssetOptions;
