Interchange / USD: Fix automated reimports still showing the options dialog.

diff --git a/Engine/Plugins/Importers/USDImporter/Source/USDStageImporter/Private/USDStageAssetImportFactory.cpp b/Engine/Plugins/Importers/USDImporter/Source/USDStageImporter/Private/USDStageAssetImportFactory.cpp
--- a/Engine/Plugins/Importers/USDImporter/Source/USDStageImporter/Private/USDStageAssetImportFactory.cpp
+++ b/Engine/Plugins/Importers/USDImporter/Source/USDStageImporter/Private/USDStageAssetImportFactory.cpp
@@ -255,15 +255,17 @@ EReimportResult::Type UUsdStageAssetImportFactory::Reimport(UObject* Obj)
 
 	ImportContext.bReadFromStageCache = false;
 
+	const bool bIsAutomated = IsAutomatedImport() || IsAutomatedReimport();
+
 	// We shouldn't be able to import actors when doing a manual asset reimport (from the content browser)
 	TOptional<UsdUtils::FScopedSuppressActorImport> SuppressActorImport;
-	if (!IsAutomatedImport())
+	if (!bIsAutomated)
 	{
 		SuppressActorImport.Emplace(ImportContext.ImportOptions);
 	}
 
 	const bool bIsReimport = true;
-	if (!ImportContext.Init(Obj->GetName(), ReimportFilePath, Obj->GetName(), Obj->GetFlags(), IsAutomatedImport(), bIsReimport))
+	if (!ImportContext.Init(Obj->GetName(), ReimportFilePath, Obj->GetName(), Obj->GetFlags(), bIsAutomated, bIsReimport))
 	{
 		USD_LOG_USERERROR(FText::Format(
 			LOCTEXT("ReimportErrorNoContext", "Failed to initialize re-import context for asset '{0}'!"),
diff --git a/Engine/Source/Editor/UnrealEd/Private/Editor.cpp b/Engine/Source/Editor/UnrealEd/Private/Editor.cpp
--- a/Engine/Source/Editor/UnrealEd/Private/Editor.cpp
+++ b/Engine/Source/Editor/UnrealEd/Private/Editor.cpp
@@ -564,7 +564,7 @@ UE::Interchange::FAssetImportResultRef FReimportManager::ReimportAsync(UObject*
 						if (InterchangeManager.CanTranslateSourceData(ScopedSourceData.GetSourceData()))
 						{
 							FImportAssetParameters ImportAssetParameters;
-							ImportAssetParameters.bIsAutomated = GIsAutomationTesting || FApp::IsUnattended() || IsRunningCommandlet() || GIsRunningUnattendedScript;
+							ImportAssetParameters.bIsAutomated = bAutomated || GIsAutomationTesting || FApp::IsUnattended() || IsRunningCommandlet() || GIsRunningUnattendedScript;
 							ImportAssetParameters.ReimportAsset = Obj;
 							ImportAssetParameters.ReimportSourceIndex = SourceFileIndex;
 							ImportAssetParameters.bForceShowDialog = bInForceShowDialog;
