USD: Use the default collision profile for generated static meshes.

diff --git a/Engine/Plugins/Importers/USDImporter/Source/USDSchemas/Private/USDGeomMeshTranslator.cpp b/Engine/Plugins/Importers/USDImporter/Source/USDSchemas/Private/USDGeomMeshTranslator.cpp
--- a/Engine/Plugins/Importers/USDImporter/Source/USDSchemas/Private/USDGeomMeshTranslator.cpp
+++ b/Engine/Plugins/Importers/USDImporter/Source/USDSchemas/Private/USDGeomMeshTranslator.cpp
@@ -1414,11 +1414,12 @@ namespace UE::UsdCollision::Private
 			return;
 		}
 
+		// Note we don't set a CollisionProfileName here to let the default value of the StaticMesh be used
+		// (The CreateBodySetup() inside PreBuildStaticMesh() should set it to BlockAll)
 		BodySetup->Modify();
 		BodySetup->RemoveSimpleCollision();
 		BodySetup->bNeverNeedsCookedCollisionData = false;
 		BodySetup->DefaultInstance.SetCollisionEnabled(ECollisionEnabled::QueryAndPhysics);
-		BodySetup->DefaultInstance.SetCollisionProfileName(UCollisionProfile::BlockAllDynamic_ProfileName);
 		BodySetup->CollisionTraceFlag = CTF_UseDefault;
 
 		switch (Approximation)
