USD: Don't collapse prims with active alternate draw modes.

diff --git a/Engine/Plugins/Importers/USDImporter/Source/USDSchemas/Private/USDGeomXformableTranslator.cpp b/Engine/Plugins/Importers/USDImporter/Source/USDSchemas/Private/USDGeomXformableTranslator.cpp
--- a/Engine/Plugins/Importers/USDImporter/Source/USDSchemas/Private/USDGeomXformableTranslator.cpp
+++ b/Engine/Plugins/Importers/USDImporter/Source/USDSchemas/Private/USDGeomXformableTranslator.cpp
@@ -1152,8 +1152,12 @@ bool FUsdGeomXformableTranslator::CanBeCollapsed(ECollapsingType CollapsingType)
 		return false;
 	}
 
-	if (UsdUtils::IsAnimated(UsdPrim) || UsdUtils::PrimHasSchema(UsdPrim, UnrealIdentifiers::LiveLinkAPI) || UsdPrim.IsA<pxr::UsdSkelRoot>()
-		|| (Context->bAllowInterpretingLODs && UsdUtils::DoesPrimContainMeshLODs(UsdPrim)))
+	if (UsdUtils::IsAnimated(UsdPrim) || 
+		UsdUtils::PrimHasSchema(UsdPrim, UnrealIdentifiers::LiveLinkAPI) || 
+		UsdPrim.IsA<pxr::UsdSkelRoot>() || 
+		(Context->bAllowInterpretingLODs && UsdUtils::DoesPrimContainMeshLODs(UsdPrim)) || 
+		UsdUtils::GetAppliedDrawMode(UsdPrim) != EUsdDrawMode::Default
+	)
 	{
 		return false;
 	}
