USD: Fix USD Stage Editor prim copy-paste behaving strangely when selecting multiple prims at once

diff --git a/Engine/Plugins/Importers/USDImporter/Source/USDStageEditor/Private/SUSDStageTreeView.cpp b/Engine/Plugins/Importers/USDImporter/Source/USDStageEditor/Private/SUSDStageTreeView.cpp
--- a/Engine/Plugins/Importers/USDImporter/Source/USDStageEditor/Private/SUSDStageTreeView.cpp
+++ b/Engine/Plugins/Importers/USDImporter/Source/USDStageEditor/Private/SUSDStageTreeView.cpp
@@ -1034,31 +1034,28 @@ void SUsdStageTreeView::OnPastePrim()
 
 	UE::FSdfChangeBlock Block;
 
-	TArray<UE::FUsdPrim> ParentPrims;
+	UE::FUsdPrim ParentPrim;
 
 	// This happens when right-clicking the background area without selecting any prim
 	if (MySelectedItems.IsEmpty())
 	{
-		ParentPrims.Add(UsdStage.GetPseudoRoot());
+		ParentPrim = UsdStage.GetPseudoRoot();
 	}
-	else
+	else if (MySelectedItems.Num() > 0)
 	{
-		ParentPrims.Reserve(MySelectedItems.Num());
-
-		// A bit unusual that we can paste to multiple locations at the same time, but why not?
-		for (FUsdPrimViewModelRef SelectedItem : MySelectedItems)
-		{
-			ParentPrims.Add(SelectedItem->UsdPrim);
-		}
+		// Only paste for the first selection, for simplicity
+		ParentPrim = MySelectedItems[0]->UsdPrim.GetParent();
 	}
 
-	for (const UE::FUsdPrim& ParentPrim : ParentPrims)
+	if (!ParentPrim)
 	{
-		// Preemptively mark the parent prims as expanded so that we can always see what we pasted
-		ExpandedPrimPaths.Add(ParentPrim.GetPrimPath().GetString());
-
-		UsdUtils::PastePrims(ParentPrim);
+		return;
 	}
+
+	// Preemptively mark the parent prims as expanded so that we can always see what we pasted
+	ExpandedPrimPaths.Add(ParentPrim.GetPrimPath().GetString());
+
+	UsdUtils::PastePrims(ParentPrim);
 }
 
 void SUsdStageTreeView::OnDuplicatePrim(EUsdDuplicateType DuplicateType)
