USD: Allow authoring references and payloads on prims that don't yet have a spec on the current edit target, as the operation will author one anyway

diff --git a/Engine/Plugins/Importers/USDImporter/Source/USDStageEditor/Private/SUSDStageTreeView.cpp b/Engine/Plugins/Importers/USDImporter/Source/USDStageEditor/Private/SUSDStageTreeView.cpp
--- a/Engine/Plugins/Importers/USDImporter/Source/USDStageEditor/Private/SUSDStageTreeView.cpp
+++ b/Engine/Plugins/Importers/USDImporter/Source/USDStageEditor/Private/SUSDStageTreeView.cpp
@@ -849,7 +849,7 @@ TSharedPtr<SWidget> SUsdStageTreeView::ConstructPrimContextMenu()
 				FCanExecuteAction::CreateLambda(
 					[this]()
 					{
-						return SUsdStageTreeView::DoesPrimExistOnEditTarget() && GetSelectedItems().Num() == 1;
+						return GetSelectedItems().Num() == 1;
 					}
 				)
 			),
@@ -885,7 +885,7 @@ TSharedPtr<SWidget> SUsdStageTreeView::ConstructPrimContextMenu()
 				FCanExecuteAction::CreateLambda(
 					[this]()
 					{
-						return SUsdStageTreeView::DoesPrimExistOnEditTarget() && GetSelectedItems().Num() == 1;
+						return GetSelectedItems().Num() == 1;
 					}
 				)
 			),
