USD: Enforce subdivision level limits on the legacy USD import dialog

diff --git a/Engine/Plugins/Importers/USDImporter/Source/USDStageImporter/Private/USDStageImportContext.cpp b/Engine/Plugins/Importers/USDImporter/Source/USDStageImporter/Private/USDStageImportContext.cpp
--- a/Engine/Plugins/Importers/USDImporter/Source/USDStageImporter/Private/USDStageImportContext.cpp
+++ b/Engine/Plugins/Importers/USDImporter/Source/USDStageImporter/Private/USDStageImportContext.cpp
@@ -69,6 +69,19 @@ bool FUsdStageImportContext::Init(
 			PackagePath = FString::Printf(TEXT("%s/%s/"), *PickContentPathDlg->GetPath().ToString(), *InName);
 		}
 
+		// Update the boundaries of the subdivision level slider according to the cvar
+#if WITH_EDITOR
+		static IConsoleVariable* MaxSubdivCvar = IConsoleManager::Get().FindConsoleVariable(TEXT("USD.Subdiv.MaxSubdivLevel"));
+		const int32 MaxSubdivLevel = FMath::Max(MaxSubdivCvar ? MaxSubdivCvar->GetInt() : 6, 1);
+		const FString MaxSubdivLevelString = LexToString(MaxSubdivLevel);
+
+		if (FProperty* Property = UUsdStageImportOptions::StaticClass()->FindPropertyByName(GET_MEMBER_NAME_CHECKED(UUsdStageImportOptions, SubdivisionLevel)))
+		{
+			Property->SetMetaData(TEXT("UIMax"), *MaxSubdivLevelString);
+			Property->SetMetaData(TEXT("ClampMax"), *MaxSubdivLevelString);
+		}
+#endif	  // WITH_EDITOR
+
 		// Show dialog for import options
 		bool bProceedWithImport = SUsdOptionsWindow::ShowImportOptions(*ImportOptions, &Stage);
 		if (!bProceedWithImport)
diff --git a/Engine/Plugins/Importers/USDImporter/Source/USDStageImporter/Public/USDStageImportOptions.h b/Engine/Plugins/Importers/USDImporter/Source/USDStageImporter/Public/USDStageImportOptions.h
--- a/Engine/Plugins/Importers/USDImporter/Source/USDStageImporter/Public/USDStageImportOptions.h
+++ b/Engine/Plugins/Importers/USDImporter/Source/USDStageImporter/Public/USDStageImportOptions.h
@@ -144,7 +144,7 @@ public:
 	EUsdCollisionType FallbackCollisionType;
 
 	/** Subdivision level to use for all subdivision meshes on the opened stage. 0 means "don't subdivide" */
-	UPROPERTY(BlueprintReadWrite, config, EditAnywhere, Category = "USD options")
+	UPROPERTY(BlueprintReadWrite, config, EditAnywhere, Category = "USD options", meta = (ClampMin = "0", ClampMax = "6", UIMin = "0", UIMax = "6"))
 	int32 SubdivisionLevel;
 
 	/* Describes if/how we should collect metadata from USD prims onto the assets and components we generate when importing */
