USD: Fix USD Stage Editor claiming it could remove ancestral reference arcs

diff --git a/Engine/Plugins/Importers/USDImporter/Source/USDStageEditorViewModels/Private/USDPrimViewModel.cpp b/Engine/Plugins/Importers/USDImporter/Source/USDStageEditorViewModels/Private/USDPrimViewModel.cpp
--- a/Engine/Plugins/Importers/USDImporter/Source/USDStageEditorViewModels/Private/USDPrimViewModel.cpp
+++ b/Engine/Plugins/Importers/USDImporter/Source/USDStageEditorViewModels/Private/USDPrimViewModel.cpp
@@ -17,6 +17,7 @@
 #include "pxr/usd/sdf/path.h"
 #include "pxr/usd/usd/payloads.h"
 #include "pxr/usd/usd/prim.h"
+#include "pxr/usd/usd/primCompositionQuery.h"
 #include "pxr/usd/usdGeom/xform.h"
 #include "pxr/usd/usdSkel/root.h"
 #include "pxr/usd/usdSkel/skeleton.h"
@@ -352,7 +353,24 @@ bool FUsdPrimViewModel::HasReferencesOnLocalLayer() const
 					pxr::SdfLayerHandle SpecLayer = Spec->GetLayer();
 					if (PrimUsdStage->GetEditTarget().GetLayer() == SpecLayer && PrimUsdStage->HasLocalLayer(SpecLayer))
 					{
-						return true;
+						// At this point we know this prim has a spec with references on this edit target layer
+						//
+						// For us to be able to interact with it however, it must be a non-ancestral reference, so we need to check if we
+						// have any on this layer and prim path.
+						//
+						// Ancestral references are reference composition arcs introduced by an ancestor prim (so e.g. a parent prim has a 
+						// reference, that contributes reference specs to this prim). We can't interact with those, as they're "encapsulated" 
+						// within the reference (see https://openusd.org/release/usdfaq.html#list-edited-composition-arcs)
+						pxr::UsdPrimCompositionQuery PrimCompositionQuery = pxr::UsdPrimCompositionQuery::GetDirectReferences(PxrUsdPrim);
+						for (const pxr::UsdPrimCompositionQueryArc& CompositionArc : PrimCompositionQuery.GetCompositionArcs())
+						{
+							if (CompositionArc.HasSpecs() && 
+								CompositionArc.GetIntroducingLayer() == SpecLayer && 
+								CompositionArc.GetIntroducingPrimPath() == PxrUsdPrim.GetPrimPath())
+							{
+								return true;
+							}
+						}
 					}
 				}
 			}
@@ -379,7 +397,23 @@ bool FUsdPrimViewModel::HasPayloadsOnLocalLayer() const
 					pxr::SdfLayerHandle SpecLayer = Spec->GetLayer();
 					if (PrimUsdStage->GetEditTarget().GetLayer() == SpecLayer && PrimUsdStage->HasLocalLayer(SpecLayer))
 					{
-						return true;
+						// See analogous comment on HasReferencesOnLocalLayer
+
+						// There's no pxr::UsdPrimCompositionQuery::GetDirectPayloads()...
+						pxr::UsdPrimCompositionQuery::Filter Filter;
+						Filter.dependencyTypeFilter = pxr::UsdPrimCompositionQuery::DependencyTypeFilter::Direct;
+						Filter.arcTypeFilter = pxr::UsdPrimCompositionQuery::ArcTypeFilter::Payload;
+						pxr::UsdPrimCompositionQuery PrimCompositionQuery{PxrUsdPrim, Filter};
+
+						for (const pxr::UsdPrimCompositionQueryArc& CompositionArc : PrimCompositionQuery.GetCompositionArcs())
+						{
+							if (CompositionArc.HasSpecs() && 
+								CompositionArc.GetIntroducingLayer() == SpecLayer && 
+								CompositionArc.GetIntroducingPrimPath() == PxrUsdPrim.GetPrimPath())
+							{
+								return true;
+							}
+						}
 					}
 				}
 			}
