USD: Adapt tests and code to new convention in UMovieSceneAudioSection::SetRange() where non-looping sections can't be longer than the audio duration

#virtualized

diff --git a/Engine/Plugins/Importers/USDImporter/Source/USDStage/Private/USDLevelSequenceHelper.cpp b/Engine/Plugins/Importers/USDImporter/Source/USDStage/Private/USDLevelSequenceHelper.cpp
--- a/Engine/Plugins/Importers/USDImporter/Source/USDStage/Private/USDLevelSequenceHelper.cpp
+++ b/Engine/Plugins/Importers/USDImporter/Source/USDStage/Private/USDLevelSequenceHelper.cpp
@@ -3349,7 +3349,9 @@ void FUsdLevelSequenceHelperImpl::AddAudioTracks(const UUsdPrimTwin& PrimTwin, c
 				// "Play the audio once beginning at startTime, continuing until endTime or until the
 				// audio completes, whichever comes first."
 				//
-				// Do nothing: Just continue using startTime and endTime as we have already
+				// Do nothing: Just continue using startTime and endTime as we have already. If we try
+				// setting a section range longer than the sound duration, the implementation of
+				// UMovieSceneAudioSection::SetRange will already clamp us to the sound duration anyway!
 			}
 			else if (PlaybackModeValue == pxr::UsdMediaTokens->loopFromStart)
 			{
