USD: Fix errors and warnings emitting multiple times when opening stages with USD errors

diff --git a/Engine/Plugins/Importers/USDImporter/Source/USDStage/Private/USDStageActor.cpp b/Engine/Plugins/Importers/USDImporter/Source/USDStage/Private/USDStageActor.cpp
--- a/Engine/Plugins/Importers/USDImporter/Source/USDStage/Private/USDStageActor.cpp
+++ b/Engine/Plugins/Importers/USDImporter/Source/USDStage/Private/USDStageActor.cpp
@@ -3088,6 +3088,17 @@ void AUsdStageActor::OpenUsdStage()
 		UsdUtils::GetUEPersistentStateSublayer(UsdStage, bCreateIfNeeded);
 #endif	  // #if USE_USD_SDK
 	}
+	else
+	{
+		// We failed to produce a valid UsdStage from this RootLayer.
+		// Reset the RootLayer back to empty so that subsequent calls to OpenUsdStage don't repeatedly try
+		// opening the same stage over and over
+		USD_LOG_USERERROR(FText::Format(
+			LOCTEXT("FailedToOpenStage", "Failed to open USD stage '{0}'"),
+			FText::FromString(*RootLayer.FilePath)
+		));
+		RootLayer = FFilePath{};
+	}
 
 	OnStageChanged.Broadcast();
 }
