USD: Fix incorrect LevelSequence cleanup when removing a prim from the LevelSequenceHelper

This was flagged by BugHawk when submitting 47032539

diff --git a/Engine/Plugins/Importers/USDImporter/Source/USDStage/Private/USDLevelSequenceHelper.cpp b/Engine/Plugins/Importers/USDImporter/Source/USDStage/Private/USDLevelSequenceHelper.cpp
--- a/Engine/Plugins/Importers/USDImporter/Source/USDStage/Private/USDLevelSequenceHelper.cpp
+++ b/Engine/Plugins/Importers/USDImporter/Source/USDStage/Private/USDLevelSequenceHelper.cpp
@@ -3918,8 +3918,12 @@ void FUsdLevelSequenceHelperImpl::RemoveSequenceForPrim(ULevelSequence& Sequence
 				RemoveSubSequenceSection(*ParentSequence, Sequence);
 			}
 
-			LevelSequencesByIdentifier.Remove(PrimTwin.PrimPath);
-			IdentifierByLevelSequence.Remove(&Sequence);
+			FString OldIdentifier;
+			if (IdentifierByLevelSequence.RemoveAndCopyValue(&Sequence, OldIdentifier) && !OldIdentifier.IsEmpty())
+			{
+				LevelSequencesByIdentifier.Remove(OldIdentifier);
+			}
+
 			DirtySequenceHierarchyCache();
 		}
 	}
