USD: Fix minor BugHawk findings on recent changelists

diff --git a/Engine/Plugins/Importers/USDImporter/Source/USDStage/Private/USDLevelSequenceHelper.cpp b/Engine/Plugins/Importers/USDImporter/Source/USDStage/Private/USDLevelSequenceHelper.cpp
--- a/Engine/Plugins/Importers/USDImporter/Source/USDStage/Private/USDLevelSequenceHelper.cpp
+++ b/Engine/Plugins/Importers/USDImporter/Source/USDStage/Private/USDLevelSequenceHelper.cpp
@@ -3924,6 +3924,8 @@ void FUsdLevelSequenceHelperImpl::RemoveSequenceForPrim(ULevelSequence& Sequence
 				LevelSequencesByIdentifier.Remove(OldIdentifier);
 			}
 
+			LayerIdentifierByLevelSequenceName.Remove(Sequence.GetName());
+
 			DirtySequenceHierarchyCache();
 		}
 	}
diff --git a/Engine/Plugins/Runtime/USDCore/Source/USDUtilities/Private/USDErrorUtils.cpp b/Engine/Plugins/Runtime/USDCore/Source/USDUtilities/Private/USDErrorUtils.cpp
--- a/Engine/Plugins/Runtime/USDCore/Source/USDUtilities/Private/USDErrorUtils.cpp
+++ b/Engine/Plugins/Runtime/USDCore/Source/USDUtilities/Private/USDErrorUtils.cpp
@@ -458,7 +458,7 @@ void FUsdLogManager::DisplayMessagesOnMessageLog(const TArray<TSharedRef<FTokeni
 		EMessageSeverity::Info,
 		bForce
 	);
-#endif	  // USE_USD_SDK
+#endif	  // WITH_EDITOR
 }
 
 FScopedUsdMessageLog::FScopedUsdMessageLog(FUsdLogManager::FMessageDisplayCallback InCallback)
