USD: Undo changelist 47032157, so that now we no longer expect non-looping audio sections to be clamped to the audio length.

Our cl 47032157 made the USD tests expect the sections to be clamped, according to the new Sequencer conventions following cl 46769967. That cl later got backed out with cl 47380523, so this cl is needed on our side to revert our expectations to what they were before.

#virtualized

diff --git a/Engine/Plugins/Importers/USDImporter/Source/USDStage/Private/USDLevelSequenceHelper.cpp b/Engine/Plugins/Importers/USDImporter/Source/USDStage/Private/USDLevelSequenceHelper.cpp
--- a/Engine/Plugins/Importers/USDImporter/Source/USDStage/Private/USDLevelSequenceHelper.cpp
+++ b/Engine/Plugins/Importers/USDImporter/Source/USDStage/Private/USDLevelSequenceHelper.cpp
@@ -3349,9 +3349,7 @@ void FUsdLevelSequenceHelperImpl::AddAudioTracks(const UUsdPrimTwin& PrimTwin, c
 				// "Play the audio once beginning at startTime, continuing until endTime or until the
 				// audio completes, whichever comes first."
 				//
-				// Do nothing: Just continue using startTime and endTime as we have already. If we try
-				// setting a section range longer than the sound duration, the implementation of
-				// UMovieSceneAudioSection::SetRange will already clamp us to the sound duration anyway!
+				// Do nothing: Just continue using startTime and endTime as we have already
 			}
 			else if (PlaybackModeValue == pxr::UsdMediaTokens->loopFromStart)
 			{
