USD: Prevent right-click actions directly on the Stage top-level tree view item

diff --git a/Engine/Plugins/Importers/USDImporter/Source/USDStageEditor/Private/SUSDStageTreeView.cpp b/Engine/Plugins/Importers/USDImporter/Source/USDStageEditor/Private/SUSDStageTreeView.cpp
--- a/Engine/Plugins/Importers/USDImporter/Source/USDStageEditor/Private/SUSDStageTreeView.cpp
+++ b/Engine/Plugins/Importers/USDImporter/Source/USDStageEditor/Private/SUSDStageTreeView.cpp
@@ -849,7 +849,7 @@ TSharedPtr<SWidget> SUsdStageTreeView::ConstructPrimContextMenu()
 				FCanExecuteAction::CreateLambda(
 					[this]()
 					{
-						return GetSelectedItems().Num() == 1;
+						return GetSelectedItems().Num() == 1 && GetSelectedItems()[0]->UsdPrim && !GetSelectedItems()[0]->UsdPrim.IsPseudoRoot();
 					}
 				)
 			),
@@ -885,7 +885,7 @@ TSharedPtr<SWidget> SUsdStageTreeView::ConstructPrimContextMenu()
 				FCanExecuteAction::CreateLambda(
 					[this]()
 					{
-						return GetSelectedItems().Num() == 1;
+						return GetSelectedItems().Num() == 1 && GetSelectedItems()[0]->UsdPrim && !GetSelectedItems()[0]->UsdPrim.IsPseudoRoot();
 					}
 				)
 			),
@@ -1439,6 +1439,16 @@ bool SUsdStageTreeView::CanPastePrim() const
 		return false;
 	}
 
+	TArray<FUsdPrimViewModelRef> MySelectedItems = GetSelectedItems();
+	for (FUsdPrimViewModelRef SelectedItem : MySelectedItems)
+	{
+		if (!SelectedItem->UsdPrim || SelectedItem->UsdPrim.IsPseudoRoot())
+		{
+			// Don't let us paste as siblings of the pseudoroot as that makes no sense
+			return false;
+		}
+	}
+
 	return UsdUtils::CanPastePrims();
 }
 
@@ -1473,6 +1483,11 @@ bool SUsdStageTreeView::DoesPrimExistOnEditTarget() const
 
 	for (FUsdPrimViewModelRef SelectedItem : MySelectedItems)
 	{
+		if (!SelectedItem->UsdPrim || SelectedItem->UsdPrim.IsPseudoRoot())
+		{
+			return false;
+		}
+
 		UE::FSdfPath SpecPath = UsdUtils::GetPrimSpecPathForLayer(SelectedItem->UsdPrim, UsdStage.GetEditTarget());
 		if (!SpecPath.IsAbsoluteRootPath() && !SpecPath.IsEmpty())
 		{
@@ -1487,6 +1502,11 @@ bool SUsdStageTreeView::DoesPrimHaveSpecOnLocalLayerStack() const
 {
 	for (FUsdPrimViewModelRef SelectedItem : GetSelectedItems())
 	{
+		if (!SelectedItem->UsdPrim || SelectedItem->UsdPrim.IsPseudoRoot())
+		{
+			return false;
+		}
+
 		if (SelectedItem->HasSpecsOnLocalLayer())
 		{
 			return true;
@@ -1500,6 +1520,11 @@ bool SUsdStageTreeView::DoesPrimHaveReferenceSpecOnLocalLayerStack() const
 {
 	for (FUsdPrimViewModelRef SelectedItem : GetSelectedItems())
 	{
+		if (!SelectedItem->UsdPrim || SelectedItem->UsdPrim.IsPseudoRoot())
+		{
+			return false;
+		}
+
 		if (SelectedItem->HasReferencesOnLocalLayer())
 		{
 			return true;
@@ -1513,6 +1538,11 @@ bool SUsdStageTreeView::DoesPrimHavePayloadSpecOnLocalLayerStack() const
 {
 	for (FUsdPrimViewModelRef SelectedItem : GetSelectedItems())
 	{
+		if (!SelectedItem->UsdPrim || SelectedItem->UsdPrim.IsPseudoRoot())
+		{
+			return false;
+		}
+
 		if (SelectedItem->HasPayloadsOnLocalLayer())
 		{
 			return true;
@@ -1533,6 +1563,11 @@ bool SUsdStageTreeView::DoSelectedPrimsHaveCollapsingPreference(UsdUtils::EColla
 
 	for (FUsdPrimViewModelRef SelectedItem : MySelectedItems)
 	{
+		if (!SelectedItem->UsdPrim || SelectedItem->UsdPrim.IsPseudoRoot())
+		{
+			return false;
+		}
+
 		UsdUtils::ECollapsingPreference PrimPreference = UsdUtils::GetCollapsingPreference(SelectedItem->UsdPrim);
 		if (PrimPreference != TargetPreference)
 		{
