USD: Fix some workflows still finding the old UObjects after an asset has been reimported.

#Synced-CL 49920387

#Synced-CL 49920387

diff --git a/Engine/Plugins/Importers/USDImporter/Source/USDStageImporter/Private/USDStageImporter.cpp b/Engine/Plugins/Importers/USDImporter/Source/USDStageImporter/Private/USDStageImporter.cpp
--- a/Engine/Plugins/Importers/USDImporter/Source/USDStageImporter/Private/USDStageImporter.cpp
+++ b/Engine/Plugins/Importers/USDImporter/Source/USDStageImporter/Private/USDStageImporter.cpp
@@ -68,6 +68,8 @@
 #include "SparseVolumeTexture/SparseVolumeTexture.h"
 #include "Subsystems/AssetEditorSubsystem.h"
 #include "Tracks/MovieSceneSubTrack.h"
+#include "UObject/Linker.h"
+#include "UObject/ReferencerFinder.h"
 #include "UObject/UObjectIterator.h"
 
 #define LOCTEXT_NAMESPACE "USDStageImporter"
@@ -753,6 +755,16 @@ namespace UE::USDStageImporter::Private
 			Package->MarkPackageDirty();
 		}
 
+		// If we swapped the asset inside of a package we need to invalidate the linker and disconnect
+		// the package from the file on disk. Otherwise some code paths that go through it to find the
+		// package's asset will find our old asset in Engine/Transient instead of the new reimported asset
+		// (Repro for this: Import static mesh and place it into the level -> Save everything -> Reimport
+		// static mesh -> Reload level -> Without this code, the reloaded level will point at the old static mesh)
+		if (ExistingPackage)
+		{
+			ResetLoaders(ExistingPackage);
+		}
+
 		// Reopen asset editor if we were editing the asset
 		if (bAssetWasOpen)
 		{
@@ -1339,23 +1351,25 @@ namespace UE::USDStageImporter::Private
 		UStaticMesh* ExistingMesh = Cast<UStaticMesh>(ExistingAsset);
 		UStaticMesh* NewMesh = Cast<UStaticMesh>(NewAsset);
 
-		if (ExistingAsset && NewMesh)
+		if (ExistingMesh && NewMesh)
 		{
-			int32 NumExistingMaterials = ExistingMesh->GetStaticMaterials().Num();
-			int32 NumNewMaterials = NewMesh->GetStaticMaterials().Num();
+			TArray<FStaticMaterial>& NewMaterials = NewMesh->GetStaticMaterials();
+			TArray<FStaticMaterial>& ExistingMaterials = ExistingMesh->GetStaticMaterials();
+			int32 NumExistingMaterials = ExistingMaterials.Num();
+			int32 NumNewMaterials = NewMaterials.Num();
 
-			for (int32 NewMaterialIndex = 0; NewMaterialIndex < NumNewMaterials; ++NewMaterialIndex)
+			for (int32 NewMaterialIndex = 0; NewMaterialIndex < NumNewMaterials && NewMaterialIndex < NumExistingMaterials; ++NewMaterialIndex)
 			{
 				UMaterialInterface* ExistingMaterial = ExistingMesh->GetMaterial(NewMaterialIndex);
 
 				// Can't use SetMaterial as it starts a scoped transaction that would hold on to our transient assets...
-				NewMesh->GetStaticMaterials()[NewMaterialIndex].MaterialInterface = ExistingMaterial;
+				NewMaterials[NewMaterialIndex].MaterialInterface = ExistingMaterial;
 			}
 
 			// Clear out any other assignments we may have
-			for (int32 Index = NumNewMaterials; Index < NumExistingMaterials; ++Index)
+			for (int32 Index = NumExistingMaterials; Index < NumNewMaterials; ++Index)
 			{
-				NewMesh->GetStaticMaterials()[Index].MaterialInterface = nullptr;
+				NewMaterials[Index].MaterialInterface = nullptr;
 			}
 
 			return;
@@ -1415,19 +1429,52 @@ namespace UE::USDStageImporter::Private
 			return;
 		}
 
-		// We never want to remap from all instances of nullptr to something else
 		TMap<UObject*, UObject*> ProcessedObjectsToRemap;
 		ProcessedObjectsToRemap.Reserve(ObjectsToRemap.Num());
 		for (const TPair<UObject*, UObject*>& SourceTarget : ObjectsToRemap)
 		{
-			if (SourceTarget.Key)
+			// We never want to remap from all instances of nullptr to something else
+			if (SourceTarget.Key == nullptr)
 			{
-				ProcessedObjectsToRemap.Add(SourceTarget);
+				continue;
 			}
+
+			ProcessedObjectsToRemap.Add(SourceTarget);
 		}
 
-		// Remap references held by assets that were moved directly to the destination package, and won't be in ProcessedObjectsToRemap
 		TSet<UObject*> Referencers = PublishedObjects;
+
+		// Get all in-memory referencers.
+		//
+		// We need this as we need to remap even references internal to other loaded assets, like a blueprint component's reference
+		// or a spawnable reference, etc.
+		//
+		// We don't need to search inside unloaded assets because those would presumably already find the proper remapped
+		// asset via path whenever they are loaded. By the same logic we don't need to do the same for soft object path references,
+		// and only really need to update live UObject references
+		{
+			TSet<UObject*> OldAssets;
+			OldAssets.Reserve(ObjectsToRemap.Num());
+			for (const TPair<UObject*, UObject*>& Pair : ObjectsToRemap)
+			{
+				UObject* OldAsset = Pair.Key;
+				if (!OldAsset)
+				{
+					continue;
+				}
+				OldAssets.Add(OldAsset);
+			}
+
+			TArray<UObject*> FoundReferencerObjects = FReferencerFinder::GetAllReferencers(
+				OldAssets,
+				nullptr,
+				EReferencerFinderFlags::SkipInnerReferences
+			);
+
+			Referencers.Append(FoundReferencerObjects);
+		}
+
+		// Remap references held by assets that were moved directly to the destination package, and won't be in ProcessedObjectsToRemap
 		if (GEditor)
 		{
 			// Remap references to spawned assets/actors from all worlds.
@@ -1450,7 +1497,7 @@ namespace UE::USDStageImporter::Private
 			Referencers.Add(Pair.Value);
 		}
 
-		// Fix references between actors and assets (e.g. mesh in final package referencing material in transient package)
+		// Actually fixup all the referencers that we've collected.
 		// Note we don't care if transient assets reference each other, as we'll delete them all at once anyway
 		for (UObject* Referencer : Referencers)
 		{
