USD: Fix deprecation warnings due to old AssetCache2 usage.

#Synced-CL 49920387

diff --git a/Engine/Plugins/Importers/USDImporter/Source/USDStageImporter/Public/USDStageImportContext.h b/Engine/Plugins/Importers/USDImporter/Source/USDStageImporter/Public/USDStageImportContext.h
--- a/Engine/Plugins/Importers/USDImporter/Source/USDStageImporter/Public/USDStageImportContext.h
+++ b/Engine/Plugins/Importers/USDImporter/Source/USDStageImporter/Public/USDStageImportContext.h
@@ -64,12 +64,14 @@ struct FUsdStageImportContext
 	UPROPERTY()
 	TObjectPtr<UUsdAssetCache3> UsdAssetCache;
 
+	PRAGMA_DISABLE_DEPRECATION_WARNINGS
 	UE_DEPRECATED(all, "Use the 'UsdAssetCache' member instead, which is of the new UUsdAssetCache3 type")
 	UPROPERTY()
 	TObjectPtr<UUsdAssetCache2> AssetCache_DEPRECATED;
 
 	UE_DEPRECATED(all, "The import process now always builds its own InfoCache, so this member is no longer used")
 	TSharedPtr<FUsdInfoCache> InfoCache_DEPRECATED;
+	PRAGMA_ENABLE_DEPRECATION_WARNINGS
 
 	/** Bounding box cache used for the USD stage in case we have to spawn bounds components */
 	TSharedPtr<UE::FUsdGeomBBoxCache> BBoxCache;
