USD: Add support for exporting mesh component overlay materials to USD.

Don't export mesh components without any mesh assigned, as that produces Mesh prims without any valid mesh data.

#Synced-CL 50488581

#Synced-CL 50488581

diff --git a/Engine/Plugins/Importers/USDImporter/Source/USDExporter/Private/USDConversionBlueprintContext.cpp b/Engine/Plugins/Importers/USDImporter/Source/USDExporter/Private/USDConversionBlueprintContext.cpp
--- a/Engine/Plugins/Importers/USDImporter/Source/USDExporter/Private/USDConversionBlueprintContext.cpp
+++ b/Engine/Plugins/Importers/USDImporter/Source/USDExporter/Private/USDConversionBlueprintContext.cpp
@@ -528,7 +528,8 @@ bool UUsdConversionBlueprintContext::ConvertMaterialOverrides(
 	const TArray<UMaterialInterface*> MaterialOverrides,
 	const FString& PrimPath,
 	int32 LowestLOD,
-	int32 HighestLOD
+	int32 HighestLOD,
+	bool bCheckAssetUserData
 )
 {
 #if USE_USD_SDK
@@ -538,7 +539,7 @@ bool UUsdConversionBlueprintContext::ConvertMaterialOverrides(
 		return false;
 	}
 
-	return UnrealToUsd::ConvertMaterialOverrides(MeshAsset, MaterialOverrides, Prim, LowestLOD, HighestLOD);
+	return UnrealToUsd::ConvertMaterialOverrides(MeshAsset, MaterialOverrides, Prim, LowestLOD, HighestLOD, bCheckAssetUserData);
 #else
 	return false;
 #endif	  // USE_USD_SDK
@@ -568,6 +569,28 @@ void UUsdConversionBlueprintContext::ReplaceUnrealMaterialsWithBaked(
 #endif	  // USE_USD_SDK
 }
 
+void UUsdConversionBlueprintContext::AuthorUnrealMaterialBinding(
+	const FString& MeshLikePrimPath,
+	const FString& UnrealMaterialPathName,
+	bool bStrongerThanDescendants
+)
+{
+#if USE_USD_SDK
+	if (!Stage || UnrealMaterialPathName.IsEmpty())
+	{
+		return;
+	}
+
+	UE::FUsdPrim Prim = Stage.GetPrimAtPath(UE::FSdfPath{*MeshLikePrimPath});
+	if (!Prim)
+	{
+		return;
+	}
+
+	UsdUtils::AuthorUnrealMaterialBinding(Prim, UnrealMaterialPathName, bStrongerThanDescendants);
+#endif	  // USE_USD_SDK
+}
+
 bool UUsdConversionBlueprintContext::RemoveUnrealSurfaceOutput(const FString& PrimPath, const FFilePath& LayerToAuthorIn)
 {
 	return false;
diff --git a/Engine/Plugins/Importers/USDImporter/Source/USDExporter/Public/USDConversionBlueprintContext.h b/Engine/Plugins/Importers/USDImporter/Source/USDExporter/Public/USDConversionBlueprintContext.h
--- a/Engine/Plugins/Importers/USDImporter/Source/USDExporter/Public/USDConversionBlueprintContext.h
+++ b/Engine/Plugins/Importers/USDImporter/Source/USDExporter/Public/USDConversionBlueprintContext.h
@@ -173,13 +173,23 @@ public:
 	);
 
 public:
+	/**
+	 * Sets the material overrides as opinions on MeshPrim, taking care to figure out whether it needs to create LOD or subsection prims, and
+	 * which ones given the chosen export LOD range.
+	 * MeshAsset should be an UStaticMesh, USkeletalMesh or UGeometryCache.
+	 *
+	 * If bCheckAssetUserData is enabled, this function which check the asset's USD Asset User Data in order to determine the mapping between
+	 * material slot and prim path, which is useful when exporting multiple slots to geom subsets or LOD prims. If this is false, The overrides
+	 * will be authored to standard paths used by the USDImporter plugin (i.e. LOD<N> prims when needed, UsdGeomSubsets named "Section<N>")
+	 */
 	UFUNCTION(BlueprintCallable, Category = "Conversion utils")
 	UE_API bool ConvertMaterialOverrides(
 		const UObject* MeshAsset,
 		const TArray<UMaterialInterface*> MaterialOverrides,
 		const FString& PrimPath,
 		int32 LowestLOD,
-		int32 HighestLOD
+		int32 HighestLOD,
+		bool bCheckAssetUserData = true
 	);
 
 	/**
@@ -202,6 +212,18 @@ public:
 		bool bRemoveUnrealMaterials = false
 	);
 
+	/**
+	 * Writes UnrealMaterialPathName as a material binding for MeshLikePrimPath, either by reusing an existing UsdShadeMaterial
+	 * binding if it already has an 'unreal' render context output and the expected structure, or by creating a new Material prim
+	 * that fulfills those requirements and binding to it (optionally with strongerThanDescendants strength in that case).
+	 */
+	UFUNCTION(BlueprintCallable, Category = "Conversion utils")
+	UE_API void AuthorUnrealMaterialBinding(
+		const FString& MeshLikePrimPath,
+		const FString& UnrealMaterialPathName,
+		bool bStrongerThanDescendants = false
+	);
+
 	UE_DEPRECATED(all, "No longer used as UE material assignments are only visible in the 'unreal' render context anyway")
 	UFUNCTION()
 	UE_API bool RemoveUnrealSurfaceOutput(const FString& PrimPath, const FFilePath& LayerToAuthorIn);
diff --git a/Engine/Plugins/Runtime/USDCore/Source/USDUtilities/Private/USDPrimConversion.cpp b/Engine/Plugins/Runtime/USDCore/Source/USDUtilities/Private/USDPrimConversion.cpp
--- a/Engine/Plugins/Runtime/USDCore/Source/USDUtilities/Private/USDPrimConversion.cpp
+++ b/Engine/Plugins/Runtime/USDCore/Source/USDUtilities/Private/USDPrimConversion.cpp
@@ -3726,7 +3726,8 @@ bool UnrealToUsd::ConvertMaterialOverrides(
 	const TArray<UMaterialInterface*> MaterialOverrides,
 	pxr::UsdPrim& UsdPrim,
 	int32 LowestLOD,
-	int32 HighestLOD
+	int32 HighestLOD,
+	bool bCheckAssetUserData
 )
 {
 #if WITH_EDITOR
@@ -3861,12 +3862,16 @@ bool UnrealToUsd::ConvertMaterialOverrides(
 				continue;
 			}
 
+			const UUsdMeshAssetUserData* UserData = bCheckAssetUserData
+				? const_cast<UGeometryCache*>(GeometryCache)->GetAssetUserData<UUsdMeshAssetUserData>()
+				: nullptr;
+
 			// If we have user data this is one of our meshes, so we know exactly the prim that corresponds to each
 			// material slot. Let's use that.
 			//
 			// const_cast as there's no const access to asset user data on IInterface_AssetUserData, but we won't
 			// modify anything
-			if (const UUsdMeshAssetUserData* UserData = const_cast<UGeometryCache*>(GeometryCache)->GetAssetUserData<UUsdMeshAssetUserData>())
+			if (UserData)
 			{
 				AuthorOverrideViaUserData(UserData, MatIndex, {MeshPrimPath}, Override->GetPathName());
 			}
@@ -3905,7 +3910,9 @@ bool UnrealToUsd::ConvertMaterialOverrides(
 		}
 		const bool bHasLODs = NumLODs > 1;
 
-		const UUsdMeshAssetUserData* UserData = const_cast<UStaticMesh*>(StaticMesh)->GetAssetUserData<UUsdMeshAssetUserData>();
+		const UUsdMeshAssetUserData* UserData = bCheckAssetUserData
+			? const_cast<UStaticMesh*>(StaticMesh)->GetAssetUserData<UUsdMeshAssetUserData>() 
+			: nullptr;
 
 		for (int32 MatIndex = 0; MatIndex < MaterialOverrides.Num(); ++MatIndex)
 		{
@@ -3978,7 +3985,9 @@ bool UnrealToUsd::ConvertMaterialOverrides(
 			return false;
 		}
 
-		const UUsdMeshAssetUserData* UserData = const_cast<USkeletalMesh*>(SkeletalMesh)->GetAssetUserData<UUsdMeshAssetUserData>();
+		const UUsdMeshAssetUserData* UserData = bCheckAssetUserData
+			? const_cast<USkeletalMesh*>(SkeletalMesh)->GetAssetUserData<UUsdMeshAssetUserData>()
+			: nullptr;
 
 		TIndirectArray<FSkeletalMeshLODRenderData>& LodRenderData = RenderData->LODRenderData;
 		if (LodRenderData.Num() == 0)
diff --git a/Engine/Plugins/Runtime/USDCore/Source/USDUtilities/Private/USDShadeConversion.cpp b/Engine/Plugins/Runtime/USDCore/Source/USDUtilities/Private/USDShadeConversion.cpp
--- a/Engine/Plugins/Runtime/USDCore/Source/USDUtilities/Private/USDShadeConversion.cpp
+++ b/Engine/Plugins/Runtime/USDCore/Source/USDUtilities/Private/USDShadeConversion.cpp
@@ -3015,7 +3015,7 @@ TArray<FString> UsdUtils::GetMaterialXFilePaths(const pxr::UsdPrim& Prim)
 }
 #endif	  // WITH_EDITOR
 
-void UsdUtils::AuthorUnrealMaterialBinding(pxr::UsdPrim& MeshOrGeomSubsetPrim, const FString& UnrealMaterialPathName)
+void UsdUtils::AuthorUnrealMaterialBinding(pxr::UsdPrim& MeshOrGeomSubsetPrim, const FString& UnrealMaterialPathName, bool bStrongerThanDescendants)
 {
 	if (!MeshOrGeomSubsetPrim || UnrealMaterialPathName.IsEmpty())
 	{
@@ -3086,7 +3086,10 @@ void UsdUtils::AuthorUnrealMaterialBinding(pxr::UsdPrim& MeshOrGeomSubsetPrim, c
 	{
 		UsdUtils::SetUnrealSurfaceOutput(MaterialPrim, UnrealMaterialPathName);
 
-		BindingAPI.Bind(ChildMaterial);
+		BindingAPI.Bind(
+			ChildMaterial,
+			bStrongerThanDescendants ? pxr::UsdShadeTokens->strongerThanDescendants : pxr::UsdShadeTokens->fallbackStrength
+		);
 	}
 }
 
diff --git a/Engine/Plugins/Runtime/USDCore/Source/USDUtilities/Public/USDPrimConversion.h b/Engine/Plugins/Runtime/USDCore/Source/USDUtilities/Public/USDPrimConversion.h
--- a/Engine/Plugins/Runtime/USDCore/Source/USDUtilities/Public/USDPrimConversion.h
+++ b/Engine/Plugins/Runtime/USDCore/Source/USDUtilities/Public/USDPrimConversion.h
@@ -326,13 +326,18 @@ namespace UnrealToUsd
 	 * which ones given the chosen export LOD range.
 	 * Leaving the LOD range at default of INDEX_NONE will author the material override opinions assuming all LODs of the MeshAsset were exported.
 	 * MeshAsset should be an UStaticMesh, USkeletalMesh or UGeometryCache.
+	 *
+	 * If bCheckAssetUserData is enabled, this function will check the asset's USD Asset User Data in order to determine the mapping between
+	 * material slot and prim path, which is useful when exporting multiple slots to geom subsets or LOD prims. If this is false, The overrides
+	 * will be authored to standard paths used by the USDImporter plugin (i.e. LOD<N> prims when needed, UsdGeomSubsets named "Section<N>")
 	 */
 	USDUTILITIES_API bool ConvertMaterialOverrides(
 		const UObject* MeshAsset,
 		const TArray<UMaterialInterface*> MaterialOverrides,
 		pxr::UsdPrim& UsdPrim,
 		int32 LowestLOD = INDEX_NONE,
-		int32 HighestLOD = INDEX_NONE
+		int32 HighestLOD = INDEX_NONE,
+		bool bCheckAssetUserData = true
 	);
 
 	/**
diff --git a/Engine/Plugins/Runtime/USDCore/Source/USDUtilities/Public/USDShadeConversion.h b/Engine/Plugins/Runtime/USDCore/Source/USDUtilities/Public/USDShadeConversion.h
--- a/Engine/Plugins/Runtime/USDCore/Source/USDUtilities/Public/USDShadeConversion.h
+++ b/Engine/Plugins/Runtime/USDCore/Source/USDUtilities/Public/USDShadeConversion.h
@@ -199,9 +199,14 @@ namespace UsdUtils
 {
 	// Writes UnrealMaterialPathName as a material binding for MeshOrGeomSubsetPrim, either by reusing an existing UsdShadeMaterial
 	// binding if it already has an 'unreal' render context output and the expected structure, or by creating a new Material prim
-	// that fulfills those requirements.
+	// that fulfills those requirements and binding to it (optionally with strongerThanDescendants strength in that case).
+	//
 	// Doesn't write to the 'unrealMaterial' attribute at all, as we intend on deprecating it in the future.
-	USDUTILITIES_API void AuthorUnrealMaterialBinding(pxr::UsdPrim& MeshOrGeomSubsetPrim, const FString& UnrealMaterialPathName);
+	USDUTILITIES_API void AuthorUnrealMaterialBinding(
+		pxr::UsdPrim& MeshOrGeomSubsetPrim,
+		const FString& UnrealMaterialPathName,
+		bool bStrongerThanDescendants = false
+	);
 
 	/**
 	 * Similar to AuthorUnrealMaterialBinding, but instead of authoring material bindings directly to TargetMeshOrGeomSubsetPrim,
