USD: Allow exporting levels to USD with the regular AssetExportTask Python snippet

#Synced-CL 50488581

diff --git a/Engine/Plugins/Importers/USDImporter/Source/USDExporter/Private/LevelExporterUSD.cpp b/Engine/Plugins/Importers/USDImporter/Source/USDExporter/Private/LevelExporterUSD.cpp
--- a/Engine/Plugins/Importers/USDImporter/Source/USDExporter/Private/LevelExporterUSD.cpp
+++ b/Engine/Plugins/Importers/USDImporter/Source/USDExporter/Private/LevelExporterUSD.cpp
@@ -69,11 +69,26 @@ bool ULevelExporterUSD::ExportBinary(UObject* Object, const TCHAR* Type, FArchiv
 		}
 	}
 
-	// Note how we don't explicitly pass the Options down to Python here: We stash our desired export options on the CDO, and
+	// Note how we don't explicitly pass the Options down to Python here: We stash our desired export options as the CDO, and
 	// those are read from Python by executing export_with_cdo_options().
 	if (IPythonScriptPlugin::Get()->IsPythonAvailable())
 	{
+		ULevelExporterUSDOptions* OptionsOriginalCDO = GetMutableDefault<ULevelExporterUSDOptions>();
+		TStrongObjectPtr<ULevelExporterUSDOptions> CDOPin;
+
+		const bool bNeedCDOSwap = Options != OptionsOriginalCDO;
+		if (bNeedCDOSwap)
+		{
+			CDOPin.Reset(OptionsOriginalCDO);
+			ULevelExporterUSDOptions::StaticClass()->SetDefaultObject(Options);
+		}
+
 		IPythonScriptPlugin::Get()->ExecPythonCommand(TEXT("import usd_unreal.level_exporter; usd_unreal.level_exporter.export_with_cdo_options()"));
+
+		if (bNeedCDOSwap)
+		{
+			ULevelExporterUSDOptions::StaticClass()->SetDefaultObject(OptionsOriginalCDO);
+		}
 	}
 	Options->CurrentTask = nullptr;
 
