USD: Add an explicit log message for renaming changes as USD doesn't fire the "add prim" notice in some cases

#Synced-CL 50957406

diff --git a/Engine/Plugins/Importers/USDImporter/Source/USDStage/Private/USDTransactor.cpp b/Engine/Plugins/Importers/USDImporter/Source/USDStage/Private/USDTransactor.cpp
--- a/Engine/Plugins/Importers/USDImporter/Source/USDStage/Private/USDTransactor.cpp
+++ b/Engine/Plugins/Importers/USDImporter/Source/USDStage/Private/USDTransactor.cpp
@@ -204,6 +204,15 @@ namespace UsdUtils
 					);
 				}
 
+				if (Prim && ConvertedChange.Flags.bDidRename)
+				{
+					USD_LOG_INFO(
+						TEXT("Recorded the rename of prim '%s' into '%s'"),
+						*ConvertedChange.OldPath,
+						*ObjectPath
+					);
+				}
+
 				// Don't record any attribute changes if we can't find the prim anyway
 				if (!Prim)
 				{
