USD: Spawn camera actors as movable, following the behavior of spawning them manually from the Place Actors menu

#Synced-CL 51330716

diff --git a/Engine/Plugins/Importers/USDImporter/Source/USDSchemas/Private/USDGeomXformableTranslator.cpp b/Engine/Plugins/Importers/USDImporter/Source/USDSchemas/Private/USDGeomXformableTranslator.cpp
--- a/Engine/Plugins/Importers/USDImporter/Source/USDSchemas/Private/USDGeomXformableTranslator.cpp
+++ b/Engine/Plugins/Importers/USDImporter/Source/USDSchemas/Private/USDGeomXformableTranslator.cpp
@@ -20,6 +20,7 @@
 #include "USDShadeConversion.h"
 #include "USDTypesConversion.h"
 
+#include "Camera/CameraActor.h"
 #include "Components/LightComponentBase.h"
 #include "Components/SceneComponent.h"
 #include "Components/SkeletalMeshComponent.h"
@@ -598,10 +599,11 @@ USceneComponent* FUsdGeomXformableTranslator::CreateComponentsEx(TOptional<TSubc
 		const bool bIsAssociating = Context->Level && Context->Level->bIsAssociatingLevel;
 		const bool bParentIsMovable = Context->ParentComponent && Context->ParentComponent->Mobility == EComponentMobility::Movable;
 		const bool bParentIsStationary = Context->ParentComponent && Context->ParentComponent->Mobility == EComponentMobility::Stationary;
+		const bool bIsCamera = SceneComponent->GetOwner() && SceneComponent->GetOwner()->IsA(ACameraActor::StaticClass());
 
 		// Don't call SetMobility as it would trigger a reregister, queuing unnecessary rhi commands since this is a brand new component
 		// Always have movable Skeletal mesh components or else we get some warnings when building physics assets
-		if (bIsAssociating || bParentIsMovable || SceneComponent->IsA<USkeletalMeshComponent>() || UsdUtils::IsAnimated(Prim))
+		if (bIsAssociating || bParentIsMovable || bIsCamera || SceneComponent->IsA<USkeletalMeshComponent>() || UsdUtils::IsAnimated(Prim))
 		{
 			SceneComponent->Mobility = EComponentMobility::Movable;
 		}
