USD: Add feedback when the reimport fails due to the stage not having the target prim.

#Synced-CL 51466015

diff --git a/Engine/Plugins/Importers/USDImporter/Source/USDStageImporter/Private/USDStageImporter.cpp b/Engine/Plugins/Importers/USDImporter/Source/USDStageImporter/Private/USDStageImporter.cpp
--- a/Engine/Plugins/Importers/USDImporter/Source/USDStageImporter/Private/USDStageImporter.cpp
+++ b/Engine/Plugins/Importers/USDImporter/Source/USDStageImporter/Private/USDStageImporter.cpp
@@ -2131,6 +2131,11 @@ bool UUsdStageImporter::ReimportSingleAsset(
 #if USE_USD_SDK
 	double StartTime = FPlatformTime::Cycles64();
 
+	if (!OriginalAsset)
+	{
+		return false;
+	}
+
 	if (!ImportContext.ImportOptions)
 	{
 		ImportContext.ImportOptions = NewObject<UUsdStageImportOptions>();
@@ -2280,14 +2285,26 @@ bool UUsdStageImporter::ReimportSingleAsset(
 	TranslationContext->BBoxCache = ImportContext.BBoxCache;
 	TranslationContext->BlendShapesByPath = &BlendShapesByPath;
 	TranslationContext->GroomInterpolationSettings = ImportContext.ImportOptions->GroomInterpolationSettings;
+
+	bool bFoundReimportPrim = false;
 	{
 		UE::USDStageImporter::Private::CacheCollapsingState(TranslationContext.Get());
 
 		UE::FUsdPrim TargetPrim = ImportContext.Stage.GetPrimAtPath(UE::FSdfPath(*OriginalPrimPath));
 		if (TargetPrim)
 		{
+			bFoundReimportPrim = true;
 			UE::USDStageImporter::Private::CreateAssetsForPrims({TargetPrim}, TranslationContext.Get(), LOCTEXT("CreateAssets", "Creating assets"));
 		}
+		else
+		{
+			USD_LOG_USERERROR(FText::Format(
+				LOCTEXT("ReimportMissingPrim", "Failed to find prim '{0}' in stage '{1}' to use for the reimport of asset '{2}'!"),
+				FText::FromString(OriginalPrimPath),
+				FText::FromString(ImportContext.FilePath),
+				FText::FromString(OriginalAsset->GetPathName())
+			));
+		}
 	}
 	TranslationContext->CompleteTasks();
 
@@ -2343,6 +2360,21 @@ bool UUsdStageImporter::ReimportSingleAsset(
 
 		bSuccess = OutReimportedAsset != nullptr && ActiveAssets.Contains(ReimportedObject);
 	}
+	else if (bFoundReimportPrim)
+	{
+		// If we're in here we did find the prim we wanted on the stage, but for some reason it didn't produce the
+		// expected asset type
+		USD_LOG_USERERROR(FText::Format(
+			LOCTEXT(
+				"ReimportMissingCompatibleAsset",
+				"Failed to produce a '{0}' asset from prim '{1}' in stage '{2}' that could be used for the reimport of '{3}'!"
+			),
+			FText::FromString(OriginalAsset->GetClass()->GetName()),
+			FText::FromString(OriginalPrimPath),
+			FText::FromString(ImportContext.FilePath),
+			FText::FromString(OriginalAsset->GetPathName())
+		));
+	}
 
 	const bool bIsReimport = true;
 	UE::USDStageImporter::Private::Cleanup(ImportContext.SceneActor, nullptr, ImportContext.ImportOptions->ExistingActorPolicy, AssetsToFinalize);
