USD: Fix missing assets and components when prim names differ only by letter case.

#Synced-CL 51466015

diff --git a/Engine/Plugins/Importers/USDImporter/Source/USDExporter/Private/LevelSequenceExporterUSD.cpp b/Engine/Plugins/Importers/USDImporter/Source/USDExporter/Private/LevelSequenceExporterUSD.cpp
--- a/Engine/Plugins/Importers/USDImporter/Source/USDExporter/Private/LevelSequenceExporterUSD.cpp
+++ b/Engine/Plugins/Importers/USDImporter/Source/USDExporter/Private/LevelSequenceExporterUSD.cpp
@@ -4,6 +4,7 @@
 
 #include "LevelExporterUSDOptions.h"
 #include "LevelSequenceExporterUSDOptions.h"
+#include "USDSetUtils.h"
 #include "UnrealUSDWrapper.h"
 #include "USDClassesModule.h"
 #include "USDConversionUtils.h"
@@ -759,7 +760,7 @@ namespace UE::LevelSequenceExporterUSD::Private
 		const UMovieSceneAudioTrack& AudioTrack,
 		const FMovieSceneSequenceTransform& SequenceTransform,
 		UE::FUsdPrim& Prim,
-		TMap<FString, int32>& AudioTracksPerPrim
+		UsdUtils::TCaseSensitiveStringMap<int32>& AudioTracksPerPrim
 	)
 	{
 		FString PrimPath = Prim.GetPrimPath().GetString();
@@ -1179,7 +1180,7 @@ namespace UE::LevelSequenceExporterUSD::Private
 				}
 			}
 
-			TMap<FString, int32> AudioTracksPerPrim;
+			UsdUtils::TCaseSensitiveStringMap<int32> AudioTracksPerPrim;
 			if (const FMovieSceneBinding* Binding = MovieScene->FindBinding(InstanceKey.Key))
 			{
 				for (const UMovieSceneTrack* Track : Binding->GetTracks())
diff --git a/Engine/Plugins/Importers/USDImporter/Source/USDSchemas/Private/USDSkelSkeletonTranslator.cpp b/Engine/Plugins/Importers/USDImporter/Source/USDSchemas/Private/USDSkelSkeletonTranslator.cpp
--- a/Engine/Plugins/Importers/USDImporter/Source/USDSchemas/Private/USDSkelSkeletonTranslator.cpp
+++ b/Engine/Plugins/Importers/USDImporter/Source/USDSchemas/Private/USDSkelSkeletonTranslator.cpp
@@ -20,6 +20,7 @@
 #include "USDMemory.h"
 #include "USDObjectUtils.h"
 #include "USDPrimConversion.h"
+#include "USDSetUtils.h"
 #include "USDSkeletalDataConversion.h"
 #include "USDTranslatorUtils.h"
 #include "USDTypesConversion.h"
@@ -490,7 +491,7 @@ namespace UsdSkelSkeletonTranslatorImpl
 
 		TMap<int32, FSkeletalMeshImportData> LODIndexToSkeletalMeshImportDataMap;
 		TMap<int32, UsdUtils::FUsdPrimMaterialAssignmentInfo> LODIndexToMaterialInfoMap;
-		TSet<FString> ProcessedLODParentPaths;
+		UsdUtils::FCaseSensitiveStringSet ProcessedLODParentPaths;
 
 		// Since we may need to switch variants to parse LODs, we could invalidate references to SkinningQuery objects, so we need
 		// to keep track of these by path and construct one whenever we need them
@@ -792,7 +793,7 @@ namespace UsdSkelSkeletonTranslatorImpl
 			return true;
 		};
 
-		TSet<FString> ProcessedLODParentPaths;
+		UsdUtils::FCaseSensitiveStringSet ProcessedLODParentPaths;
 
 		// Because we combine all skinning target meshes into a single skeletal mesh, we'll have to reconstruct the combined
 		// material assignment info that this SkelRoot wants in order to compare with the existing assignments.
diff --git a/Engine/Plugins/Importers/USDImporter/Source/USDStage/Private/USDPrimTwin.cpp b/Engine/Plugins/Importers/USDImporter/Source/USDStage/Private/USDPrimTwin.cpp
--- a/Engine/Plugins/Importers/USDImporter/Source/USDStage/Private/USDPrimTwin.cpp
+++ b/Engine/Plugins/Importers/USDImporter/Source/USDStage/Private/USDPrimTwin.cpp
@@ -6,6 +6,7 @@
 #include "USDStageActor.h"
 
 #include "Components/InstancedStaticMeshComponent.h"
+#include "Containers/StringView.h"
 #include "Engine/World.h"
 
 #include UE_INLINE_GENERATED_CPP_BY_NAME(USDPrimTwin)
@@ -22,7 +23,7 @@ UUsdPrimTwin& UUsdPrimTwin::AddChild(const FString& InPrimPath)
 
 	Modify();
 
-	TObjectPtr<UUsdPrimTwin>& ChildPrim = Children.Add(ChildPrimName);
+	TObjectPtr<UUsdPrimTwin>& ChildPrim = Children.Add(UsdUtils::AddHashToPrimPath(ChildPrimName));
 
 	// Needs public because this will mostly live on the transient package (c.f. AUsdStageActor::GetRootPrimTwin())
 	ChildPrim = NewObject<UUsdPrimTwin>(this, NAME_None, GetFlags() | RF_Public);
@@ -154,15 +155,16 @@ UUsdPrimTwin* UUsdPrimTwin::Find(const FString& InPrimPath)
 		RestOfPrimPathToFind.Empty();
 	}
 
-	if (Children.Contains(ChildPrimName))
+	const FString& HashedChildName = UsdUtils::AddHashToPrimPath(ChildPrimName);
+	if (Children.Contains(HashedChildName))
 	{
 		if (RestOfPrimPathToFind.IsEmpty())
 		{
-			return Children[ChildPrimName];
+			return Children[HashedChildName];
 		}
 		else
 		{
-			return Children[ChildPrimName]->Find(RestOfPrimPathToFind);
+			return Children[HashedChildName]->Find(RestOfPrimPathToFind);
 		}
 	}
 
diff --git a/Engine/Plugins/Importers/USDImporter/Source/USDStage/Private/USDStageActor.cpp b/Engine/Plugins/Importers/USDImporter/Source/USDStage/Private/USDStageActor.cpp
--- a/Engine/Plugins/Importers/USDImporter/Source/USDStage/Private/USDStageActor.cpp
+++ b/Engine/Plugins/Importers/USDImporter/Source/USDStage/Private/USDStageActor.cpp
@@ -843,9 +843,9 @@ struct FUsdStageActorImpl
 #endif	  // USE_USD_SDK
 	}
 
-	static TSet<FString> GetPointInstancerPrototypes(const UE::FUsdPrim& Prim)
+	static UsdUtils::FCaseSensitiveStringSet GetPointInstancerPrototypes(const UE::FUsdPrim& Prim)
 	{
-		TSet<FString> PrototypePaths;
+		UsdUtils::FCaseSensitiveStringSet PrototypePaths;
 
 #if USE_USD_SDK
 		static FString PrototypesStr = UsdToUnreal::ConvertToken(pxr::UsdGeomTokens->prototypes);
@@ -2118,8 +2118,6 @@ UUsdPrimTwin* AUsdStageActor::ExpandPrim(
 				}
 				else if (Prim.IsA(TEXT("PointInstancer")))
 				{
-					TSet<FString> PrototypePaths = FUsdStageActorImpl::GetPointInstancerPrototypes(Prim);
-
 					for (const TObjectPtr<USceneComponent>& Child : TwinSceneComponent->GetAttachChildren())
 					{
 						UInstancedStaticMeshComponent* ISMComponent = Cast<UInstancedStaticMeshComponent>(Child.Get());
@@ -2216,7 +2214,7 @@ UUsdPrimTwin* AUsdStageActor::ExpandPrim(
 		else if (Prim.IsA(TEXT("PointInstancer")))
 		{
 			// Collect all the known prototype paths for this PointInstancer
-			TSet<FString> PrototypePaths = FUsdStageActorImpl::GetPointInstancerPrototypes(Prim);
+			UsdUtils::FCaseSensitiveStringSet PrototypePaths = FUsdStageActorImpl::GetPointInstancerPrototypes(Prim);
 
 			const TArray<TObjectPtr<USceneComponent>>& ChildComponents = TwinSceneComponent->GetAttachChildren();
 			for (const TObjectPtr<USceneComponent>& Child : ChildComponents)
@@ -2294,15 +2292,17 @@ UUsdPrimTwin* AUsdStageActor::ExpandPrim(
 		}
 	}
 
+	const FString HashedPrimPath = UsdUtils::AddHashToPrimPath(UsdPrimTwin->PrimPath);
+
 	if (bIsAnimated)
 	{
 		const bool bForceVisibilityTracks = bParentHasAnimatedVisibility.GetValue();
 		LevelSequenceHelper.AddPrim(*UsdPrimTwin, bForceVisibilityTracks, bHasAnimatedBounds);
-		PrimsToAnimate.Add(UsdPrimTwin->PrimPath);
+		PrimsToAnimate.Add(HashedPrimPath);
 	}
 	else
 	{
-		PrimsToAnimate.Remove(UsdPrimTwin->PrimPath);
+		PrimsToAnimate.Remove(HashedPrimPath);
 		LevelSequenceHelper.RemovePrim(*UsdPrimTwin);
 	}
 
@@ -2324,7 +2324,7 @@ UUsdPrimTwin* AUsdStageActor::ExpandPrim(
 			// considered animated too, so lets add it to the proper locations. This will also ensure that
 			// we can close the sequencer after creating a new animation in this way and see it animate on
 			// the level
-			PrimsToAnimate.Add(UsdPrimTwin->PrimPath);
+			PrimsToAnimate.Add(HashedPrimPath);
 			LevelSequenceHelper.AddPrim(*UsdPrimTwin);
 
 			// Prevent register/unregister spam when calling FUsdGeomXformableTranslator::UpdateComponents later
@@ -4291,7 +4291,7 @@ void AUsdStageActor::RequestDelayedTransactorReset()
 
 void AUsdStageActor::OnUsdPrimTwinDestroyed(const UUsdPrimTwin& UsdPrimTwin)
 {
-	PrimsToAnimate.Remove(UsdPrimTwin.PrimPath);
+	PrimsToAnimate.Remove(UsdUtils::AddHashToPrimPath(UsdPrimTwin.PrimPath));
 	ObjectsToWatch.Remove(UsdPrimTwin.SceneComponent.Get());
 	LevelSequenceHelper.RemovePrim(UsdPrimTwin);
 }
@@ -5090,16 +5090,17 @@ void AUsdStageActor::AnimatePrims()
 	}
 #endif	  // WITH_EDITOR
 
-	for (const FString& PrimToAnimate : PrimsToAnimate)
+	for (const FString& PrimPathWithHash : PrimsToAnimate)
 	{
-		UE::FSdfPath PrimPath(*PrimToAnimate);
+		const FString PrimPathWithoutHash = UsdUtils::RemoveHashFromPrimPath(PrimPathWithHash);
+		UE::FSdfPath PrimPath{*PrimPathWithoutHash};
 
 		if (TSharedPtr<FUsdSchemaTranslator> SchemaTranslator = FUsdSchemaTranslatorRegistry::Get().CreateTranslatorForSchema(
 				TranslationContext,
 				UE::FUsdTyped(CurrentStage.GetPrimAtPath(PrimPath))
 			))
 		{
-			if (UUsdPrimTwin* UsdPrimTwin = GetRootPrimTwin()->Find(PrimToAnimate))
+			if (UUsdPrimTwin* UsdPrimTwin = GetRootPrimTwin()->Find(PrimPathWithoutHash))
 			{
 				SchemaTranslator->UpdateComponents(UsdPrimTwin->SceneComponent.Get());
 			}
diff --git a/Engine/Plugins/Importers/USDImporter/Source/USDStage/Private/USDTransactor.cpp b/Engine/Plugins/Importers/USDImporter/Source/USDStage/Private/USDTransactor.cpp
--- a/Engine/Plugins/Importers/USDImporter/Source/USDStage/Private/USDTransactor.cpp
+++ b/Engine/Plugins/Importers/USDImporter/Source/USDStage/Private/USDTransactor.cpp
@@ -8,6 +8,7 @@
 #include "USDListener.h"
 #include "USDMemory.h"
 #include "USDPrimConversion.h"
+#include "USDSetUtils.h"
 #include "USDStageActor.h"
 #include "USDTypesConversion.h"
 #include "USDValueConversion.h"
@@ -682,7 +683,7 @@ namespace UsdUtils
 	}
 
 	/** Applies the field value pairs to all prims on the stage, and returns a list of prim paths for modified prims */
-	TSet<FString> ApplyFieldMapToStage(const FTransactorEditStorage& EditStorage, EApplicationDirection Direction, AUsdStageActor* StageActor)
+	UsdUtils::FCaseSensitiveStringSet ApplyFieldMapToStage(const FTransactorEditStorage& EditStorage, EApplicationDirection Direction, AUsdStageActor* StageActor)
 	{
 		UE::FUsdStage& Stage = StageActor->GetOrOpenUsdStage();
 		if (!Stage)
@@ -690,7 +691,7 @@ namespace UsdUtils
 			return {};
 		}
 
-		TSet<FString> PrimsChanged;
+		UsdUtils::FCaseSensitiveStringSet PrimsChanged;
 
 		int32 Start = 0;
 		int32 End = 0;
@@ -977,7 +978,7 @@ namespace UsdUtils
 			// transaction buffer.
 			FScopedBlockNotices BlockNotices(StageActor->GetUsdListener());
 
-			TSet<FString> PrimsChanged = UsdUtils::ApplyFieldMapToStage(	//
+			UsdUtils::FCaseSensitiveStringSet PrimsChanged = UsdUtils::ApplyFieldMapToStage(	//
 				Values,
 				UsdUtils::EApplicationDirection::Reverse,
 				StageActor
@@ -1030,7 +1031,7 @@ namespace UsdUtils
 
 			UE::FUsdStage& Stage = StageActor->GetOrOpenUsdStage();
 
-			TSet<FString> PrimsChanged;
+			UsdUtils::FCaseSensitiveStringSet PrimsChanged;
 			if (bIsApplyingConcertSync && ReceivedValuesBeforeUndo.IsSet())
 			{
 				// If we're applying a received ConcertSync transaction that actually is an undo on the source client then we want to use it's
diff --git a/Engine/Plugins/Importers/USDImporter/Source/USDStage/Public/USDTransactor.h b/Engine/Plugins/Importers/USDImporter/Source/USDStage/Public/USDTransactor.h
--- a/Engine/Plugins/Importers/USDImporter/Source/USDStage/Public/USDTransactor.h
+++ b/Engine/Plugins/Importers/USDImporter/Source/USDStage/Public/USDTransactor.h
@@ -5,6 +5,8 @@
 #include "CoreMinimal.h"
 #include "UObject/WeakObjectPtr.h"
 
+#include "USDSetUtils.h"
+
 #include "USDTransactor.generated.h"
 
 #define UE_API USDSTAGE_API
@@ -16,7 +18,7 @@ namespace UsdUtils
 	class FUsdTransactorImpl;
 
 	struct FSdfChangeListEntry;
-	using FObjectChangesByPath = TMap<FString, TArray<FSdfChangeListEntry>>;
+	using FObjectChangesByPath = TCaseSensitiveStringMap<TArray<FSdfChangeListEntry>>;
 }
 
 namespace UE
diff --git a/Engine/Plugins/Importers/USDImporter/Source/USDStageEditor/Private/SUSDStage.cpp b/Engine/Plugins/Importers/USDImporter/Source/USDStageEditor/Private/SUSDStage.cpp
--- a/Engine/Plugins/Importers/USDImporter/Source/USDStageEditor/Private/SUSDStage.cpp
+++ b/Engine/Plugins/Importers/USDImporter/Source/USDStageEditor/Private/SUSDStage.cpp
@@ -14,6 +14,7 @@
 #include "USDLayerUtils.h"
 #include "USDMaterialUtils.h"
 #include "USDProjectSettings.h"
+#include "USDSetUtils.h"
 #include "USDStageActor.h"
 #include "USDStageEditorSettings.h"
 #include "USDStageImportOptions.h"
@@ -58,24 +59,6 @@ namespace SUSDStageConstants
 
 namespace SUSDStageImpl
 {
-#if PLATFORM_LINUX
-	struct FCaseSensitiveStringSetFuncs : BaseKeyFuncs<FString, FString>
-	{
-		static FORCEINLINE const FString& GetSetKey(const FString& Element)
-		{
-			return Element;
-		}
-		static FORCEINLINE bool Matches(const FString& A, const FString& B)
-		{
-			return A.Equals(B, ESearchCase::CaseSensitive);
-		}
-		static FORCEINLINE uint32 GetKeyHash(const FString& Key)
-		{
-			return FCrc::StrCrc32<TCHAR>(*Key);
-		}
-	};
-#endif
-
 	/**
 	 * Makes sure that AllLayers includes all the external references of any of its members.
 	 * Example: Receives AllLayers [a.usda], that references b.usda and c.usda, while c.usda also references d.usda -> Modifies AllLayers to be
@@ -85,7 +68,7 @@ namespace SUSDStageImpl
 	{
 		// Consider paths in a case-sensitive way for linux
 #if PLATFORM_LINUX
-		TSet<FString, FCaseSensitiveStringSetFuncs> LoadedLayers;
+		UsdUtils::FCaseSensitiveStringSet LoadedLayers;
 #else
 		TSet<FString> LoadedLayers;
 #endif
@@ -3092,8 +3075,8 @@ void SUsdStage::OnSlateTick(float Time)
 		// Don't bother refreshing individual prims if we'll do a full refresh anyway
 		if (!bStageTreeViewNeedsFullUpdate)
 		{
-			TSet<FString> RefreshedPrims;
-			TSet<FString> ResyncedPrims;
+			UsdUtils::FCaseSensitiveStringSet RefreshedPrims;
+			UsdUtils::FCaseSensitiveStringSet ResyncedPrims;
 
 			// Note how we go from oldest to newest update here, so after the loop we'll end up with LastResyncedPrim
 			// pointing at what it should, and PrimPathToSelect pointing at the last updated prim path to select
diff --git a/Engine/Plugins/Importers/USDImporter/Source/USDStageEditor/Private/SUSDStageTreeView.cpp b/Engine/Plugins/Importers/USDImporter/Source/USDStageEditor/Private/SUSDStageTreeView.cpp
--- a/Engine/Plugins/Importers/USDImporter/Source/USDStageEditor/Private/SUSDStageTreeView.cpp
+++ b/Engine/Plugins/Importers/USDImporter/Source/USDStageEditor/Private/SUSDStageTreeView.cpp
@@ -13,6 +13,7 @@
 #include "USDPrimViewModel.h"
 #include "USDProjectSettings.h"
 #include "USDReferenceOptions.h"
+#include "USDSetUtils.h"
 #include "USDTypesConversion.h"
 #include "UsdWrappers/SdfChangeBlock.h"
 #include "UsdWrappers/SdfPath.h"
@@ -2099,8 +2100,8 @@ void SUsdStageTreeView::OnPrimNameCommitted(const FUsdPrimViewModelRef& ViewMode
 		{
 			// e.g. "/Root/NewPrim"
 			FString NewPath = FString::Printf(TEXT("%s/%s"), *FPaths::GetPath(OldPath), *NewNameStr);
-			TSet<FString> EntriesToAdd;
-			for (TSet<FString>::TIterator It(ExpandedPrimPaths); It; ++It)
+			UsdUtils::FCaseSensitiveStringSet EntriesToAdd;
+			for (UsdUtils::FCaseSensitiveStringSet::TIterator It{ExpandedPrimPaths}; It; ++It)
 			{
 				// e.g. "/Root/OldPrim/SomeChild"
 				TStringView SomePrimPath = *It;
diff --git a/Engine/Plugins/Importers/USDImporter/Source/USDStageEditor/Private/SUSDStageTreeView.h b/Engine/Plugins/Importers/USDImporter/Source/USDStageEditor/Private/SUSDStageTreeView.h
--- a/Engine/Plugins/Importers/USDImporter/Source/USDStageEditor/Private/SUSDStageTreeView.h
+++ b/Engine/Plugins/Importers/USDImporter/Source/USDStageEditor/Private/SUSDStageTreeView.h
@@ -3,6 +3,7 @@
 #pragma once
 
 #include "USDPrimViewModel.h"
+#include "USDSetUtils.h"
 #include "UsdWrappers/UsdStage.h"
 #include "Widgets/SUSDTreeView.h"
 
@@ -131,7 +132,7 @@ private:
 	TWeakPtr<FUsdPrimViewModel> PendingRenameItem;
 
 	// So that we can store these across refreshes
-	TSet<FString> ExpandedPrimPaths;
+	UsdUtils::FCaseSensitiveStringSet ExpandedPrimPaths;
 	TOptional<bool> RootWasExpanded;
 	bool bNeedExpansionStateRefresh = false;
 
diff --git a/Engine/Plugins/Runtime/USDCore/Source/USDClasses/Public/USDSetUtils.h b/Engine/Plugins/Runtime/USDCore/Source/USDClasses/Public/USDSetUtils.h
new file mode 100644
--- /dev/null
+++ b/Engine/Plugins/Runtime/USDCore/Source/USDClasses/Public/USDSetUtils.h
@@ -0,0 +1,51 @@
+// Copyright Epic Games, Inc. All Rights Reserved.
+
+#pragma once
+
+#include "Containers/SetUtilities.h"
+#include "Containers/UnrealString.h"
+
+namespace UsdUtils
+{
+	/** Case sensitive hashing function for TMap */
+	template<typename ValueType>
+	struct FCaseSensitiveStringMapFuncs : BaseKeyFuncs<ValueType, FString, /*bInAllowDuplicateKeys*/ false>
+	{
+		static FORCEINLINE const FString& GetSetKey(const TPair<FString, ValueType>& Element)
+		{
+			return Element.Key;
+		}
+
+		static FORCEINLINE bool Matches(const FString& A, const FString& B)
+		{
+			return A.Equals(B, ESearchCase::CaseSensitive);
+		}
+
+		static FORCEINLINE uint32 GetKeyHash(const FString& Key)
+		{
+			return FCrc::StrCrc32<TCHAR>(*Key);
+		}
+	};
+
+	/** Case sensitive hashing function for TSet */
+	struct FCaseSensitiveStringSetFuncs : BaseKeyFuncs<FString, FString>
+	{
+		static FORCEINLINE const FString& GetSetKey(const FString& Element)
+		{
+			return Element;
+		}
+		static FORCEINLINE bool Matches(const FString& A, const FString& B)
+		{
+			return A.Equals(B, ESearchCase::CaseSensitive);
+		}
+		static FORCEINLINE uint32 GetKeyHash(const FString& Key)
+		{
+			return FCrc::StrCrc32<TCHAR>(*Key);
+		}
+	};
+
+	template<typename ValueType>
+	using TCaseSensitiveStringMap = TMap<FString, ValueType, FDefaultSetAllocator, FCaseSensitiveStringMapFuncs<ValueType>>;
+
+	using FCaseSensitiveStringSet = TSet<FString, FCaseSensitiveStringSetFuncs>;
+};	  // namespace UsdUnreal::ObjectUtils
diff --git a/Engine/Plugins/Runtime/USDCore/Source/USDUtilities/Private/USDConversionUtils.cpp b/Engine/Plugins/Runtime/USDCore/Source/USDUtilities/Private/USDConversionUtils.cpp
--- a/Engine/Plugins/Runtime/USDCore/Source/USDUtilities/Private/USDConversionUtils.cpp
+++ b/Engine/Plugins/Runtime/USDCore/Source/USDUtilities/Private/USDConversionUtils.cpp
@@ -52,6 +52,7 @@
 #include "GeometryCache.h"
 #include "GroomAsset.h"
 #include "GroomCache.h"
+#include "Hash/Fnv.h"
 #include "InstancedFoliageActor.h"
 #include "LandscapeProxy.h"
 #include "Materials/Material.h"
@@ -373,6 +374,20 @@ namespace USDConversionUtilsImpl
 	}
 }	 // namespace USDConversionUtilsImpl
 
+FString UsdUtils::AddHashToPrimPath(const FString& PrimPath)
+{
+	// Pad to exactly 10 digits, so it's easy to strip
+	// Adding the underscore makes it so that we can use this even on the pseudoroot path "/", without
+	// making an invalid SdfPath in the process (as something like /6154761547 is not a valid prim path)
+	return FString::Printf(TEXT("%s_%010u"), *PrimPath, UE::HashStringFNV1a32(MakeStringView(PrimPath)));
+}
+
+FString UsdUtils::RemoveHashFromPrimPath(const FString& HashedPrimPath)
+{
+	// Make sure this matches the number of digits on the Printf format above (+ underscore)!
+	return HashedPrimPath.LeftChop(11);
+}
+
 template<typename ValueType>
 ValueType UsdUtils::GetUsdValue(const pxr::UsdAttribute& Attribute, pxr::UsdTimeCode TimeCode)
 {
diff --git a/Engine/Plugins/Runtime/USDCore/Source/USDUtilities/Private/USDSkeletalDataConversion.cpp b/Engine/Plugins/Runtime/USDCore/Source/USDUtilities/Private/USDSkeletalDataConversion.cpp
--- a/Engine/Plugins/Runtime/USDCore/Source/USDUtilities/Private/USDSkeletalDataConversion.cpp
+++ b/Engine/Plugins/Runtime/USDCore/Source/USDUtilities/Private/USDSkeletalDataConversion.cpp
@@ -2,7 +2,6 @@
 
 #include "USDSkeletalDataConversion.h"
 
-#include "Materials/Material.h"
 #include "UnrealUSDWrapper.h"
 #include "USDAttributeUtils.h"
 #include "USDClassesModule.h"
@@ -13,6 +12,7 @@
 #include "USDMemory.h"
 #include "USDObjectUtils.h"
 #include "USDPrimConversion.h"
+#include "USDSetUtils.h"
 #include "USDShadeConversion.h"
 #include "USDTypesConversion.h"
 
@@ -38,6 +38,7 @@
 #include "IMovieScenePlayer.h"
 #include "Interfaces/ITargetPlatformManagerModule.h"
 #include "MaterialDomain.h"
+#include "Materials/Material.h"
 #include "Misc/CoreMisc.h"
 #include "Misc/MemStack.h"
 #include "Modules/ModuleManager.h"
@@ -2312,7 +2313,7 @@ bool UsdToUnreal::ConvertSkelAnim(
 				}
 			}
 
-			TSet<FString> ProcessedLODParentPaths;
+			UsdUtils::FCaseSensitiveStringSet ProcessedLODParentPaths;
 
 			// Since we may need to switch variants to parse LODs, we could invalidate references to SkinningQuery objects, so we need
 			// to keep track of these by path and construct one whenever we need them
diff --git a/Engine/Plugins/Runtime/USDCore/Source/USDUtilities/Public/USDConversionUtils.h b/Engine/Plugins/Runtime/USDCore/Source/USDUtilities/Public/USDConversionUtils.h
--- a/Engine/Plugins/Runtime/USDCore/Source/USDUtilities/Public/USDConversionUtils.h
+++ b/Engine/Plugins/Runtime/USDCore/Source/USDUtilities/Public/USDConversionUtils.h
@@ -9,6 +9,7 @@
 
 #include "UnrealUSDWrapper.h"
 #include "USDProjectSettings.h"
+#include "USDSetUtils.h"
 #include "UsdWrappers/ForwardDeclarations.h"
 #include "UsdWrappers/SdfLayer.h"
 #include "UsdWrappers/SdfPath.h"
@@ -85,25 +86,16 @@ namespace UsdUtils
 		return Object;
 	}
 
-	/** Case sensitive hashing function for TMap */
-	template<typename ValueType>
-	struct FCaseSensitiveStringMapFuncs : BaseKeyFuncs<ValueType, FString, /*bInAllowDuplicateKeys*/ false>
-	{
-		static inline const FString& GetSetKey(const TPair<FString, ValueType>& Element)
-		{
-			return Element.Key;
-		}
-
-		static inline bool Matches(const FString& A, const FString& B)
-		{
-			return A.Equals(B, ESearchCase::CaseSensitive);
-		}
-
-		static inline uint32 GetKeyHash(const FString& Key)
-		{
-			return FCrc::StrCrc32<TCHAR>(*Key);
-		}
-	};
+	/**
+	 * Adds or removes a case sensitive hash suffix to the PrimPath string.
+	 *
+	 * We mainly use this to make sure that prim paths that differ only by case don't cause collisions inside
+	 * case insensitive TMaps and TSets. We can sometimes use the case sensitive versions of those containers
+	 * (see USDSetUtils.h), but sometimes that is not possible, due to backwards compatibility or due to the
+	 * container being a UPROPERTY
+	 */
+	USDUTILITIES_API FString AddHashToPrimPath(const FString& PrimPath);
+	USDUTILITIES_API FString RemoveHashFromPrimPath(const FString& HashedPrimPath);
 
 #if USE_USD_SDK
 	template<typename ValueType>
diff --git a/Engine/Plugins/Runtime/USDCore/Source/UnrealUSDWrapper/Public/USDListener.h b/Engine/Plugins/Runtime/USDCore/Source/UnrealUSDWrapper/Public/USDListener.h
--- a/Engine/Plugins/Runtime/USDCore/Source/UnrealUSDWrapper/Public/USDListener.h
+++ b/Engine/Plugins/Runtime/USDCore/Source/UnrealUSDWrapper/Public/USDListener.h
@@ -6,6 +6,7 @@
 #include "Containers/Map.h"
 #include "Delegates/Delegate.h"
 
+#include "USDSetUtils.h"
 #include "UsdWrappers/ForwardDeclarations.h"
 #include "UsdWrappers/VtValue.h"
 
@@ -92,7 +93,7 @@ namespace UsdUtils
 	/**
 	 * Describes USD object changes by object path (an object can be a prim, a property, etc.)
 	 */
-	using FObjectChangesByPath = TMap<FString, TArray<FSdfChangeListEntry>>;
+	using FObjectChangesByPath = TCaseSensitiveStringMap<TArray<FSdfChangeListEntry>>;
 }	 // namespace UsdUtils
 
 /**
