USD: Add the ShapingAPI during level sequence export, or else the spot light attributes don't get exported

#Synced-CL 51466015

diff --git a/Engine/Plugins/Importers/USDImporter/Source/USDExporter/Private/LevelSequenceExporterUSD.cpp b/Engine/Plugins/Importers/USDImporter/Source/USDExporter/Private/LevelSequenceExporterUSD.cpp
--- a/Engine/Plugins/Importers/USDImporter/Source/USDExporter/Private/LevelSequenceExporterUSD.cpp
+++ b/Engine/Plugins/Importers/USDImporter/Source/USDExporter/Private/LevelSequenceExporterUSD.cpp
@@ -30,6 +30,7 @@
 #include "CameraRig_Rail.h"
 #include "Compilation/MovieSceneCompiledDataManager.h"
 #include "Components/SkeletalMeshComponent.h"
+#include "Components/SpotLightComponent.h"
 #include "Editor.h"
 #include "Engine/SkeletalMesh.h"
 #include "Engine/SphereReflectionCapture.h"
@@ -1106,7 +1107,16 @@ namespace UE::LevelSequenceExporterUSD::Private
 
 				// We will define a prim here so that we can apply schemas and use the shortcut CreateXAttribute functions,
 				// and not have to worry about attribute names and types. Later on we will convert these prims back into just 'overs' though
-				return UsdStage.DefinePrim(UE::FSdfPath{*PrimPath}, *SchemaName);
+				UE::FUsdPrim Prim = UsdStage.DefinePrim(UE::FSdfPath{*PrimPath}, *SchemaName);
+
+				// For the spot light case we also have to author our API schema here, or else our property bakers will fail to see this
+				// as a valid prim for a spot light
+				if (Prim && Component.IsA<USpotLightComponent>())
+				{
+					Prim.ApplyAPI(TEXT("ShapingAPI"));
+				}
+
+				return Prim;
 			};
 
 			UE::FUsdPrim Prim = GetPrimForComponent(*BoundComponent, &PrimPath);
